Interchange framework initial commit. This is an experimental feature that is off by default

diff --git a/Engine/Plugins/Runtime/Interchange/Interchange.uplugin b/Engine/Plugins/Runtime/Interchange/Interchange.uplugin
new file mode 100644
--- /dev/null
+++ b/Engine/Plugins/Runtime/Interchange/Interchange.uplugin
@@ -0,0 +1,54 @@
+{
+	"FileVersion": 3,
+	"Version": 1,
+	"VersionName": "1.0",
+	"FriendlyName": "Interchange Framework",
+	"Description": "Interchange framework plugin implement all implementation of engine importers and exporters.",
+	"Category": "Import Export",
+	"CreatedBy": "Epic Games, Inc.",
+	"CreatedByURL": "http://epicgames.com",
+	"DocsURL": "",
+	"MarketplaceURL": "",
+	"SupportURL": "",
+	"EnabledByDefault": true,
+	"CanContainContent": false,
+	"IsBetaVersion": false,
+	"Installed": false,
+	"SupportedTargetPlatforms": [
+		"Win64",
+		"Linux",
+		"Mac"
+	],
+	"Modules": [
+		{
+			"Name": "InterchangeNodePlugin",
+			"Type": "Runtime",
+			"LoadingPhase": "Default",
+			"WhitelistPlatforms": [
+				"Win64",
+				"Linux",
+				"Mac"
+			]
+		},
+		{
+			"Name": "InterchangeImportPlugin",
+			"Type": "Runtime",
+			"LoadingPhase": "Default",
+			"WhitelistPlatforms": [
+				"Win64",
+				"Linux",
+				"Mac"
+			]
+		},
+		{
+			"Name": "InterchangeExportPlugin",
+			"Type": "Runtime",
+			"LoadingPhase": "Default",
+			"WhitelistPlatforms": [
+				"Win64",
+				"Linux",
+				"Mac"
+			]
+		}
+	]
+}
\ No newline at end of file
diff --git a/Engine/Plugins/Runtime/Interchange/Source/Export/InterchangeExportPlugin.Build.cs b/Engine/Plugins/Runtime/Interchange/Source/Export/InterchangeExportPlugin.Build.cs
new file mode 100644
--- /dev/null
+++ b/Engine/Plugins/Runtime/Interchange/Source/Export/InterchangeExportPlugin.Build.cs
@@ -0,0 +1,20 @@
+// Copyright Epic Games, Inc. All Rights Reserved.
+
+namespace UnrealBuildTool.Rules
+{
+	public class InterchangeExportPlugin : ModuleRules
+	{
+		public InterchangeExportPlugin(ReadOnlyTargetRules Target) : base(Target)
+		{
+			PublicDependencyModuleNames.AddRange(
+				new string[]
+				{
+					"Core",
+					"CoreUObject",
+					"Engine",
+					"InterchangeCore",
+				}
+				);
+		}
+	}
+}
diff --git a/Engine/Plugins/Runtime/Interchange/Source/Export/Private/InterchangeExportPlugin.cpp b/Engine/Plugins/Runtime/Interchange/Source/Export/Private/InterchangeExportPlugin.cpp
new file mode 100644
--- /dev/null
+++ b/Engine/Plugins/Runtime/Interchange/Source/Export/Private/InterchangeExportPlugin.cpp
@@ -0,0 +1,35 @@
+// Copyright Epic Games, Inc. All Rights Reserved.
+
+#include "CoreMinimal.h"
+#include "Modules/ModuleManager.h"
+#include "IInterchangeExportPlugin.h"
+#include "InterchangeManager.h"
+#include "InterchangeTextureWriter.h"
+
+class FInterchangeExportPlugin : public IInterchangeExportPlugin
+{
+	/** IModuleInterface implementation */
+	virtual void StartupModule() override;
+	virtual void ShutdownModule() override;
+};
+
+IMPLEMENT_MODULE( FInterchangeExportPlugin, InterchangeExportPlugin)
+
+
+
+void FInterchangeExportPlugin::StartupModule()
+{
+	UInterchangeManager& InterchangeManager = UInterchangeManager::GetInterchangeManager();
+
+	//Register the Writers
+	InterchangeManager.RegisterWriter(UInterchangeTextureWriter::StaticClass());
+}
+
+
+void FInterchangeExportPlugin::ShutdownModule()
+{
+	
+}
+
+
+
diff --git a/Engine/Plugins/Runtime/Interchange/Source/Export/Private/InterchangeTextureWriter.cpp b/Engine/Plugins/Runtime/Interchange/Source/Export/Private/InterchangeTextureWriter.cpp
new file mode 100644
--- /dev/null
+++ b/Engine/Plugins/Runtime/Interchange/Source/Export/Private/InterchangeTextureWriter.cpp
@@ -0,0 +1,7 @@
+// Copyright Epic Games, Inc. All Rights Reserved. 
+#include "InterchangeTextureWriter.h"
+
+bool UInterchangeTextureWriter::Export(Interchange::FBaseNodeContainer& BaseNodeContainer) const
+{
+	return false;
+}
diff --git a/Engine/Plugins/Runtime/Interchange/Source/Export/Public/IInterchangeExportPlugin.h b/Engine/Plugins/Runtime/Interchange/Source/Export/Public/IInterchangeExportPlugin.h
new file mode 100644
--- /dev/null
+++ b/Engine/Plugins/Runtime/Interchange/Source/Export/Public/IInterchangeExportPlugin.h
@@ -0,0 +1,39 @@
+// Copyright Epic Games, Inc. All Rights Reserved.
+
+#pragma once
+
+#include "CoreMinimal.h"
+#include "Modules/ModuleInterface.h"
+#include "Modules/ModuleManager.h"
+
+
+/**
+ * The public interface to this module
+ */
+class IInterchangeExportPlugin : public IModuleInterface
+{
+
+public:
+
+	/**
+	 * Singleton-like access to this module's interface.  This is just for convenience!
+	 * Beware of calling this during the shutdown phase, though.  Your module might have been unloaded already.
+	 *
+	 * @return Returns singleton instance, loading the module on demand if needed
+	 */
+	static inline IInterchangeExportPlugin& Get()
+	{
+		return FModuleManager::LoadModuleChecked< IInterchangeExportPlugin >( "InterchangeExportPlugin" );
+	}
+
+	/**
+	 * Checks to see if this module is loaded and ready.  It is only valid to call Get() if IsAvailable() returns true.
+	 *
+	 * @return True if the module is loaded and ready to use
+	 */
+	static inline bool IsAvailable()
+	{
+		return FModuleManager::Get().IsModuleLoaded( "InterchangeExportPlugin" );
+	}
+};
+
diff --git a/Engine/Plugins/Runtime/Interchange/Source/Export/Public/InterchangeTextureWriter.h b/Engine/Plugins/Runtime/Interchange/Source/Export/Public/InterchangeTextureWriter.h
new file mode 100644
--- /dev/null
+++ b/Engine/Plugins/Runtime/Interchange/Source/Export/Public/InterchangeTextureWriter.h
@@ -0,0 +1,26 @@
+// Copyright Epic Games, Inc. All Rights Reserved.
+
+#pragma once
+
+#include "CoreMinimal.h"
+#include "UObject/Object.h"
+#include "UObject/ObjectMacros.h"
+#include "InterchangeWriterBase.h"
+
+#include "InterchangeTextureWriter.generated.h"
+
+UCLASS(BlueprintType)
+class INTERCHANGEEXPORTPLUGIN_API UInterchangeTextureWriter : public UInterchangeWriterBase
+{
+	GENERATED_BODY()
+public:
+	
+	/**
+	 * Export all nodes of type FTextureNode hold by the BaseNodeContainer
+	 * @param BaseNodeContainer - Contain nodes describing what to export
+	 * @return true if the writer can export the nodes, false otherwise.
+	 */
+	virtual bool Export(Interchange::FBaseNodeContainer& BaseNodeContainer) const override;
+};
+
+
diff --git a/Engine/Plugins/Runtime/Interchange/Source/Import/InterchangeImportPlugin.Build.cs b/Engine/Plugins/Runtime/Interchange/Source/Import/InterchangeImportPlugin.Build.cs
new file mode 100644
--- /dev/null
+++ b/Engine/Plugins/Runtime/Interchange/Source/Import/InterchangeImportPlugin.Build.cs
@@ -0,0 +1,28 @@
+// Copyright Epic Games, Inc. All Rights Reserved.
+
+namespace UnrealBuildTool.Rules
+{
+	public class InterchangeImportPlugin : ModuleRules
+	{
+		public InterchangeImportPlugin(ReadOnlyTargetRules Target) : base(Target)
+		{
+			PublicDependencyModuleNames.AddRange(
+				new string[]
+				{
+					"Core",
+					"CoreUObject",
+					"Engine",
+					"InterchangeCore",
+					"InterchangeNodePlugin",
+				}
+				);
+			PrivateDependencyModuleNames.AddRange(
+				new string[]
+				{
+					"ImageWrapper",
+					"RHI",
+				}
+				);
+		}
+	}
+}
diff --git a/Engine/Plugins/Runtime/Interchange/Source/Import/Private/InterchangeImportPlugin.cpp b/Engine/Plugins/Runtime/Interchange/Source/Import/Private/InterchangeImportPlugin.cpp
new file mode 100644
--- /dev/null
+++ b/Engine/Plugins/Runtime/Interchange/Source/Import/Private/InterchangeImportPlugin.cpp
@@ -0,0 +1,43 @@
+// Copyright Epic Games, Inc. All Rights Reserved.
+
+#include "CoreMinimal.h"
+#include "Modules/ModuleManager.h"
+#include "IInterchangeImportPlugin.h"
+#include "InterchangeManager.h"
+#include "Texture/InterchangeTextureFactory.h"
+#include "Texture/InterchangePNGTranslator.h"
+#include "Texture/InterchangeJPGTranslator.h"
+
+DEFINE_LOG_CATEGORY(LogInterchangeImportPlugin);
+
+class FInterchangeImportPlugin : public IInterchangeImportPlugin
+{
+	/** IModuleInterface implementation */
+	virtual void StartupModule() override;
+	virtual void ShutdownModule() override;
+};
+
+IMPLEMENT_MODULE( FInterchangeImportPlugin, InterchangeImportPlugin)
+
+
+
+void FInterchangeImportPlugin::StartupModule()
+{
+	UInterchangeManager& InterchangeManager = UInterchangeManager::GetInterchangeManager();
+
+	//Register the translators
+	InterchangeManager.RegisterTranslator(UInterchangePNGTranslator::StaticClass());
+	InterchangeManager.RegisterTranslator(UInterchangeJPGTranslator::StaticClass());
+
+	//Register the factories
+	InterchangeManager.RegisterFactory(UInterchangeTextureFactory::StaticClass());
+}
+
+
+void FInterchangeImportPlugin::ShutdownModule()
+{
+	
+}
+
+
+
diff --git a/Engine/Plugins/Runtime/Interchange/Source/Import/Private/Texture/InterchangeJPGTranslator.cpp b/Engine/Plugins/Runtime/Interchange/Source/Import/Private/Texture/InterchangeJPGTranslator.cpp
new file mode 100644
--- /dev/null
+++ b/Engine/Plugins/Runtime/Interchange/Source/Import/Private/Texture/InterchangeJPGTranslator.cpp
@@ -0,0 +1,141 @@
+// Copyright Epic Games, Inc. All Rights Reserved. 
+#include "Texture/InterchangeJPGTranslator.h"
+
+#include "Nodes/BaseNodeContainer.h"
+#include "TextureNode.h"
+#include "Engine/Texture.h"
+#include "IImageWrapper.h"
+#include "IImageWrapperModule.h"
+#include "Misc/ScopedSlowTask.h"
+
+bool UInterchangeJPGTranslator::CanImportSourceData(const UInterchangeSourceData* InSourceData) const
+{
+	FString Extension = FPaths::GetExtension(InSourceData->GetFilename());
+	FString JPGExtension = (TEXT("jpg;Texture"));
+	return JPGExtension.StartsWith(Extension);
+}
+
+bool UInterchangeJPGTranslator::Translate(const UInterchangeSourceData* SourceData, Interchange::FBaseNodeContainer& BaseNodeContainer) const
+{
+	FString Filename = SourceData->GetFilename();
+	if (!FPaths::FileExists(Filename))
+	{
+		return false;
+	}
+
+	FName DisplayLabel = *FPaths::GetBaseFilename(Filename);
+	Interchange::FNodeUniqueID NodeUID(*Filename);
+	//PNG is creating a UTexture2D
+	Interchange::FTextureNode* TextureNode = new Interchange::FTextureNode(NodeUID, DisplayLabel, UTexture2D::StaticClass());
+	TextureNode->SetPayLoadKey(Filename);
+
+	//Test node change
+ 	//TextureNode->SetCustomLODGroup((uint8)TextureGroup::TEXTUREGROUP_WorldNormalMap);
+ 	//TextureNode->SetCustomCompressionSettings((uint8)TextureCompressionSettings::TC_Normalmap);
+
+	BaseNodeContainer.AddNode(static_cast<Interchange::FBaseNode*>(TextureNode));
+	return true;
+}
+
+const TOptional<Interchange::FImportImage> UInterchangeJPGTranslator::GetPayloadData(const UInterchangeSourceData* SourceData, const FString& PayLoadKey) const
+{
+	if (!SourceData)
+	{
+		UE_LOG(LogInterchangeImportPlugin, Error, TEXT("Failed to import JPEG, bad source data."));
+		return TOptional<Interchange::FImportImage>();
+	}
+
+	TArray64<uint8> SourceDataBuffer;
+	FString Filename = SourceData->GetFilename();
+	
+	//Make sure the key fit the filename, The key should always be valid
+	if (!Filename.Equals(PayLoadKey))
+	{
+		UE_LOG(LogInterchangeImportPlugin, Error, TEXT("Failed to import JPEG, wrong payload key. [%s]"), *Filename);
+		return TOptional<Interchange::FImportImage>();
+	}
+
+	if (!FPaths::FileExists(Filename))
+	{
+		UE_LOG(LogInterchangeImportPlugin, Error, TEXT("Failed to import JPEG, cannot open file. [%s]"), *Filename);
+		return TOptional<Interchange::FImportImage>();
+	}
+
+	if (!FFileHelper::LoadFileToArray(SourceDataBuffer, *Filename))
+	{
+		UE_LOG(LogInterchangeImportPlugin, Error, TEXT("Failed to import JPEG, cannot load file content into an array. [%s]"), *Filename);
+		return TOptional<Interchange::FImportImage>();
+	}
+
+	const uint8* Buffer = SourceDataBuffer.GetData();
+	const uint8* BufferEnd = Buffer + SourceDataBuffer.Num();
+
+	bool bAllowNonPowerOfTwo = false;
+	GConfig->GetBool(TEXT("TextureImporter"), TEXT("AllowNonPowerOfTwoTextures"), bAllowNonPowerOfTwo, GEditorIni);
+
+	// Validate it.
+	const int32 Length = BufferEnd - Buffer;
+
+	IImageWrapperModule& ImageWrapperModule = FModuleManager::LoadModuleChecked<IImageWrapperModule>(FName("ImageWrapper"));
+
+	//
+	// JPG
+	//
+	TSharedPtr<IImageWrapper> JpegImageWrapper = ImageWrapperModule.CreateImageWrapper(EImageFormat::JPEG);
+	if (!JpegImageWrapper.IsValid() || !JpegImageWrapper->SetCompressed(Buffer, Length))
+	{
+		UE_LOG(LogInterchangeImportPlugin, Error, TEXT("Failed to decode JPEG. [%s]"), *Filename);
+		return TOptional<Interchange::FImportImage>();
+	}
+	if (!Interchange::IsImportResolutionValid(JpegImageWrapper->GetWidth(), JpegImageWrapper->GetHeight(), bAllowNonPowerOfTwo))
+	{
+		UE_LOG(LogInterchangeImportPlugin, Error, TEXT("Failed to import JPEG, invalid resolution. Resolution[%d, %d], AllowPowerOfTwo[%s], [%s]"), JpegImageWrapper->GetWidth(), JpegImageWrapper->GetHeight(), bAllowNonPowerOfTwo ? TEXT("True") : TEXT("false"), *Filename);
+		return TOptional<Interchange::FImportImage>();
+	}
+
+	// Select the texture's source format
+	ETextureSourceFormat TextureFormat = TSF_Invalid;
+	int32 BitDepth = JpegImageWrapper->GetBitDepth();
+	ERGBFormat Format = JpegImageWrapper->GetFormat();
+
+	if (Format == ERGBFormat::Gray)
+	{
+		if (BitDepth <= 8)
+		{
+			TextureFormat = TSF_G8;
+			Format = ERGBFormat::Gray;
+			BitDepth = 8;
+		}
+	}
+	else if (Format == ERGBFormat::RGBA)
+	{
+		if (BitDepth <= 8)
+		{
+			TextureFormat = TSF_BGRA8;
+			Format = ERGBFormat::BGRA;
+			BitDepth = 8;
+		}
+	}
+
+	if (TextureFormat == TSF_Invalid)
+	{
+		UE_LOG(LogInterchangeImportPlugin, Error, TEXT("JPEG file [%s] contains data in an unsupported format"), *Filename);
+		return TOptional<Interchange::FImportImage>();
+	}
+
+	TOptional<Interchange::FImportImage> PayloadData = Interchange::FImportImage();
+
+	PayloadData.GetValue().Init2DWithParams(
+		JpegImageWrapper->GetWidth(),
+		JpegImageWrapper->GetHeight(),
+		TextureFormat,
+		BitDepth < 16
+	);
+
+	if (!JpegImageWrapper->GetRaw(Format, BitDepth, PayloadData.GetValue().RawData))
+	{
+		UE_LOG(LogInterchangeImportPlugin, Error, TEXT("Failed to decode JPEG. [%s]"), *Filename);
+		return TOptional<Interchange::FImportImage>();
+	}
+	return PayloadData;
+}
\ No newline at end of file
diff --git a/Engine/Plugins/Runtime/Interchange/Source/Import/Private/Texture/InterchangePNGTranslator.cpp b/Engine/Plugins/Runtime/Interchange/Source/Import/Private/Texture/InterchangePNGTranslator.cpp
new file mode 100644
--- /dev/null
+++ b/Engine/Plugins/Runtime/Interchange/Source/Import/Private/Texture/InterchangePNGTranslator.cpp
@@ -0,0 +1,330 @@
+// Copyright Epic Games, Inc. All Rights Reserved. 
+#include "Texture/InterchangePNGTranslator.h"
+
+#include "Nodes/BaseNodeContainer.h"
+#include "TextureNode.h"
+#include "Engine/Texture.h"
+#include "IImageWrapper.h"
+#include "IImageWrapperModule.h"
+#include "Misc/ScopedSlowTask.h"
+
+
+//////////////////////////////////////////////////////////////////////////
+// PNG helper local function
+
+namespace PNGParserHelper
+{
+	/**
+	 * This fills any pixels of a texture with have an alpha value of zero,
+	 * with an RGB from the nearest neighboring pixel which has non-zero alpha.
+	 */
+	template<typename PixelDataType, typename ColorDataType, int32 RIdx, int32 GIdx, int32 BIdx, int32 AIdx>
+	class TPNGDataFill
+	{
+	public:
+
+		TPNGDataFill(int32 SizeX, int32 SizeY, uint8* SourceTextureData)
+			: SourceData(reinterpret_cast<PixelDataType*>(SourceTextureData))
+			, TextureWidth(SizeX)
+			, TextureHeight(SizeY)
+		{
+		}
+
+		void ProcessData()
+		{
+			int32 NumZeroedTopRowsToProcess = 0;
+			int32 FillColorRow = -1;
+			for (int32 Y = 0; Y < TextureHeight; ++Y)
+			{
+				if (!ProcessHorizontalRow(Y))
+				{
+					if (FillColorRow != -1)
+					{
+						FillRowColorPixels(FillColorRow, Y);
+					}
+					else
+					{
+						NumZeroedTopRowsToProcess = Y;
+					}
+				}
+				else
+				{
+					FillColorRow = Y;
+				}
+			}
+
+			// Can only fill upwards if image not fully zeroed
+			if (NumZeroedTopRowsToProcess > 0 && NumZeroedTopRowsToProcess + 1 < TextureHeight)
+			{
+				for (int32 Y = 0; Y <= NumZeroedTopRowsToProcess; ++Y)
+				{
+					FillRowColorPixels(NumZeroedTopRowsToProcess + 1, Y);
+				}
+			}
+		}
+
+		/** returns False if requires further processing because entire row is filled with zeroed alpha values */
+		bool ProcessHorizontalRow(int32 Y)
+		{
+			// only wipe out colors that are affected by png turning valid colors white if alpha = 0
+			const uint32 WhiteWithZeroAlpha = FColor(255, 255, 255, 0).DWColor();
+
+			// Left -> Right
+			int32 NumLeftmostZerosToProcess = 0;
+			const PixelDataType* FillColor = nullptr;
+			for (int32 X = 0; X < TextureWidth; ++X)
+			{
+				PixelDataType* PixelData = SourceData + (Y * TextureWidth + X) * 4;
+				ColorDataType* ColorData = reinterpret_cast<ColorDataType*>(PixelData);
+
+				if (*ColorData == WhiteWithZeroAlpha)
+				{
+					if (FillColor)
+					{
+						PixelData[RIdx] = FillColor[RIdx];
+						PixelData[GIdx] = FillColor[GIdx];
+						PixelData[BIdx] = FillColor[BIdx];
+					}
+					else
+					{
+						// Mark pixel as needing fill
+						*ColorData = 0;
+
+						// Keep track of how many pixels to fill starting at beginning of row
+						NumLeftmostZerosToProcess = X;
+					}
+				}
+				else
+				{
+					FillColor = PixelData;
+				}
+			}
+
+			if (NumLeftmostZerosToProcess == 0)
+			{
+				// No pixels left that are zero
+				return true;
+			}
+
+			if (NumLeftmostZerosToProcess + 1 >= TextureWidth)
+			{
+				// All pixels in this row are zero and must be filled using rows above or below
+				return false;
+			}
+
+			// Fill using non zero pixel immediately to the right of the beginning series of zeros
+			FillColor = SourceData + (Y * TextureWidth + NumLeftmostZerosToProcess + 1) * 4;
+
+			// Fill zero pixels found at beginning of row that could not be filled during the Left to Right pass
+			for (int32 X = 0; X <= NumLeftmostZerosToProcess; ++X)
+			{
+				PixelDataType* PixelData = SourceData + (Y * TextureWidth + X) * 4;
+				PixelData[RIdx] = FillColor[RIdx];
+				PixelData[GIdx] = FillColor[GIdx];
+				PixelData[BIdx] = FillColor[BIdx];
+			}
+
+			return true;
+		}
+
+		void FillRowColorPixels(int32 FillColorRow, int32 Y)
+		{
+			for (int32 X = 0; X < TextureWidth; ++X)
+			{
+				const PixelDataType* FillColor = SourceData + (FillColorRow * TextureWidth + X) * 4;
+				PixelDataType* PixelData = SourceData + (Y * TextureWidth + X) * 4;
+				PixelData[RIdx] = FillColor[RIdx];
+				PixelData[GIdx] = FillColor[GIdx];
+				PixelData[BIdx] = FillColor[BIdx];
+			}
+		}
+
+		PixelDataType* SourceData;
+		int32 TextureWidth;
+		int32 TextureHeight;
+	};
+
+	/**
+	 * For PNG texture importing, this ensures that any pixels with an alpha value of zero have an RGB
+	 * assigned to them from a neighboring pixel which has non-zero alpha.
+	 * This is needed as PNG exporters tend to turn pixels that are RGBA = (x,x,x,0) to (1,1,1,0)
+	 * and this produces artifacts when drawing the texture with bilinear filtering.
+	 *
+	 * @param TextureSource - The source texture
+	 * @param SourceData - The source texture data
+	 */
+	void FillZeroAlphaPNGData(int32 SizeX, int32 SizeY, ETextureSourceFormat SourceFormat, uint8* SourceData)
+	{
+		switch (SourceFormat)
+		{
+		case TSF_BGRA8:
+		{
+			TPNGDataFill<uint8, uint32, 2, 1, 0, 3> PNGFill(SizeX, SizeY, SourceData);
+			PNGFill.ProcessData();
+			break;
+		}
+
+		case TSF_RGBA16:
+		{
+			TPNGDataFill<uint16, uint64, 0, 1, 2, 3> PNGFill(SizeX, SizeY, SourceData);
+			PNGFill.ProcessData();
+			break;
+		}
+		}
+	}
+
+} //PNGParserHelper namespace end
+
+bool UInterchangePNGTranslator::CanImportSourceData(const UInterchangeSourceData* InSourceData) const
+{
+	FString Extension = FPaths::GetExtension(InSourceData->GetFilename());
+	FString PNGExtension = (TEXT("png;Texture"));
+	return PNGExtension.StartsWith(Extension);
+}
+
+bool UInterchangePNGTranslator::Translate(const UInterchangeSourceData* SourceData, Interchange::FBaseNodeContainer& BaseNodeContainer) const
+{
+	FString Filename = SourceData->GetFilename();
+	if (!FPaths::FileExists(Filename))
+	{
+		return false;
+	}
+
+	FName DisplayLabel = *FPaths::GetBaseFilename(Filename);
+	Interchange::FNodeUniqueID NodeUID(*Filename);
+	//PNG is creating a UTexture2D
+	Interchange::FTextureNode* TextureNode = new Interchange::FTextureNode(NodeUID, DisplayLabel, UTexture2D::StaticClass());
+	TextureNode->SetPayLoadKey(Filename);
+
+	//Test node change
+ 	//TextureNode->SetCustomLODGroup((uint8)TextureGroup::TEXTUREGROUP_WorldNormalMap);
+ 	//TextureNode->SetCustomCompressionSettings((uint8)TextureCompressionSettings::TC_Normalmap);
+
+	BaseNodeContainer.AddNode(static_cast<Interchange::FBaseNode*>(TextureNode));
+	return true;
+}
+
+const TOptional<Interchange::FImportImage> UInterchangePNGTranslator::GetPayloadData(const UInterchangeSourceData* SourceData, const FString& PayLoadKey) const
+{
+	if (!SourceData)
+	{
+		UE_LOG(LogInterchangeImportPlugin, Error, TEXT("Failed to import PNG, bad source data."));
+		return TOptional<Interchange::FImportImage>();
+	}
+
+	TArray64<uint8> SourceDataBuffer;
+	FString Filename = SourceData->GetFilename();
+	
+	//Make sure the key fit the filename, The key should always be valid
+	if (!Filename.Equals(PayLoadKey))
+	{
+		UE_LOG(LogInterchangeImportPlugin, Error, TEXT("Failed to import PNG, wrong payload key. [%s]"), *Filename);
+		return TOptional<Interchange::FImportImage>();
+	}
+
+	if (!FPaths::FileExists(Filename))
+	{
+		UE_LOG(LogInterchangeImportPlugin, Error, TEXT("Failed to import PNG, cannot open file. [%s]"), *Filename);
+		return TOptional<Interchange::FImportImage>();
+	}
+
+	if (!FFileHelper::LoadFileToArray(SourceDataBuffer, *Filename))
+	{
+		UE_LOG(LogInterchangeImportPlugin, Error, TEXT("Failed to import PNG, cannot load file content into an array. [%s]"), *Filename);
+		return TOptional<Interchange::FImportImage>();
+	}
+
+	const uint8* Buffer = SourceDataBuffer.GetData();
+	const uint8* BufferEnd = Buffer + SourceDataBuffer.Num();
+
+	bool bAllowNonPowerOfTwo = false;
+	GConfig->GetBool(TEXT("TextureImporter"), TEXT("AllowNonPowerOfTwoTextures"), bAllowNonPowerOfTwo, GEditorIni);
+
+	// Validate it.
+	const int32 Length = BufferEnd - Buffer;
+
+	IImageWrapperModule& ImageWrapperModule = FModuleManager::LoadModuleChecked<IImageWrapperModule>(FName("ImageWrapper"));
+
+	//
+	// PNG
+	//
+	TSharedPtr<IImageWrapper> PngImageWrapper = ImageWrapperModule.CreateImageWrapper(EImageFormat::PNG);
+	if (!PngImageWrapper.IsValid() || !PngImageWrapper->SetCompressed(Buffer, Length))
+	{
+		UE_LOG(LogInterchangeImportPlugin, Error, TEXT("Failed to decode PNG. [%s]"), *Filename);
+		return TOptional<Interchange::FImportImage>();
+	}
+	if (!Interchange::IsImportResolutionValid(PngImageWrapper->GetWidth(), PngImageWrapper->GetHeight(), bAllowNonPowerOfTwo))
+	{
+		UE_LOG(LogInterchangeImportPlugin, Error, TEXT("Failed to import PNG, invalid resolution. Resolution[%d, %d], AllowPowerOfTwo[%s], [%s]"), PngImageWrapper->GetWidth(), PngImageWrapper->GetHeight(), bAllowNonPowerOfTwo ? TEXT("True") : TEXT("false"), *Filename);
+		return TOptional<Interchange::FImportImage>();
+	}
+
+	// Select the texture's source format
+	ETextureSourceFormat TextureFormat = TSF_Invalid;
+	int32 BitDepth = PngImageWrapper->GetBitDepth();
+	ERGBFormat Format = PngImageWrapper->GetFormat();
+
+	if (Format == ERGBFormat::Gray)
+	{
+		if (BitDepth <= 8)
+		{
+			TextureFormat = TSF_G8;
+			Format = ERGBFormat::Gray;
+			BitDepth = 8;
+		}
+		else if (BitDepth == 16)
+		{
+			// TODO: TSF_G16?
+			TextureFormat = TSF_RGBA16;
+			Format = ERGBFormat::RGBA;
+			BitDepth = 16;
+		}
+	}
+	else if (Format == ERGBFormat::RGBA || Format == ERGBFormat::BGRA)
+	{
+		if (BitDepth <= 8)
+		{
+			TextureFormat = TSF_BGRA8;
+			Format = ERGBFormat::BGRA;
+			BitDepth = 8;
+		}
+		else if (BitDepth == 16)
+		{
+			TextureFormat = TSF_RGBA16;
+			Format = ERGBFormat::RGBA;
+			BitDepth = 16;
+		}
+	}
+
+	if (TextureFormat == TSF_Invalid)
+	{
+		UE_LOG(LogInterchangeImportPlugin, Error, TEXT("PNG file [%s] contains data in an unsupported format"), *Filename);
+		return TOptional<Interchange::FImportImage>();
+	}
+	TOptional<Interchange::FImportImage> PayloadData = Interchange::FImportImage();
+	PayloadData.GetValue().Init2DWithParams(
+		PngImageWrapper->GetWidth(),
+		PngImageWrapper->GetHeight(),
+		TextureFormat,
+		BitDepth < 16
+	);
+
+	if (PngImageWrapper->GetRaw(Format, BitDepth, PayloadData.GetValue().RawData))
+	{
+		bool bFillPNGZeroAlpha = true;
+		GConfig->GetBool(TEXT("TextureImporter"), TEXT("FillPNGZeroAlpha"), bFillPNGZeroAlpha, GEditorIni);
+
+		if (bFillPNGZeroAlpha)
+		{
+			// Replace the pixels with 0.0 alpha with a color value from the nearest neighboring color which has a non-zero alpha
+			PNGParserHelper::FillZeroAlphaPNGData(PayloadData.GetValue().SizeX, PayloadData.GetValue().SizeY, PayloadData.GetValue().Format, PayloadData.GetValue().RawData.GetData());
+		}
+	}
+	else
+	{
+		UE_LOG(LogInterchangeImportPlugin, Error, TEXT("Failed to decode JPEG. [%s]"), *Filename);
+		return TOptional<Interchange::FImportImage>();
+	}
+	return PayloadData;
+}
\ No newline at end of file
diff --git a/Engine/Plugins/Runtime/Interchange/Source/Import/Private/Texture/InterchangeTextureFactory.cpp b/Engine/Plugins/Runtime/Interchange/Source/Import/Private/Texture/InterchangeTextureFactory.cpp
new file mode 100644
--- /dev/null
+++ b/Engine/Plugins/Runtime/Interchange/Source/Import/Private/Texture/InterchangeTextureFactory.cpp
@@ -0,0 +1,191 @@
+// Copyright Epic Games, Inc. All Rights Reserved. 
+#include "Texture/InterchangeTextureFactory.h"
+
+#include "Nodes/BaseNode.h"
+#include "TextureNode.h"
+#include "TextureCompiler.h"
+#include "TextureNode.h"
+
+#if WITH_EDITORONLY_DATA
+
+#include "EditorFramework/AssetImportData.h"
+
+#endif //WITH_EDITORONLY_DATA
+
+UClass* UInterchangeTextureFactory::GetFactoryClass() const
+{
+	return UTexture::StaticClass();
+}
+
+UObject* UInterchangeTextureFactory::CreateEmptyAsset(const FCreateAssetParams& Arguments) const
+{
+	UObject* Texture = nullptr;
+
+#if WITH_EDITORONLY_DATA
+	if (!Arguments.AssetNode || !Arguments.AssetNode->GetAssetClass()->IsChildOf(UTexture::StaticClass()))
+	{
+		return nullptr;
+	}
+
+	const Interchange::FTextureNode* TextureNode = static_cast<const Interchange::FTextureNode*>(Arguments.AssetNode);
+	if (TextureNode == nullptr)
+	{
+		return nullptr;
+	}
+
+	//
+	// Generic 2D Image
+	//
+	const TOptional<FString>& PayLoadKey = TextureNode->GetPayLoadKey();
+	if (!PayLoadKey.IsSet())
+	{
+		return nullptr;
+	}
+
+	const UClass* TextureClass = TextureNode->GetAssetClass();
+	check(TextureClass && TextureClass->IsChildOf(UTexture::StaticClass()));
+
+	// create an asset if it doesn't exist
+	UObject* ExistingAsset = StaticFindObject(nullptr, Arguments.Parent, *Arguments.AssetName);
+
+	// create a new texture or overwrite existing asset, if possible
+	if (!ExistingAsset)
+	{
+		Texture = NewObject<UObject>(Arguments.Parent, TextureClass, *Arguments.AssetName, RF_Public);
+	}
+	else if (ExistingAsset->GetClass()->IsChildOf(TextureClass))
+	{
+		//This is a reimport, we are just re-updating the source data
+		Texture = ExistingAsset;
+	}
+
+	if (!Texture)
+	{
+		UE_LOG(LogInterchangeImportPlugin, Warning, TEXT("Could not create texture asset %s"), *Arguments.AssetName);
+		return nullptr;
+	}
+#endif //WITH_EDITORONLY_DATA
+	return Texture;
+}
+
+UObject* UInterchangeTextureFactory::CreateAsset(const UInterchangeTextureFactory::FCreateAssetParams& Arguments) const
+{
+#if !WITH_EDITORONLY_DATA
+
+	UE_LOG(LogInterchangeImportPlugin, Error, TEXT("Cannot import texture asset in runtime, this is an editor only feature."));
+	return nullptr;
+
+#else
+
+	if (!Arguments.AssetNode || !Arguments.AssetNode->GetAssetClass()->IsChildOf(UTexture::StaticClass()))
+	{
+		return nullptr;
+	}
+
+	const Interchange::FTextureNode* TextureNode = static_cast<const Interchange::FTextureNode*>(Arguments.AssetNode);
+	if (TextureNode == nullptr)
+	{
+		return nullptr;
+	}
+
+	const UInterchangeTextureTranslator* TextureTranslator = Cast<UInterchangeTextureTranslator>(Arguments.Translator);
+	if (!TextureTranslator)
+	{
+		return nullptr;
+	}
+
+	//
+	// Generic 2D Image
+	//
+	const TOptional<FString>& PayLoadKey = TextureNode->GetPayLoadKey();
+	if (!PayLoadKey.IsSet())
+	{
+		return nullptr;
+	}
+	const TOptional<Interchange::FImportImage> PayloadData = TextureTranslator->GetPayloadData(Arguments.SourceData, PayLoadKey.GetValue());
+	if(!PayloadData.IsSet())
+	{
+		return nullptr;
+	}
+	const Interchange::FImportImage& Image = PayloadData.GetValue();
+
+	const UClass* TextureClass = TextureNode->GetAssetClass();
+	check(TextureClass && TextureClass->IsChildOf(UTexture::StaticClass()));
+
+	// create an asset if it doesn't exist
+	UObject* ExistingAsset = StaticFindObject(nullptr, Arguments.Parent, *Arguments.AssetName);
+
+	UObject* Texture = nullptr;
+	// create a new texture or overwrite existing asset, if possible
+	if (!ExistingAsset)
+	{
+		//NewObject is not thread safe, the asset registry directory watcher tick on the main thread can trig before we finish initializing the UObject and will crash
+		//The UObject should have been create by calling CreateEmptyAsset on the main thread.
+		check(IsInGameThread());
+		Texture = NewObject<UObject>(Arguments.Parent, TextureClass, *Arguments.AssetName, RF_Public);
+	}
+	else if(ExistingAsset->GetClass()->IsChildOf(TextureClass))
+	{
+		//This is a reimport, we are just re-updating the source data
+		Texture = ExistingAsset;
+	}
+
+	if (!Texture)
+	{
+		UE_LOG(LogInterchangeImportPlugin, Warning, TEXT("Could not create texture asset %s"), *Arguments.AssetName);
+		return nullptr;
+	}
+
+	UTexture2D* Texture2D = Cast<UTexture2D>(Texture);
+	if (Texture2D)
+	{
+		Texture2D->Source.Init(
+			Image.SizeX,
+			Image.SizeY,
+			/*NumSlices=*/ 1,
+			Image.NumMips,
+			Image.Format,
+			Image.RawData.GetData()
+		);
+		Texture2D->CompressionSettings = Image.CompressionSettings;
+		Texture2D->SRGB = Image.SRGB;
+
+		//The re-import change only the source data and will not apply any custom attribute to the asset.
+		if(!ExistingAsset)
+		{
+			/** Apply all TextureNode custom attributes to the texture asset */
+			TextureNode->ApplyAllCustomAttributeToAsset(Texture2D);
+		}
+		
+		//Getting the file Hash will cache it into the source data
+		Arguments.SourceData->GetFileContentHash();
+
+		//The interchange completion task (call in the GameThread after the factories pass), will call PostEditChange which will trig another asynchronous system that will build all texture in parallel
+	}
+	else
+	{
+		//The texture is not a UTexture2D
+		Texture->RemoveFromRoot();
+		Texture->MarkPendingKill();
+	}
+	return Texture2D;
+#endif
+}
+
+/* This function is call in the completion task on the main thread, use it to call main thread post creation step for your assets*/
+void UInterchangeTextureFactory::PostImportGameThreadCallback(const FPostImportGameThreadCallbackParams& Arguments) const
+{
+	check(IsInGameThread());
+	Super::PostImportGameThreadCallback(Arguments);
+#if WITH_EDITORONLY_DATA
+	if(ensure(Arguments.ReimportObject && Arguments.SourceData))
+	{
+		//We must call the Update of the asset source file in the main thread because UAssetImportData::Update execute some delegate we do not control
+		UTexture* ImportedTexture = CastChecked<UTexture>(Arguments.ReimportObject);
+		//Set the asset import data file source to allow reimport. TODO: manage MD5 Hash properly
+		TOptional<FMD5Hash> FileContentHash = Arguments.SourceData->GetFileContentHash();
+		ImportedTexture->AssetImportData->Update(Arguments.SourceData->GetFilename(), FileContentHash.IsSet() ? &FileContentHash.GetValue() : nullptr);
+	}
+#endif //#if WITH_EDITORONLY_DATA
+	return;
+}
\ No newline at end of file
diff --git a/Engine/Plugins/Runtime/Interchange/Source/Import/Private/Texture/InterchangeTexturePayloadData.cpp b/Engine/Plugins/Runtime/Interchange/Source/Import/Private/Texture/InterchangeTexturePayloadData.cpp
new file mode 100644
--- /dev/null
+++ b/Engine/Plugins/Runtime/Interchange/Source/Import/Private/Texture/InterchangeTexturePayloadData.cpp
@@ -0,0 +1,118 @@
+// Copyright Epic Games, Inc. All Rights Reserved.
+
+#pragma once
+#include "Texture/InterchangeTexturePayloadData.h"
+
+#include "CoreMinimal.h"
+#include "Engine/Texture.h"
+#include "HAL/IConsoleManager.h"
+#include "Misc/MessageDialog.h"
+#include "RHI.h"
+
+bool Interchange::IsImportResolutionValid(int32 Width, int32 Height, bool bAllowNonPowerOfTwo)
+{
+	static const auto CVarVirtualTexturesEnabled = IConsoleManager::Get().FindTConsoleVariableDataInt(TEXT("r.VirtualTextures")); check(CVarVirtualTexturesEnabled);
+
+	// In theory this value could be much higher, but various UE4 image code currently uses 32bit size/offset values
+	const int32 MaximumSupportedVirtualTextureResolution = 16 * 1024;
+
+	// Calculate the maximum supported resolution utilizing the global max texture mip count
+	// (Note, have to subtract 1 because 1x1 is a valid mip-size; this means a GMaxTextureMipCount of 4 means a max resolution of 8x8, not 2^4 = 16x16)
+	const int32 MaximumSupportedResolution = CVarVirtualTexturesEnabled->GetValueOnAnyThread() ? MaximumSupportedVirtualTextureResolution : (1 << (GMaxTextureMipCount - 1));
+
+	bool bValid = true;
+
+	// Check if the texture is above the supported resolution and prompt the user if they wish to continue if it is
+	if (Width > MaximumSupportedResolution || Height > MaximumSupportedResolution)
+	{
+		if (EAppReturnType::Yes != FMessageDialog::Open(EAppMsgType::YesNo, FText::Format(
+			NSLOCTEXT("Interchange", "Warning_LargeTextureImport", "Attempting to import {0} x {1} texture, proceed?\nLargest supported texture size: {2} x {3}"),
+			FText::AsNumber(Width), FText::AsNumber(Height), FText::AsNumber(MaximumSupportedResolution), FText::AsNumber(MaximumSupportedResolution))))
+		{
+			bValid = false;
+		}
+
+		if (bValid && (Width * Height) > FMath::Square(MaximumSupportedVirtualTextureResolution))
+		{
+			//Warn->Log(ELogVerbosity::Error, *NSLOCTEXT("UnrealEd", "Warning_TextureSizeTooLarge", "Texture is too large to import").ToString());
+			bValid = false;
+		}
+	}
+
+	const bool bIsPowerOfTwo = FMath::IsPowerOfTwo(Width) && FMath::IsPowerOfTwo(Height);
+	// Check if the texture dimensions are powers of two
+	if (!bAllowNonPowerOfTwo && !bIsPowerOfTwo)
+	{
+		//Warn->Log(ELogVerbosity::Error, *NSLOCTEXT("UnrealEd", "Warning_TextureNotAPowerOfTwo", "Cannot import texture with non-power of two dimensions").ToString());
+		bValid = false;
+	}
+
+	return bValid;
+}
+
+void Interchange::FImportImage::Init2DWithParams(int32 InSizeX, int32 InSizeY, ETextureSourceFormat InFormat, bool InSRGB)
+{
+	SizeX = InSizeX;
+	SizeY = InSizeY;
+	NumMips = 1;
+	Format = InFormat;
+	SRGB = InSRGB;
+}
+
+void Interchange::FImportImage::Init2DWithOneMip(int32 InSizeX, int32 InSizeY, ETextureSourceFormat InFormat, const void* InData)
+{
+#if WITH_EDITOR
+	SizeX = InSizeX;
+	SizeY = InSizeY;
+	NumMips = 1;
+	Format = InFormat;
+	RawData.AddUninitialized((int64)SizeX * SizeY * FTextureSource::GetBytesPerPixel(Format));
+	if (InData)
+	{
+		FMemory::Memcpy(RawData.GetData(), InData, RawData.Num());
+	}
+#endif
+}
+
+void Interchange::FImportImage::Init2DWithMips(int32 InSizeX, int32 InSizeY, int32 InNumMips, ETextureSourceFormat InFormat, const void* InData)
+{
+	SizeX = InSizeX;
+	SizeY = InSizeY;
+	NumMips = InNumMips;
+	Format = InFormat;
+
+	int64 TotalSize = 0;
+	for (int32 MipIndex = 0; MipIndex < InNumMips; ++MipIndex)
+	{
+		TotalSize += GetMipSize(MipIndex);
+	}
+	RawData.AddUninitialized(TotalSize);
+
+	if (InData)
+	{
+		FMemory::Memcpy(RawData.GetData(), InData, RawData.Num());
+	}
+}
+
+int64 Interchange::FImportImage::GetMipSize(int32 InMipIndex) const
+{
+#if WITH_EDITOR
+	check(InMipIndex >= 0);
+	check(InMipIndex < NumMips);
+	const int32 MipSizeX = FMath::Max(SizeX >> InMipIndex, 1);
+	const int32 MipSizeY = FMath::Max(SizeY >> InMipIndex, 1);
+	return (int64)MipSizeX * MipSizeY * FTextureSource::GetBytesPerPixel(Format);
+#else
+	return 0;
+#endif
+}
+
+void* Interchange::FImportImage::GetMipData(int32 InMipIndex)
+{
+	int64 Offset = 0;
+	for (int32 MipIndex = 0; MipIndex < InMipIndex; ++MipIndex)
+	{
+		Offset += GetMipSize(MipIndex);
+	}
+	return &RawData[Offset];
+}
diff --git a/Engine/Plugins/Runtime/Interchange/Source/Import/Private/Texture/InterchangeTextureTranslator.cpp b/Engine/Plugins/Runtime/Interchange/Source/Import/Private/Texture/InterchangeTextureTranslator.cpp
new file mode 100644
--- /dev/null
+++ b/Engine/Plugins/Runtime/Interchange/Source/Import/Private/Texture/InterchangeTextureTranslator.cpp
@@ -0,0 +1,18 @@
+// Copyright Epic Games, Inc. All Rights Reserved.
+
+#pragma once
+#include "Texture/InterchangeTextureTranslator.h"
+
+#include "CoreMinimal.h"
+#include "UObject/Object.h"
+#include "UObject/ObjectMacros.h"
+#include "InterchangeTranslatorBase.h"
+#include "Texture/InterchangeTexturePayloadData.h"
+
+
+const TOptional<Interchange::FImportImage> UInterchangeTextureTranslator::GetPayloadData(const UInterchangeSourceData* SourceData, const FString& PayLoadKey) const
+{
+	return TOptional<Interchange::FImportImage>();
+}
+
+
diff --git a/Engine/Plugins/Runtime/Interchange/Source/Import/Public/IInterchangeImportPlugin.h b/Engine/Plugins/Runtime/Interchange/Source/Import/Public/IInterchangeImportPlugin.h
new file mode 100644
--- /dev/null
+++ b/Engine/Plugins/Runtime/Interchange/Source/Import/Public/IInterchangeImportPlugin.h
@@ -0,0 +1,40 @@
+// Copyright Epic Games, Inc. All Rights Reserved.
+
+#pragma once
+
+#include "CoreMinimal.h"
+#include "Modules/ModuleInterface.h"
+#include "Modules/ModuleManager.h"
+
+DECLARE_LOG_CATEGORY_EXTERN(LogInterchangeImportPlugin, Log, All);
+
+/**
+ * The public interface to this module
+ */
+class IInterchangeImportPlugin : public IModuleInterface
+{
+
+public:
+
+	/**
+	 * Singleton-like access to this module's interface.  This is just for convenience!
+	 * Beware of calling this during the shutdown phase, though.  Your module might have been unloaded already.
+	 *
+	 * @return Returns singleton instance, loading the module on demand if needed
+	 */
+	static inline IInterchangeImportPlugin& Get()
+	{
+		return FModuleManager::LoadModuleChecked< IInterchangeImportPlugin >( "InterchangeImportPlugin" );
+	}
+
+	/**
+	 * Checks to see if this module is loaded and ready.  It is only valid to call Get() if IsAvailable() returns true.
+	 *
+	 * @return True if the module is loaded and ready to use
+	 */
+	static inline bool IsAvailable()
+	{
+		return FModuleManager::Get().IsModuleLoaded( "InterchangeImportPlugin" );
+	}
+};
+
diff --git a/Engine/Plugins/Runtime/Interchange/Source/Import/Public/Texture/InterchangeJPGTranslator.h b/Engine/Plugins/Runtime/Interchange/Source/Import/Public/Texture/InterchangeJPGTranslator.h
new file mode 100644
--- /dev/null
+++ b/Engine/Plugins/Runtime/Interchange/Source/Import/Public/Texture/InterchangeJPGTranslator.h
@@ -0,0 +1,45 @@
+// Copyright Epic Games, Inc. All Rights Reserved.
+
+#pragma once
+
+#include "CoreMinimal.h"
+#include "UObject/Object.h"
+#include "UObject/ObjectMacros.h"
+#include "InterchangeTranslatorBase.h"
+#include "Texture/InterchangeTexturePayloadData.h"
+#include "Texture/InterchangeTextureTranslator.h"
+
+#include "InterchangeJPGTranslator.generated.h"
+
+UCLASS(BlueprintType)
+class INTERCHANGEIMPORTPLUGIN_API UInterchangeJPGTranslator : public UInterchangeTextureTranslator
+{
+	GENERATED_BODY()
+public:
+	
+	/**
+	 * return true if the translator can translate the specified file.
+	 */
+	virtual bool CanImportSourceData(const UInterchangeSourceData* InSourceData) const override;
+
+	/**
+	 * Translate a source data into a node hold by the specified nodes container.
+	 *
+	 * @param InSourceData - The source data containing the JPG data
+	 * @param BaseNodeContainer - The unreal objects descriptions container where to put the translated source data.
+	 * @return true if the translator can translate the source data, false otherwise.
+	 */
+	virtual bool Translate(const UInterchangeSourceData* SourceData, Interchange::FBaseNodeContainer& BaseNodeContainer) const override;
+
+	/**
+	 * Once the translation is done, the import process need a way to retrieve payload data.
+	 * This payload will be use by the factories to create the asset.
+	 *
+	 * @param SourceData - The source data containing the data to translate
+	 * @param PayloadKey - The key to retrieve the a particular payload contain into the specified source data.
+	 * @return a PayloadData containing the import image data. The TOptional will not be set if there is an error.
+	 */
+	virtual const TOptional<Interchange::FImportImage> GetPayloadData(const UInterchangeSourceData* SourceData, const FString& PayLoadKey) const override;
+};
+
+
diff --git a/Engine/Plugins/Runtime/Interchange/Source/Import/Public/Texture/InterchangePNGTranslator.h b/Engine/Plugins/Runtime/Interchange/Source/Import/Public/Texture/InterchangePNGTranslator.h
new file mode 100644
--- /dev/null
+++ b/Engine/Plugins/Runtime/Interchange/Source/Import/Public/Texture/InterchangePNGTranslator.h
@@ -0,0 +1,45 @@
+// Copyright Epic Games, Inc. All Rights Reserved.
+
+#pragma once
+
+#include "CoreMinimal.h"
+#include "UObject/Object.h"
+#include "UObject/ObjectMacros.h"
+#include "InterchangeTranslatorBase.h"
+#include "Texture/InterchangeTexturePayloadData.h"
+#include "Texture/InterchangeTextureTranslator.h"
+
+#include "InterchangePNGTranslator.generated.h"
+
+UCLASS(BlueprintType)
+class INTERCHANGEIMPORTPLUGIN_API UInterchangePNGTranslator : public UInterchangeTextureTranslator
+{
+	GENERATED_BODY()
+public:
+	
+	/*
+	 * return true if the translator can translate the specified file.
+	 */
+	virtual bool CanImportSourceData(const UInterchangeSourceData* InSourceData) const override;
+
+	/**
+	 * Translate a source data into a node hold by the specified nodes container.
+	 *
+	 * @param SourceData - The source data containing the data to translate
+	 * @param BaseNodeContainer - The unreal objects descriptions container where to put the translated source data.
+	 * @return true if the translator can translate the source data, false otherwise.
+	 */
+	virtual bool Translate(const UInterchangeSourceData* SourceData, Interchange::FBaseNodeContainer& BaseNodeContainer) const override;
+
+	/**
+	 * Once the translation is done, the import process need a way to retrieve payload data.
+	 * This payload will be use by the factories to create the asset.
+	 *
+	 * @param SourceData - The source data containing the data to translate
+	 * @param PayloadKey - The key to retrieve the a particular payload contain into the specified source data.
+	 * @return a PayloadData containing the import image data. The TOptional will not be set if there is an error.
+	 */
+	virtual const TOptional<Interchange::FImportImage> GetPayloadData(const UInterchangeSourceData* SourceData, const FString& PayLoadKey) const override;
+};
+
+
diff --git a/Engine/Plugins/Runtime/Interchange/Source/Import/Public/Texture/InterchangeTextureFactory.h b/Engine/Plugins/Runtime/Interchange/Source/Import/Public/Texture/InterchangeTextureFactory.h
new file mode 100644
--- /dev/null
+++ b/Engine/Plugins/Runtime/Interchange/Source/Import/Public/Texture/InterchangeTextureFactory.h
@@ -0,0 +1,31 @@
+// Copyright Epic Games, Inc. All Rights Reserved.
+
+#pragma once
+
+#include "CoreMinimal.h"
+#include "Engine/Texture.h"
+#include "InterchangeFactoryBase.h"
+#include "UObject/Object.h"
+#include "UObject/ObjectMacros.h"
+#include "Texture/InterchangeTexturePayloadData.h"
+
+#include "InterchangeTextureFactory.generated.h"
+
+UCLASS(BlueprintType)
+class INTERCHANGEIMPORTPLUGIN_API UInterchangeTextureFactory : public UInterchangeFactoryBase
+{
+	GENERATED_BODY()
+public:
+
+	//////////////////////////////////////////////////////////////////////////
+	// Interchange factory base interface begin
+
+	virtual UClass* GetFactoryClass() const override;
+	virtual UObject* CreateEmptyAsset(const FCreateAssetParams& Arguments) const override;
+	virtual UObject* CreateAsset(const FCreateAssetParams& Arguments) const override;
+	virtual void PostImportGameThreadCallback(const FPostImportGameThreadCallbackParams& Arguments) const override;
+	// Interchange factory base interface end
+	//////////////////////////////////////////////////////////////////////////
+};
+
+
diff --git a/Engine/Plugins/Runtime/Interchange/Source/Import/Public/Texture/InterchangeTexturePayloadData.h b/Engine/Plugins/Runtime/Interchange/Source/Import/Public/Texture/InterchangeTexturePayloadData.h
new file mode 100644
--- /dev/null
+++ b/Engine/Plugins/Runtime/Interchange/Source/Import/Public/Texture/InterchangeTexturePayloadData.h
@@ -0,0 +1,40 @@
+// Copyright Epic Games, Inc. All Rights Reserved.
+
+#pragma once
+
+#include "CoreMinimal.h"
+#include "Engine/Texture.h"
+
+namespace Interchange
+{
+	struct FImportImage
+	{
+		TArray64<uint8> RawData;
+		ETextureSourceFormat Format = TSF_Invalid;
+		TextureCompressionSettings CompressionSettings = TC_Default;
+		int32 NumMips = 0;
+		int32 SizeX = 0;
+		int32 SizeY = 0;
+		bool SRGB = true;
+
+		void Init2DWithParams(int32 InSizeX, int32 InSizeY, ETextureSourceFormat InFormat, bool InSRGB);
+		void Init2DWithOneMip(int32 InSizeX, int32 InSizeY, ETextureSourceFormat InFormat, const void* InData = nullptr);
+		void Init2DWithMips(int32 InSizeX, int32 InSizeY, int32 InNumMips, ETextureSourceFormat InFormat, const void* InData = nullptr);
+
+		int64 GetMipSize(int32 InMipIndex) const;
+		void* GetMipData(int32 InMipIndex);
+	};
+
+	/**
+	 * Tests if the given height and width specify a supported texture resolution to import; Can optionally check if the height/width are powers of two
+	 *
+	 * @param Width - The width of an imported texture whose validity should be checked
+	 * @param Height - The height of an imported texture whose validity should be checked
+	 * @param bAllowNonPowerOfTwo - Whether or not non-power-of-two textures are allowed
+	 *
+	 * @return bool true if the given height/width represent a supported texture resolution, false if not
+	 */
+	static bool IsImportResolutionValid(int32 Width, int32 Height, bool bAllowNonPowerOfTwo);
+}
+
+
diff --git a/Engine/Plugins/Runtime/Interchange/Source/Import/Public/Texture/InterchangeTextureTranslator.h b/Engine/Plugins/Runtime/Interchange/Source/Import/Public/Texture/InterchangeTextureTranslator.h
new file mode 100644
--- /dev/null
+++ b/Engine/Plugins/Runtime/Interchange/Source/Import/Public/Texture/InterchangeTextureTranslator.h
@@ -0,0 +1,30 @@
+// Copyright Epic Games, Inc. All Rights Reserved.
+
+#pragma once
+
+#include "CoreMinimal.h"
+#include "UObject/Object.h"
+#include "UObject/ObjectMacros.h"
+#include "InterchangeTranslatorBase.h"
+#include "Texture/InterchangeTexturePayloadData.h"
+
+#include "InterchangeTextureTranslator.generated.h"
+
+UCLASS(Abstract, BlueprintType)
+class INTERCHANGEIMPORTPLUGIN_API UInterchangeTextureTranslator : public UInterchangeTranslatorBase
+{
+	GENERATED_BODY()
+public:
+
+	/**
+	 * Once the translation is done, the import process need a way to retrieve payload data.
+	 * This payload will be use by the factories to create the asset.
+	 *
+	 * @param SourceData - The source data containing the data to translate
+	 * @param PayloadKey - The key to retrieve the a particular payload contain into the specified source data.
+	 * @return a PayloadData containing the import image data. The TOptional will not be set if there is an error.
+	 */
+	virtual const TOptional<Interchange::FImportImage> GetPayloadData(const UInterchangeSourceData* SourceData, const FString& PayloadKey) const;
+};
+
+
diff --git a/Engine/Plugins/Runtime/Interchange/Source/Nodes/InterchangeNodePlugin.Build.cs b/Engine/Plugins/Runtime/Interchange/Source/Nodes/InterchangeNodePlugin.Build.cs
new file mode 100644
--- /dev/null
+++ b/Engine/Plugins/Runtime/Interchange/Source/Nodes/InterchangeNodePlugin.Build.cs
@@ -0,0 +1,19 @@
+// Copyright Epic Games, Inc. All Rights Reserved.
+
+namespace UnrealBuildTool.Rules
+{
+	public class InterchangeNodePlugin : ModuleRules
+	{
+		public InterchangeNodePlugin(ReadOnlyTargetRules Target) : base(Target)
+		{
+			PublicDependencyModuleNames.AddRange(
+				new string[]
+				{
+					"Core",
+					"CoreUObject",
+					"InterchangeCore",
+				}
+			);
+		}
+	}
+}
diff --git a/Engine/Plugins/Runtime/Interchange/Source/Nodes/Private/InterchangeNodePlugin.cpp b/Engine/Plugins/Runtime/Interchange/Source/Nodes/Private/InterchangeNodePlugin.cpp
new file mode 100644
--- /dev/null
+++ b/Engine/Plugins/Runtime/Interchange/Source/Nodes/Private/InterchangeNodePlugin.cpp
@@ -0,0 +1,35 @@
+// Copyright Epic Games, Inc. All Rights Reserved.
+
+#include "CoreMinimal.h"
+#include "Modules/ModuleManager.h"
+#include "IInterchangeNodePlugin.h"
+//#include "InterchangeManager.h"
+
+DEFINE_LOG_CATEGORY(LogInterchangeNodePlugin);
+
+class FInterchangeNodePlugin : public IInterchangeNodePlugin
+{
+	/** IModuleInterface implementation */
+	virtual void StartupModule() override;
+	virtual void ShutdownModule() override;
+};
+
+IMPLEMENT_MODULE( FInterchangeNodePlugin, InterchangeNodePlugin)
+
+
+
+void FInterchangeNodePlugin::StartupModule()
+{
+	//Register anything needed to the interchange manager
+	//UInterchangeManager& InterchangeManager = UInterchangeManager::GetInterchangeManager();
+	//InterchangeManager.Register...
+}
+
+
+void FInterchangeNodePlugin::ShutdownModule()
+{
+	
+}
+
+
+
diff --git a/Engine/Plugins/Runtime/Interchange/Source/Nodes/Public/IInterchangeNodePlugin.h b/Engine/Plugins/Runtime/Interchange/Source/Nodes/Public/IInterchangeNodePlugin.h
new file mode 100644
--- /dev/null
+++ b/Engine/Plugins/Runtime/Interchange/Source/Nodes/Public/IInterchangeNodePlugin.h
@@ -0,0 +1,40 @@
+// Copyright Epic Games, Inc. All Rights Reserved.
+
+#pragma once
+
+#include "CoreMinimal.h"
+#include "Modules/ModuleInterface.h"
+#include "Modules/ModuleManager.h"
+
+DECLARE_LOG_CATEGORY_EXTERN(LogInterchangeNodePlugin, Log, All);
+
+/**
+ * The public interface to this module
+ */
+class IInterchangeNodePlugin : public IModuleInterface
+{
+
+public:
+
+	/**
+	 * Singleton-like access to this module's interface.  This is just for convenience!
+	 * Beware of calling this during the shutdown phase, though.  Your module might have been unloaded already.
+	 *
+	 * @return Returns singleton instance, loading the module on demand if needed
+	 */
+	static inline IInterchangeNodePlugin& Get()
+	{
+		return FModuleManager::LoadModuleChecked< IInterchangeNodePlugin >( "InterchangeNodePlugin" );
+	}
+
+	/**
+	 * Checks to see if this module is loaded and ready.  It is only valid to call Get() if IsAvailable() returns true.
+	 *
+	 * @return True if the module is loaded and ready to use
+	 */
+	static inline bool IsAvailable()
+	{
+		return FModuleManager::Get().IsModuleLoaded( "InterchangeNodePlugin" );
+	}
+};
+
diff --git a/Engine/Plugins/Runtime/Interchange/Source/Nodes/Public/TextureNode.h b/Engine/Plugins/Runtime/Interchange/Source/Nodes/Public/TextureNode.h
new file mode 100644
--- /dev/null
+++ b/Engine/Plugins/Runtime/Interchange/Source/Nodes/Public/TextureNode.h
@@ -0,0 +1,153 @@
+// Copyright Epic Games, Inc. All Rights Reserved.
+
+#pragma once
+
+#include "CoreMinimal.h"
+#include "Nodes/BaseNode.h"
+#include "Engine/Texture.h"
+
+//Interchange namespace
+namespace Interchange
+{
+
+struct FTextureNodeStaticData : public FBaseNodeStaticData
+{
+	static const FAttributeKey& PayloadSourceFileKey()
+	{
+		static FAttributeKey AttributeKey(TEXT("__PayloadSourceFile__"));
+		return AttributeKey;
+	}
+};
+
+class FTextureNode : public FBaseNode
+{
+public:
+	/**
+	 * Construct FTextureNode
+	 * @param: UniqueID - The uniqueId for this node
+	 * @param DisplayLabel - The name of the node
+	 * @param InAssetClass - The class the texture factory will create for this node.
+	 */
+	FTextureNode(const FNodeUniqueID& UniqueID, const FName& DisplayLabel, const TSubclassOf<UTexture>& InAssetClass)
+	: FBaseNode(UniqueID, DisplayLabel)
+	, AssetClass(InAssetClass)
+	{
+	}
+
+	/**
+	 * Return the node type name of the class, we use this when reporting error
+	 */
+	virtual FString GetTypeName() const override
+	{
+		const FString TypeName = TEXT("TextureNode");
+		return TypeName;
+	}
+
+	/** Get the class this node want to create */
+	virtual class UClass* GetAssetClass() const override
+	{
+		return AssetClass.Get() != nullptr ? AssetClass.Get() : UTexture::StaticClass();
+	}
+
+	virtual FGuid GetHash() const override
+	{
+		return Attributes.GetStorageHash();
+	}
+
+	/** Texture node Interface Begin */
+
+	virtual const TOptional<FString> GetPayLoadKey() const
+	{
+		if (!Attributes.ContainAttribute(FTextureNodeStaticData::PayloadSourceFileKey()))
+		{
+			return TOptional<FString>();
+		}
+		FAttributeStorage::TAttributeHandle<FString> AttributeHandle = Attributes.GetAttributeHandle<FString>(FTextureNodeStaticData::PayloadSourceFileKey());
+		if (!AttributeHandle.IsValid())
+		{
+			return TOptional<FString>();
+		}
+		FString PayloadKey;
+		EAttributeStorageResult Result = AttributeHandle.Get(PayloadKey);
+		if (!IsAttributeStorageResultSuccess(Result))
+		{
+			LogAttributeStorageErrors(Result, TEXT("FTextureNode.GetPayLoadKey"), FTextureNodeStaticData::PayloadSourceFileKey());
+			return TOptional<FString>();
+		}
+		return TOptional<FString>(PayloadKey);
+	}
+
+	virtual void SetPayLoadKey(const FString& PayloadKey)
+	{
+		EAttributeStorageResult Result = Attributes.RegisterAttribute(FTextureNodeStaticData::PayloadSourceFileKey(), PayloadKey);
+		if (!IsAttributeStorageResultSuccess(Result))
+		{
+			LogAttributeStorageErrors(Result, TEXT("FTextureNode.SetPayLoadKey"), FTextureNodeStaticData::PayloadSourceFileKey());
+		}
+	}
+
+#if WITH_EDITORONLY_DATA
+	//Texture Adjustment
+	IMPLEMENT_NODE_ATTRIBUTE(FTextureNode, AdjustBrightness, float, UTexture,);
+	IMPLEMENT_NODE_ATTRIBUTE(FTextureNode, AdjustBrightnessCurve, float, UTexture,);
+	IMPLEMENT_NODE_ATTRIBUTE(FTextureNode, AdjustVibrance, float, UTexture,);
+	IMPLEMENT_NODE_ATTRIBUTE(FTextureNode, AdjustSaturation, float, UTexture,);
+	IMPLEMENT_NODE_ATTRIBUTE(FTextureNode, AdjustRGBCurve, float, UTexture,);
+	IMPLEMENT_NODE_ATTRIBUTE(FTextureNode, AdjustHue, float, UTexture,);
+	IMPLEMENT_NODE_ATTRIBUTE(FTextureNode, AdjustMinAlpha, float, UTexture,);
+	IMPLEMENT_NODE_ATTRIBUTE(FTextureNode, AdjustMaxAlpha, float, UTexture,);
+	IMPLEMENT_NODE_ATTRIBUTE(FTextureNode, bChromaKeyTexture, bool, UTexture,);
+	IMPLEMENT_NODE_ATTRIBUTE(FTextureNode, ChromaKeyThreshold, float, UTexture,);
+	IMPLEMENT_NODE_ATTRIBUTE(FTextureNode, ChromaKeyColor, FColor, UTexture,);
+
+	//Texture Compression
+	IMPLEMENT_NODE_ATTRIBUTE(FTextureNode, CompressionNoAlpha, bool, UTexture,);
+	IMPLEMENT_NODE_ATTRIBUTE(FTextureNode, DeferCompression, bool, UTexture,);
+	IMPLEMENT_NODE_ATTRIBUTE(FTextureNode, LossyCompressionAmount, uint8, UTexture, TEnumAsByte<ETextureLossyCompressionAmount>);
+	IMPLEMENT_NODE_ATTRIBUTE(FTextureNode, MaxTextureSize, int32, UTexture,);
+	IMPLEMENT_NODE_ATTRIBUTE(FTextureNode, CompressionQuality, uint8, UTexture, TEnumAsByte<enum ETextureCompressionQuality>);
+#endif //#if WITH_EDITORONLY_DATA
+	IMPLEMENT_NODE_ATTRIBUTE(FTextureNode, CompressionSettings, uint8, UTexture, TEnumAsByte<enum TextureCompressionSettings>);
+
+#if WITH_EDITORONLY_DATA
+	//Texture general
+	IMPLEMENT_NODE_ATTRIBUTE(FTextureNode, bDitherMipMapAlpha, bool, UTexture,);
+	IMPLEMENT_NODE_ATTRIBUTE(FTextureNode, AlphaCoverageThresholds, FVector4, UTexture,);
+	IMPLEMENT_NODE_ATTRIBUTE(FTextureNode, bFlipGreenChannel, bool, UTexture,);
+	IMPLEMENT_NODE_ATTRIBUTE(FTextureNode, bForcePVRTC4, bool, UTexture,);
+	IMPLEMENT_NODE_ATTRIBUTE(FTextureNode, PowerOfTwoMode, uint8, UTexture, TEnumAsByte<enum ETexturePowerOfTwoSetting::Type>);
+	IMPLEMENT_NODE_ATTRIBUTE(FTextureNode, PaddingColor, FColor, UTexture,);
+#endif //#if WITH_EDITORONLY_DATA
+	IMPLEMENT_NODE_ATTRIBUTE(FTextureNode, Filter, uint8, UTexture, TEnumAsByte<enum TextureFilter>);
+	IMPLEMENT_NODE_ATTRIBUTE(FTextureNode, MipLoadOptions, uint8, UTexture, ETextureMipLoadOptions);
+	IMPLEMENT_NODE_ATTRIBUTE(FTextureNode, SRGB, bool, UTexture,);
+#if WITH_EDITORONLY_DATA
+	IMPLEMENT_NODE_ATTRIBUTE(FTextureNode, bUseLegacyGamma, bool, UTexture,);
+#endif //#if WITH_EDITORONLY_DATA
+	IMPLEMENT_NODE_ATTRIBUTE(FTextureNode, VirtualTextureStreaming, bool, UTexture,);
+	IMPLEMENT_NODE_ATTRIBUTE(FTextureNode, AddressX, uint8, UTexture2D, TEnumAsByte<enum TextureAddress>);
+	IMPLEMENT_NODE_ATTRIBUTE(FTextureNode, AddressY, uint8, UTexture2D, TEnumAsByte<enum TextureAddress>);
+
+#if WITH_EDITORONLY_DATA
+	//Level of Detail
+	IMPLEMENT_NODE_ATTRIBUTE(FTextureNode, bPreserveBorder, bool, UTexture,);
+	IMPLEMENT_NODE_ATTRIBUTE(FTextureNode, MipGenSettings, uint8, UTexture, TEnumAsByte<enum TextureMipGenSettings>);
+#endif //#if WITH_EDITORONLY_DATA
+	IMPLEMENT_NODE_ATTRIBUTE(FTextureNode, LODBias, int32, UTexture,);
+	IMPLEMENT_NODE_ATTRIBUTE(FTextureNode, LODGroup, uint8, UTexture, TEnumAsByte<enum TextureGroup>);
+	//TODO support per platform data in the FAttributeStorage, so we can set different value per platform at the pipeline stage, We set only the default value for now
+	IMPLEMENT_NODE_ATTRIBUTE(FTextureNode, Downscale, float, UTexture,);
+	IMPLEMENT_NODE_ATTRIBUTE(FTextureNode, DownscaleOptions, uint8, UTexture, ETextureDownscaleOptions);
+
+#if WITH_EDITORONLY_DATA
+	//Compositing
+	IMPLEMENT_NODE_ATTRIBUTE(FTextureNode, CompositeTextureMode, uint8, UTexture, TEnumAsByte<enum ECompositeTextureMode>);
+	IMPLEMENT_NODE_ATTRIBUTE(FTextureNode, CompositePower, float, UTexture,);
+#endif //#if WITH_EDITORONLY_DATA
+
+protected:
+	TSubclassOf<UTexture> AssetClass = nullptr;
+};
+
+
+}
\ No newline at end of file
diff --git a/Engine/Source/Developer/AssetTools/AssetTools.Build.cs b/Engine/Source/Developer/AssetTools/AssetTools.Build.cs
--- a/Engine/Source/Developer/AssetTools/AssetTools.Build.cs
+++ b/Engine/Source/Developer/AssetTools/AssetTools.Build.cs
@@ -35,6 +35,7 @@ public class AssetTools : ModuleRules
 				"RHI",
 				"MaterialEditor",
 				"ToolMenus",
+				"InterchangeCore",
             }
 		);
 
diff --git a/Engine/Source/Developer/AssetTools/Private/AssetTools.cpp b/Engine/Source/Developer/AssetTools/Private/AssetTools.cpp
--- a/Engine/Source/Developer/AssetTools/Private/AssetTools.cpp
+++ b/Engine/Source/Developer/AssetTools/Private/AssetTools.cpp
@@ -134,9 +134,11 @@
 #include "AssetVtConversion.h"
 #include "Misc/ConfigCacheIni.h"
 #include "Misc/BlacklistNames.h"
+#include "InterchangeManager.h"
 
 #if WITH_EDITOR
 #include "Subsystems/AssetEditorSubsystem.h"
+#include "Settings/EditorExperimentalSettings.h"
 #endif
 
 #define LOCTEXT_NAMESPACE "AssetTools"
@@ -1650,7 +1652,14 @@ TArray<UObject*> UAssetToolsImpl::ImportAssetsInternal(const TArray<FString>& Fi
 	TMap< FString, TArray<UFactory*> > ExtensionToFactoriesMap;
 
 	FScopedSlowTask SlowTask(Files.Num(), LOCTEXT("ImportSlowTask", "Importing"));
-	if (Files.Num() > 1)
+
+	bool bUseInterchangeFramework = false;
+	UInterchangeManager& InterchangeManager = UInterchangeManager::GetInterchangeManager();
+#if WITH_EDITOR
+	bUseInterchangeFramework = GetDefault<UEditorExperimentalSettings>()->bEnableInterchangeFramework;
+#endif
+
+	if (!bUseInterchangeFramework && Files.Num() > 1)
 	{	
 		//Always allow user to cancel the import task if they are importing multiple files.
 		//If we're importing a single file, then the factory policy will dictate if the import if cancelable.
@@ -1790,6 +1799,41 @@ TArray<UObject*> UAssetToolsImpl::ImportAssetsInternal(const TArray<FString>& Fi
 
 	TArray<UFactory*> UsedFactories;
 	bool bImportWasCancelled = false;
+	bool bOnlyInterchangeImport = bUseInterchangeFramework;
+	if(bUseInterchangeFramework)
+	{
+		for (int32 FileIdx = 0; FileIdx < FilesAndDestinations.Num(); ++FileIdx)
+		{
+			// Filename and DestinationPath will need to get santized before we create an asset out of them as they
+			// can be created out of sources that contain spaces and other invalid characters. Filename cannot be sanitized
+			// until other checks are done that rely on looking at the actual source file so sanitation is delayed.
+			const FString& Filename = FilesAndDestinations[FileIdx].Key;
+			const FString DestinationPath = ObjectTools::SanitizeObjectPath(FilesAndDestinations[FileIdx].Value);
+			{
+				Interchange::FScopedSourceData ScopedSourceData(Filename);
+				if (!InterchangeManager.CanTranslateSourceData(ScopedSourceData.GetSourceData()))
+				{
+					bOnlyInterchangeImport = false;
+					break;
+				}
+			}
+		}
+
+		if (!bOnlyInterchangeImport)
+		{
+			if (Files.Num() > 1)
+			{	
+				//Always allow user to cancel the import task if they are importing multiple files.
+				//If we're importing a single file, then the factory policy will dictate if the import if cancelable.
+				SlowTask.MakeDialog(true);
+			}
+		}
+		else
+		{
+			//Complete the slow task
+			SlowTask.CompletedWork = FilesAndDestinations.Num();
+		}
+	}
 	// Now iterate over the input files and use the same factory object for each file with the same extension
 	for(int32 FileIdx = 0; FileIdx < FilesAndDestinations.Num() && !bImportWasCancelled; ++FileIdx)
 	{
@@ -1798,6 +1842,38 @@ TArray<UObject*> UAssetToolsImpl::ImportAssetsInternal(const TArray<FString>& Fi
 		// until other checks are done that rely on looking at the actual source file so sanitation is delayed.
 		const FString& Filename = FilesAndDestinations[FileIdx].Key;
 		const FString DestinationPath = ObjectTools::SanitizeObjectPath(FilesAndDestinations[FileIdx].Value);
+		if (bUseInterchangeFramework)
+		{
+			Interchange::FScopedSourceData ScopedSourceData(Filename);
+			if (InterchangeManager.CanTranslateSourceData(ScopedSourceData.GetSourceData()))
+			{
+				auto PostImportedLambda = [bSyncToBrowser](UObject* ImportedObject)
+				{
+					if (ImportedObject && (bSyncToBrowser != false))
+					{
+						TArray<UObject*> ObjectArray;
+						ObjectArray.Add(ImportedObject);
+						UAssetToolsImpl::Get().SyncBrowserToAssets(ObjectArray);
+					}
+				};
+				FDelegateHandle PostImportHandle = InterchangeManager.OnAssetPostImport.AddLambda(PostImportedLambda);
+			
+				FImportAssetParameters ImportAssetParameters;
+				ImportAssetParameters.bIsAutomated = bAutomatedImport;
+				ImportAssetParameters.ReimportAsset = nullptr;
+
+				InterchangeManager.ImportAsset(DestinationPath, ScopedSourceData.GetSourceData(), ImportAssetParameters);
+				InterchangeManager.OnAssetPostImport.Remove(PostImportHandle);
+				//Import done, iterate the next file and destination
+				
+				//If we do not import only interchange file, update the progress for each interchange task
+				if (!bOnlyInterchangeImport)
+				{
+					SlowTask.EnterProgressFrame(1, FText::Format(LOCTEXT("Import_ImportingFile", "Importing \"{0}\"..."), FText::FromString(FPaths::GetBaseFilename(Filename))));
+				}
+				continue;
+			}
+		}
 		FString FileExtension = FPaths::GetExtension(Filename);
 		const TArray<UFactory*>* FactoriesPtr = ExtensionToFactoriesMap.Find(FileExtension);
 		UFactory* Factory = nullptr;
diff --git a/Engine/Source/Editor/UnrealEd/Classes/Settings/EditorExperimentalSettings.h b/Engine/Source/Editor/UnrealEd/Classes/Settings/EditorExperimentalSettings.h
--- a/Engine/Source/Editor/UnrealEd/Classes/Settings/EditorExperimentalSettings.h
+++ b/Engine/Source/Editor/UnrealEd/Classes/Settings/EditorExperimentalSettings.h
@@ -23,6 +23,10 @@ public:
 	UPROPERTY(EditAnywhere, config, Category = Performance, meta = (DisplayName = "Enable async texture compilation to improve PIE and map load time performance when texture compilation is required"))
 	bool bEnableAsyncTextureCompilation;
 
+	/** Enable interchange framework, the interchange framework is a new import system which can import asynchronously and in parallel. */
+	UPROPERTY(EditAnywhere, config, Category = Performance, meta = (DisplayName = "Enable interchange framework import"))
+	bool bEnableInterchangeFramework;
+
 	/** Allows the editor to run on HDR monitors on Windows 10 */
 	UPROPERTY(EditAnywhere, config, Category = HDR, meta = (ConfigRestartRequired = true, DisplayName = "Enable Editor Support for HDR Monitors"))
 	bool bHDREditor;
diff --git a/Engine/Source/Editor/UnrealEd/Private/Editor.cpp b/Engine/Source/Editor/UnrealEd/Private/Editor.cpp
--- a/Engine/Source/Editor/UnrealEd/Private/Editor.cpp
+++ b/Engine/Source/Editor/UnrealEd/Private/Editor.cpp
@@ -68,6 +68,12 @@
 #include "AutoReimport/AutoReimportUtilities.h"
 #include "AssetToolsModule.h"
 
+#include "InterchangeManager.h"
+
+#if WITH_EDITOR
+#include "Subsystems/AssetEditorSubsystem.h"
+#include "Settings/EditorExperimentalSettings.h"
+#endif
 
 #define LOCTEXT_NAMESPACE "UnrealEd.Editor"
 
@@ -249,6 +255,12 @@ bool FReimportManager::Reimport( UObject* Obj, bool bAskForNewFileIfMissing, boo
 	// Warn that were about to reimport, so prep for it
 	PreReimport.Broadcast( Obj );
 
+	bool bUseInterchangeFramework = false;
+	UInterchangeManager& InterchangeManager = UInterchangeManager::GetInterchangeManager();
+#if WITH_EDITOR
+	bUseInterchangeFramework = GetDefault<UEditorExperimentalSettings>()->bEnableInterchangeFramework;
+#endif
+
 	bool bSuccess = false;
 	if ( Obj )
 	{
@@ -368,6 +380,42 @@ bool FReimportManager::Reimport( UObject* Obj, bool bAskForNewFileIfMissing, boo
 
 			if ( bValidSourceFilename )
 			{
+				if (bUseInterchangeFramework)
+				{
+					Interchange::FScopedSourceData ScopedSourceData(SourceFilenames[0]);
+					if (InterchangeManager.CanTranslateSourceData(ScopedSourceData.GetSourceData()))
+					{
+						auto PostImportedLambda = [](UObject* ImportedObject)
+						{
+							if (ImportedObject)
+							{
+								TArray<UObject*> ObjectArray;
+								ObjectArray.Add(ImportedObject);
+								//UAssetToolsImpl::Get().SyncBrowserToAssets(ObjectArray);
+								GEditor->BroadcastObjectReimported(ImportedObject);
+								if (FEngineAnalytics::IsAvailable())
+								{
+									TArray<FAnalyticsEventAttribute> Attributes;
+									Attributes.Add(FAnalyticsEventAttribute(TEXT("ObjectType"), ImportedObject->GetClass()->GetName()));
+									FEngineAnalytics::GetProvider().RecordEvent(TEXT("Editor.Usage.AssetReimported"), Attributes);
+								}
+								//PostReimport.Broadcast(ImportedObject, true);
+								GEditor->RedrawAllViewports();
+							}
+						};
+						FDelegateHandle PostImportHandle = InterchangeManager.OnAssetPostImport.AddLambda(PostImportedLambda);
+
+						FImportAssetParameters ImportAssetParameters;
+						ImportAssetParameters.bIsAutomated = GIsAutomationTesting || FApp::IsUnattended() || IsRunningCommandlet() || GIsRunningUnattendedScript;
+						ImportAssetParameters.ReimportAsset = Obj;
+						InterchangeManager.ImportAsset(FString(), ScopedSourceData.GetSourceData(), ImportAssetParameters);
+						InterchangeManager.OnAssetPostImport.Remove(PostImportHandle);
+						return true;
+					}
+				}
+
+
+
 				// Do the reimport
 				EReimportResult::Type Result = CanReimportHandler->Reimport( Obj, SourceFileIndex );
 				if( Result == EReimportResult::Succeeded )
diff --git a/Engine/Source/Editor/UnrealEd/Private/Subsystems/ImportSubsystem.cpp b/Engine/Source/Editor/UnrealEd/Private/Subsystems/ImportSubsystem.cpp
--- a/Engine/Source/Editor/UnrealEd/Private/Subsystems/ImportSubsystem.cpp
+++ b/Engine/Source/Editor/UnrealEd/Private/Subsystems/ImportSubsystem.cpp
@@ -8,6 +8,8 @@
 #include "FileHelpers.h"
 #include "EditorReimportHandler.h"
 #include "Misc/BlacklistNames.h"
+#include "InterchangeManager.h"
+#include "InterchangeFactoryBase.h"
 
 class FImportFilesByPath : public IImportSubsystemTask
 {
@@ -110,12 +112,17 @@ UImportSubsystem::UImportSubsystem()
 
 void UImportSubsystem::Initialize(FSubsystemCollectionBase& Collection)
 {
-
+	//Hook on Interchange manager to know when to broadcast import/reimport
+	UInterchangeManager& InterchangeManager = UInterchangeManager::GetInterchangeManager();
+	InterchangeManager.OnAssetPostImport.AddUObject(this, &UImportSubsystem::InterchangeBroadcastAssetPostImport);
+	InterchangeManager.OnAssetPostReimport.AddUObject(this, &UImportSubsystem::InterchangeBroadcastAssetPostReimport);
 }
 
 void UImportSubsystem::Deinitialize()
 {
-
+	UInterchangeManager& InterchangeManager = UInterchangeManager::GetInterchangeManager();
+	InterchangeManager.OnAssetPostImport.RemoveAll(this);
+	InterchangeManager.OnAssetPostReimport.RemoveAll(this);
 }
 
 void UImportSubsystem::ImportNextTick(const TArray<FString>& Files, const FString& DestinationPath)
@@ -175,3 +182,14 @@ void UImportSubsystem::BroadcastAssetPostLODImport(UObject* InObject, int32 inLO
 	OnAssetPostLODImport_BP.Broadcast(InObject, inLODIndex);
 }
 PRAGMA_ENABLE_DEPRECATION_WARNINGS
+
+void UImportSubsystem::InterchangeBroadcastAssetPostImport(UObject* InCreatedObject)
+{
+	UFactory* NullFactory = nullptr;
+	BroadcastAssetPostImport(NullFactory, InCreatedObject);
+}
+
+void UImportSubsystem::InterchangeBroadcastAssetPostReimport(UObject* InCreatedObject)
+{
+	BroadcastAssetReimport(InCreatedObject);
+}
diff --git a/Engine/Source/Editor/UnrealEd/Public/Subsystems/ImportSubsystem.h b/Engine/Source/Editor/UnrealEd/Public/Subsystems/ImportSubsystem.h
--- a/Engine/Source/Editor/UnrealEd/Public/Subsystems/ImportSubsystem.h
+++ b/Engine/Source/Editor/UnrealEd/Public/Subsystems/ImportSubsystem.h
@@ -69,6 +69,9 @@ public:
 
 private:
 
+	void InterchangeBroadcastAssetPostImport(UObject* InCreatedObject);
+	void InterchangeBroadcastAssetPostReimport(UObject* InCreatedObject);
+
 	/* Run deferred logic waiting to be run next tick */
 	void HandleNextTick();
 
diff --git a/Engine/Source/Editor/UnrealEd/UnrealEd.Build.cs b/Engine/Source/Editor/UnrealEd/UnrealEd.Build.cs
--- a/Engine/Source/Editor/UnrealEd/UnrealEd.Build.cs
+++ b/Engine/Source/Editor/UnrealEd/UnrealEd.Build.cs
@@ -109,6 +109,7 @@ public class UnrealEd : ModuleRules
                 "InteractiveToolsFramework",
 				"ToolMenusEditor",
 				"StatusBar",
+				"InterchangeCore",
 			}
 		);
 
diff --git a/Engine/Source/Runtime/Interchange/Core/InterchangeCore.Build.cs b/Engine/Source/Runtime/Interchange/Core/InterchangeCore.Build.cs
new file mode 100644
--- /dev/null
+++ b/Engine/Source/Runtime/Interchange/Core/InterchangeCore.Build.cs
@@ -0,0 +1,23 @@
+// Copyright Epic Games, Inc. All Rights Reserved.
+
+using UnrealBuildTool;
+
+public class InterchangeCore : ModuleRules
+{
+	public InterchangeCore(ReadOnlyTargetRules Target) : base(Target)
+	{
+		PublicDependencyModuleNames.AddRange(
+				new string[] {
+					"Core",
+					"CoreUObject",
+					"Engine",
+				}
+			);
+
+		PrivateDependencyModuleNames.AddRange(
+				new string[] {
+					"AssetRegistry",
+				}
+			);
+	}
+}
diff --git a/Engine/Source/Runtime/Interchange/Core/Private/InterchangeCoreModule.cpp b/Engine/Source/Runtime/Interchange/Core/Private/InterchangeCoreModule.cpp
new file mode 100644
--- /dev/null
+++ b/Engine/Source/Runtime/Interchange/Core/Private/InterchangeCoreModule.cpp
@@ -0,0 +1,30 @@
+// Copyright Epic Games, Inc. All Rights Reserved.
+#include "CoreMinimal.h"
+#include "InterchangeLogPrivate.h"
+#include "Modules/ModuleInterface.h"
+#include "Modules/ModuleManager.h"
+
+DEFINE_LOG_CATEGORY(LogInterchangeCore);
+
+
+/**
+ * Implements the Interchange module.
+ */
+class FInterchangeCoreModule
+	: public IModuleInterface
+{
+public:
+
+	// IModuleInterface interface
+
+	virtual void StartupModule() override { }
+	virtual void ShutdownModule() override { }
+
+	virtual bool SupportsDynamicReloading() override
+	{
+		return false;
+	}
+};
+
+
+IMPLEMENT_MODULE(FInterchangeCoreModule, InterchangeCore);
diff --git a/Engine/Source/Runtime/Interchange/Core/Private/InterchangeLogPrivate.h b/Engine/Source/Runtime/Interchange/Core/Private/InterchangeLogPrivate.h
new file mode 100644
--- /dev/null
+++ b/Engine/Source/Runtime/Interchange/Core/Private/InterchangeLogPrivate.h
@@ -0,0 +1,7 @@
+// Copyright Epic Games, Inc. All Rights Reserved.
+
+#pragma once
+
+#include "CoreMinimal.h"
+
+DECLARE_LOG_CATEGORY_EXTERN(LogInterchangeCore, Log, All);
diff --git a/Engine/Source/Runtime/Interchange/Core/Private/InterchangeManager.cpp b/Engine/Source/Runtime/Interchange/Core/Private/InterchangeManager.cpp
new file mode 100644
--- /dev/null
+++ b/Engine/Source/Runtime/Interchange/Core/Private/InterchangeManager.cpp
@@ -0,0 +1,491 @@
+// Copyright Epic Games, Inc. All Rights Reserved.
+
+#pragma once
+
+#include "InterchangeManager.h"
+
+#include "AssetDataTagMap.h"
+#include "AssetRegistryModule.h"
+#include "CoreMinimal.h"
+#include "Nodes/BaseNodeContainer.h"
+#include "InterchangeFactoryBase.h"
+#include "InterchangeSourceData.h"
+#include "InterchangeTranslatorBase.h"
+#include "InterchangeWriterBase.h"
+#include "Internationalization/Internationalization.h"
+#include "Misc/App.h"
+#include "Misc/AsyncTaskNotification.h"
+#include "Misc/MessageDialog.h"
+#include "Misc/ScopedSlowTask.h"
+#include "PackageUtils/PackageUtils.h"
+#include "Tasks/InterchangeTaskParsing.h"
+#include "Tasks/InterchangeTaskPipeline.h"
+#include "Tasks/InterchangeTaskTranslator.h"
+#include "UObject/NameTypes.h"
+#include "UObject/Object.h"
+#include "UObject/ObjectMacros.h"
+#include "UObject/WeakObjectPtrTemplates.h"
+
+namespace InternalInterchangePrivate
+{
+	const FLogCategoryBase* GetLogInterchangePtr()
+	{
+#if NO_LOGGING
+		return nullptr;
+#else
+		return &LogInterchangeCore;
+#endif
+	}
+}
+
+Interchange::FScopedSourceData::FScopedSourceData(const FString& Filename)
+{
+	//Found the translator
+	SourceDataPtr = TStrongObjectPtr<UInterchangeSourceData>(UInterchangeManager::GetInterchangeManager().CreateSourceData(Filename));
+	check(SourceDataPtr.IsValid());
+}
+
+UInterchangeSourceData* Interchange::FScopedSourceData::GetSourceData() const
+{
+	return SourceDataPtr.Get();
+}
+
+Interchange::FScopedTranslator::FScopedTranslator(const UInterchangeSourceData* SourceData)
+{
+	//Found the translator
+	ScopedTranslatorPtr = TStrongObjectPtr<UInterchangeTranslatorBase>(UInterchangeManager::GetInterchangeManager().GetTranslatorForSourceData(SourceData));
+}
+
+UInterchangeTranslatorBase* Interchange::FScopedTranslator::GetTranslator()
+{
+	return ScopedTranslatorPtr.Get();
+}
+
+void Interchange::FImportAsyncHelper::AddReferencedObjects(FReferenceCollector& Collector)
+{
+	for (UInterchangeSourceData* SourceData : SourceDatas)
+	{
+		Collector.AddReferencedObject(SourceData);
+	}
+	for (UInterchangeTranslatorBase* Translator : Translators)
+	{
+		Collector.AddReferencedObject(Translator);
+	}
+	for (UInterchangePipelineBase* Pipeline : Pipelines)
+	{
+		Collector.AddReferencedObject(Pipeline);
+	}
+	for (UInterchangeFactoryBase* Factory : Factories)
+	{
+		Collector.AddReferencedObject(Factory);
+	}
+}
+
+void Interchange::FImportAsyncHelper::CleanUp()
+{
+	//Release the node graph adapters
+	BaseNodeContainerAdapters.Empty();
+
+	for (UInterchangeSourceData* SourceData : SourceDatas)
+	{
+		if (SourceData)
+		{
+			SourceData->RemoveFromRoot();
+			SourceData->MarkPendingKill();
+			SourceData = nullptr;
+		}
+	}
+	SourceDatas.Empty();
+
+	for (UInterchangeTranslatorBase* Translator : Translators)
+	{
+		if(Translator)
+		{
+			Translator->RemoveFromRoot();
+			Translator->MarkPendingKill();
+			Translator = nullptr;
+		}
+	}
+	Translators.Empty();
+
+	for (UInterchangePipelineBase* Pipeline : Pipelines)
+	{
+		if(Pipeline)
+		{
+			Pipeline->RemoveFromRoot();
+			Pipeline->MarkPendingKill();
+			Pipeline = nullptr;
+		}
+	}
+	Pipelines.Empty();
+
+	//Factories are not instantiate, we use the registered one directly
+	Factories.Empty();
+}
+
+void Interchange::SanitizeInvalidChar(FString& String)
+{
+	const TCHAR* InvalidChar = *FString(INVALID_OBJECTNAME_CHARACTERS);
+	while (*InvalidChar)
+	{
+		String.ReplaceCharInline(*InvalidChar, TCHAR('_'), ESearchCase::CaseSensitive);
+		++InvalidChar;
+	}
+}
+
+bool UInterchangeManager::RegisterTranslator(const UClass* TranslatorClass)
+{
+	if(!TranslatorClass)
+	{
+		return false;
+	}
+
+	if(RegisteredTranslators.Contains(TranslatorClass))
+	{
+		return true;
+	}
+	UInterchangeTranslatorBase* TranslatorToRegister = NewObject<UInterchangeTranslatorBase>(GetTransientPackage(), TranslatorClass, NAME_None);
+	if(!TranslatorToRegister)
+	{
+		return false;
+	}
+	RegisteredTranslators.Add(TranslatorClass, TranslatorToRegister);
+	return true;
+}
+
+bool UInterchangeManager::RegisterFactory(const UClass* FactoryClass)
+{
+	if (!FactoryClass)
+	{
+		return false;
+	}
+
+	UInterchangeFactoryBase* FactoryToRegister = NewObject<UInterchangeFactoryBase>(GetTransientPackage(), FactoryClass, NAME_None);
+	if (!FactoryToRegister)
+	{
+		return false;
+	}
+	if (FactoryToRegister->GetFactoryClass() == nullptr || RegisteredFactories.Contains(FactoryToRegister->GetFactoryClass()))
+	{
+		FactoryToRegister->MarkPendingKill();
+		return FactoryToRegister->GetFactoryClass() == nullptr ? false : true;
+	}
+	RegisteredFactories.Add(FactoryToRegister->GetFactoryClass(), FactoryToRegister);
+	return true;
+}
+
+bool UInterchangeManager::RegisterWriter(const UClass* WriterClass)
+{
+	if (!WriterClass)
+	{
+		return false;
+	}
+
+	if (RegisteredFactories.Contains(WriterClass))
+	{
+		return true;
+	}
+	UInterchangeWriterBase* WriterToRegister = NewObject<UInterchangeWriterBase>(GetTransientPackage(), WriterClass, NAME_None);
+	if (!WriterToRegister)
+	{
+		return false;
+	}
+	RegisteredWriters.Add(WriterClass, WriterToRegister);
+	return true;
+}
+
+bool UInterchangeManager::CanTranslateSourceData(const UInterchangeSourceData* SourceData) const
+{
+	//Found the translator
+	Interchange::FScopedTranslator ScopeDataTranslator(SourceData);
+	const UInterchangeTranslatorBase* SourceDataTranslator = ScopeDataTranslator.GetTranslator();
+	if (SourceDataTranslator)
+	{
+		return true;
+	}
+	return false;
+}
+
+bool UInterchangeManager::ImportAsset(const FString& ContentPath, const UInterchangeSourceData* SourceData, const FImportAssetParameters& ImportAssetParameters)
+{
+	FString PackageBasePath = ContentPath;
+	if(!ImportAssetParameters.ReimportAsset)
+	{
+		Interchange::SanitizeInvalidChar(PackageBasePath);
+	}
+
+	bool bCanShowDialog = !ImportAssetParameters.bIsAutomated && !IsAttended();
+
+	//Create a task for every source data
+	Interchange::FImportAsyncHelperData TaskData;
+	TaskData.bIsAutomated = ImportAssetParameters.bIsAutomated;
+	TaskData.ImportType = Interchange::EImportType::ImportType_Asset;
+	TaskData.ReimportObject = ImportAssetParameters.ReimportAsset;
+	TWeakPtr<Interchange::FImportAsyncHelper> WeakAsyncHelper = CreateAsyncHelper(TaskData);
+	check(WeakAsyncHelper.IsValid());
+	TSharedPtr<Interchange::FImportAsyncHelper> AsyncHelper = WeakAsyncHelper.Pin();
+		
+	FText TitleText = NSLOCTEXT("Interchange", "Asynchronous_import_start", "Importing");
+	if(!Notification.IsValid())
+	{
+		FAsyncTaskNotificationConfig NotificationConfig;
+		NotificationConfig.bIsHeadless = false;
+		NotificationConfig.bKeepOpenOnFailure = true;
+		NotificationConfig.TitleText = TitleText;
+		NotificationConfig.LogCategory = InternalInterchangePrivate::GetLogInterchangePtr();
+		NotificationConfig.bCanCancel.Set(false);
+		NotificationConfig.bKeepOpenOnFailure.Set(true);
+
+		Notification = MakeShared<FAsyncTaskNotification>(NotificationConfig);
+		Notification->SetNotificationState(FAsyncNotificationStateData(TitleText, FText::GetEmpty(), EAsyncTaskNotificationState::Pending));
+	}
+	
+	//Create a duplicate of the source data, we need to be multithread safe so we copy it to control the life cycle. The async helper will hold it and delete it when the import task will be completed.
+	UInterchangeSourceData* DuplicateSourceData = Cast<UInterchangeSourceData>(StaticDuplicateObject(SourceData, GetTransientPackage()));
+	//Array of source data to build one graph per source
+	AsyncHelper->SourceDatas.Add(DuplicateSourceData);
+		
+
+	//Get all the translators for the source datas
+	for (int32 SourceDataIndex = 0; SourceDataIndex < AsyncHelper->SourceDatas.Num(); ++SourceDataIndex)
+	{
+		ensure(AsyncHelper->Translators.Add(GetTranslatorForSourceData(AsyncHelper->SourceDatas[SourceDataIndex])) == SourceDataIndex);
+	}
+
+	//Create the node graphs and the adapters (TStrongObjectPtr must be use only on the main thread, so we create them here before filling the node Containers)
+	AsyncHelper->BaseNodeContainers.AddDefaulted(AsyncHelper->SourceDatas.Num());
+	for (int32 SourceDataIndex = 0; SourceDataIndex < AsyncHelper->SourceDatas.Num(); ++SourceDataIndex)
+	{
+		AsyncHelper->BaseNodeContainerAdapters.Add(TStrongObjectPtr<UInterchangeBaseNodeContainerAdapter>(NewObject<UInterchangeBaseNodeContainerAdapter>(GetTransientPackage(), NAME_None)));
+		check(AsyncHelper->BaseNodeContainerAdapters[SourceDataIndex].IsValid());
+		AsyncHelper->BaseNodeContainerAdapters[SourceDataIndex].Get()->SetBaseNodeContainer(&AsyncHelper->BaseNodeContainers[SourceDataIndex]);
+	}
+	
+
+	//Get all pipeline candidate we want for this import
+	TArray<UClass*> PipelineCandidates;
+	FindPipelineCandidate(PipelineCandidates);
+
+	// Stack all pipelines, for this import proto. TODO: We need to be able to control which pipeline we use for the import
+	// It can be set in the project settings and can also be set by a UI where the user create the pipeline stack he want.
+	// This should be a list of available pipelines that can be drop into a stack where you can control the order.
+	for (int32 GraphPipelineNumber = 0; GraphPipelineNumber < PipelineCandidates.Num(); ++GraphPipelineNumber)
+	{
+		UInterchangePipelineBase* GeneratedPipeline = NewObject<UInterchangePipelineBase>(GetTransientPackage(), PipelineCandidates[GraphPipelineNumber], NAME_None, RF_NoFlags);
+		AsyncHelper->Pipelines.Add(GeneratedPipeline);
+	}
+
+	//Create/Start import tasks
+	FGraphEventArray PipelinePrerequistes;
+	check(AsyncHelper->Translators.Num() == AsyncHelper->SourceDatas.Num());
+	for (int32 SourceDataIndex = 0; SourceDataIndex < AsyncHelper->SourceDatas.Num(); ++SourceDataIndex)
+	{
+		int32 TranslatorTaskIndex = AsyncHelper->TranslatorTasks.Add(TGraphTask<Interchange::FTaskTranslator>::CreateTask().ConstructAndDispatchWhenReady(SourceDataIndex, WeakAsyncHelper));
+		PipelinePrerequistes.Add(AsyncHelper->TranslatorTasks[TranslatorTaskIndex]);
+	}
+		
+	FGraphEventArray GraphParsingPrerequistes;
+	for(int32 GraphPipelineIndex = 0; GraphPipelineIndex < AsyncHelper->Pipelines.Num(); ++GraphPipelineIndex)
+	{
+		UInterchangePipelineBase* GraphPipeline = AsyncHelper->Pipelines[GraphPipelineIndex];
+		TWeakObjectPtr<UInterchangePipelineBase> WeakPipelinePtr = GraphPipeline;
+		int32 GraphPipelineTaskIndex = INDEX_NONE;
+		GraphPipelineTaskIndex = AsyncHelper->PipelineTasks.Add(TGraphTask<Interchange::FTaskPipeline>::CreateTask(&PipelinePrerequistes).ConstructAndDispatchWhenReady(WeakPipelinePtr, WeakAsyncHelper));
+		//Ensure we run the pipeline in the same order we create the task, since pipeline modify the node container, its important that its not process in parallel, Adding the one we start to the prerequisites
+		//is the way to go here
+		PipelinePrerequistes.Add(AsyncHelper->PipelineTasks[GraphPipelineTaskIndex]);
+
+		//Add pipeline to the graph parsing prerequisites
+		GraphParsingPrerequistes.Add(AsyncHelper->PipelineTasks[GraphPipelineTaskIndex]);
+	}
+
+	if (GraphParsingPrerequistes.Num() > 0)
+	{
+		AsyncHelper->ParsingTask = TGraphTask<Interchange::FTaskParsing>::CreateTask(&GraphParsingPrerequistes).ConstructAndDispatchWhenReady(this, PackageBasePath, WeakAsyncHelper);
+	}
+	else
+	{
+		//Fallback on the translator pipeline prerequisites (translator must be done if there is no pipeline)
+		AsyncHelper->ParsingTask = TGraphTask<Interchange::FTaskParsing>::CreateTask(&PipelinePrerequistes).ConstructAndDispatchWhenReady(this, PackageBasePath, WeakAsyncHelper);
+	}
+
+	//The graph parsing task will create the FCreateAssetTask that will run after them, the FAssetImportTask will call the appropriate Post asset import pipeline when the asset is completed
+
+	return true;
+}
+
+bool UInterchangeManager::ImportScene(const FString& ImportContext, const UInterchangeSourceData* SourceData, bool bIsReimport, bool bIsAutomated)
+{
+	return false;
+}
+
+bool UInterchangeManager::ExportAsset(const UObject* Asset, bool bIsAutomated)
+{
+	return false;
+}
+
+bool UInterchangeManager::ExportScene(const UObject* World, bool bIsAutomated)
+{
+	return false;
+}
+
+UInterchangeSourceData* UInterchangeManager::CreateSourceData(const FString& InFileName)
+{
+	UInterchangeSourceData* SourceDataAsset = NewObject<UInterchangeSourceData>(GetTransientPackage(), NAME_None);
+	if(!InFileName.IsEmpty())
+	{
+		SourceDataAsset->SetFilename(InFileName);
+	}
+	return SourceDataAsset;
+}
+
+TWeakPtr<Interchange::FImportAsyncHelper> UInterchangeManager::CreateAsyncHelper(const Interchange::FImportAsyncHelperData& Data)
+{
+	TSharedPtr<Interchange::FImportAsyncHelper> AsyncHelper = MakeShared<Interchange::FImportAsyncHelper>();
+	//Copy the task data
+	AsyncHelper->TaskData = Data;
+	int32 AsyncHelperIndex = ImportTasks.Add(AsyncHelper);
+	//Update the asynchronous notification
+	if (Notification.IsValid())
+	{
+		int32 ImportTaskNumber = ImportTasks.Num();
+		FString ImportTaskNumberStr = TEXT(" (") + FString::FromInt(ImportTaskNumber) + TEXT(")");
+		Notification->SetProgressText(FText::FromString(ImportTaskNumberStr));
+	}
+	
+	return ImportTasks[AsyncHelperIndex];
+}
+
+void UInterchangeManager::ReleaseAsyncHelper(TWeakPtr<Interchange::FImportAsyncHelper> AsyncHelper)
+{
+	check(AsyncHelper.IsValid());
+	ImportTasks.RemoveSingle(AsyncHelper.Pin());
+	check(!AsyncHelper.IsValid());
+
+	int32 ImportTaskNumber = ImportTasks.Num();
+	FString ImportTaskNumberStr = TEXT(" (") + FString::FromInt(ImportTaskNumber) + TEXT(")");
+	if (ImportTasks.Num() == 0 && Notification.IsValid())
+	{
+		FText TitleText = NSLOCTEXT("Interchange", "Asynchronous_import_end", "Import Done");
+		//TODO make sure any error are reported so we can control success or not
+		const bool bSuccess = true;
+		Notification->SetComplete(TitleText, FText::GetEmpty(), bSuccess);
+		Notification = nullptr; //This should delete the notification
+	}
+	else if(Notification.IsValid())
+	{
+		Notification->SetProgressText(FText::FromString(ImportTaskNumberStr));
+	}
+}
+
+UInterchangeTranslatorBase* UInterchangeManager::GetTranslatorForSourceData(const UInterchangeSourceData* SourceData) const
+{
+	if (RegisteredTranslators.Num() == 0)
+	{
+		return nullptr;
+	}
+	//Found the translator
+	for (const auto Kvp : RegisteredTranslators)
+	{
+		if (Kvp.Value->CanImportSourceData(SourceData))
+		{
+			UInterchangeTranslatorBase* SourceDataTranslator = NewObject<UInterchangeTranslatorBase>(GetTransientPackage(), Kvp.Key, NAME_None);
+			return SourceDataTranslator;
+		}
+	}
+	return nullptr;
+}
+
+bool UInterchangeManager::IsAttended()
+{
+	if (FApp::IsGame())
+	{
+		return false;
+	}
+	if (FApp::IsUnattended())
+	{
+		return false;
+	}
+	return true;
+}
+
+void UInterchangeManager::FindPipelineCandidate(TArray<UClass*>& PipelineCandidates)
+{
+	//Find in memory pipeline class
+	for (TObjectIterator< UClass > ClassIt; ClassIt; ++ClassIt)
+	{
+		UClass* Class = *ClassIt;
+		// Only interested in native C++ classes
+		if (!Class->IsNative())
+		{
+			continue;
+		}
+		// Ignore deprecated
+		if (Class->HasAnyClassFlags(CLASS_Deprecated | CLASS_NewerVersionExists))
+		{
+			continue;
+		}
+
+		// Check this class is a subclass of Base and not the base itself
+		if (Class == UInterchangePipelineBase::StaticClass() || !Class->IsChildOf(UInterchangePipelineBase::StaticClass()))
+		{
+			continue;
+		}
+
+		//We found a candidate
+		PipelineCandidates.AddUnique(Class);
+	}
+
+	// Load the asset registry module
+	FAssetRegistryModule& AssetRegistryModule = FModuleManager::LoadModuleChecked< FAssetRegistryModule >(FName("AssetRegistry"));
+	IAssetRegistry& AssetRegistry = AssetRegistryModule.Get();
+
+	TArray< FString > ContentPaths;
+	ContentPaths.Add(TEXT("/Game"));
+	//TODO do we have an other alternative, this call is synchronous and will wait unitl the registry database have finish the initial scan. If there is a lot of asset it can take multiple second the first time we call it.
+	AssetRegistry.ScanPathsSynchronous(ContentPaths);
+
+	FName BaseClassName = UInterchangePipelineBase::StaticClass()->GetFName();
+
+	// Use the asset registry to get the set of all class names deriving from Base
+	TSet< FName > DerivedNames;
+	{
+		TArray< FName > BaseNames;
+		BaseNames.Add(BaseClassName);
+
+		TSet< FName > Excluded;
+		AssetRegistry.GetDerivedClassNames(BaseNames, Excluded, DerivedNames);
+	}
+
+	FARFilter Filter;
+	Filter.ClassNames.Add(UBlueprint::StaticClass()->GetFName());
+	Filter.bRecursiveClasses = true;
+	Filter.bRecursivePaths = true;
+
+	TArray< FAssetData > AssetList;
+	AssetRegistry.GetAssets(Filter, AssetList);
+
+	// Iterate over retrieved blueprint assets
+	for (const FAssetData& Asset : AssetList)
+	{
+		//Only get the asset with the native parent class using UInterchangePipelineBase
+		FAssetDataTagMapSharedView::FFindTagResult GeneratedClassPath = Asset.TagsAndValues.FindTag(TEXT("GeneratedClass"));
+		if (GeneratedClassPath.IsSet())
+		{
+			// Convert path to just the name part
+			const FString ClassObjectPath = FPackageName::ExportTextPathToObjectPath(*GeneratedClassPath.GetValue());
+			const FString ClassName = FPackageName::ObjectPathToObjectName(ClassObjectPath);
+
+			// Check if this class is in the derived set
+			if (!DerivedNames.Contains(*ClassName))
+			{
+				continue;
+			}
+
+			UBlueprint* Blueprint = Cast<UBlueprint>(Asset.GetAsset());
+			check(Blueprint);
+			check(Blueprint->ParentClass == UInterchangePipelineBase::StaticClass());
+			PipelineCandidates.AddUnique(Blueprint->GeneratedClass);
+		}
+	}
+}
\ No newline at end of file
diff --git a/Engine/Source/Runtime/Interchange/Core/Private/Nodes/InterchangeBaseNodeContainerAdapter.cpp b/Engine/Source/Runtime/Interchange/Core/Private/Nodes/InterchangeBaseNodeContainerAdapter.cpp
new file mode 100644
--- /dev/null
+++ b/Engine/Source/Runtime/Interchange/Core/Private/Nodes/InterchangeBaseNodeContainerAdapter.cpp
@@ -0,0 +1,369 @@
+// Copyright Epic Games, Inc. All Rights Reserved.
+
+#pragma once
+
+#include "Nodes/InterchangeBaseNodeContainerAdapter.h"
+
+#include "CoreMinimal.h"
+#include "UObject/Object.h"
+#include "UObject/ObjectMacros.h"
+#include "Nodes/BaseNode.h"
+#include "Nodes/BaseNodeContainer.h"
+
+bool UInterchangeBaseNodeContainerAdapter::GetNodeAttributeValueBool(const FName NodeIdentity, const FName AttributeName, bool& OutValue)
+{
+	return GetNodeAttributeValue(NodeIdentity, AttributeName, OutValue);
+}
+
+bool UInterchangeBaseNodeContainerAdapter::GetNodeAttributeValueFloat(const FName NodeIdentity, const FName AttributeName, float& OutValue)
+{
+	return GetNodeAttributeValue(NodeIdentity, AttributeName, OutValue);
+}
+
+bool UInterchangeBaseNodeContainerAdapter::GetNodeAttributeValueDouble(const FName NodeIdentity, const FName AttributeName, double& OutValue)
+{
+	return GetNodeAttributeValue(NodeIdentity, AttributeName, OutValue);
+}
+
+bool UInterchangeBaseNodeContainerAdapter::GetNodeAttributeValueInt8(const FName NodeIdentity, const FName AttributeName, int8& OutValue)
+{
+	return GetNodeAttributeValue(NodeIdentity, AttributeName, OutValue);
+}
+
+bool UInterchangeBaseNodeContainerAdapter::GetNodeAttributeValueInt16(const FName NodeIdentity, const FName AttributeName, int16& OutValue)
+{
+	return GetNodeAttributeValue(NodeIdentity, AttributeName, OutValue);
+}
+
+bool UInterchangeBaseNodeContainerAdapter::GetNodeAttributeValueInt32(const FName NodeIdentity, const FName AttributeName, int32& OutValue)
+{
+	return GetNodeAttributeValue(NodeIdentity, AttributeName, OutValue);
+}
+
+bool UInterchangeBaseNodeContainerAdapter::GetNodeAttributeValueInt64(const FName NodeIdentity, const FName AttributeName, int64& OutValue)
+{
+	return GetNodeAttributeValue(NodeIdentity, AttributeName, OutValue);
+}
+
+bool UInterchangeBaseNodeContainerAdapter::GetNodeAttributeValueUint8(const FName NodeIdentity, const FName AttributeName, uint8& OutValue)
+{
+	return GetNodeAttributeValue(NodeIdentity, AttributeName, OutValue);
+}
+
+bool UInterchangeBaseNodeContainerAdapter::GetNodeAttributeValueUint16(const FName NodeIdentity, const FName AttributeName, uint16& OutValue)
+{
+	return GetNodeAttributeValue(NodeIdentity, AttributeName, OutValue);
+}
+
+bool UInterchangeBaseNodeContainerAdapter::GetNodeAttributeValueUint32(const FName NodeIdentity, const FName AttributeName, uint32& OutValue)
+{
+	return GetNodeAttributeValue(NodeIdentity, AttributeName, OutValue);
+}
+
+bool UInterchangeBaseNodeContainerAdapter::GetNodeAttributeValueUint64(const FName NodeIdentity, const FName AttributeName, uint64& OutValue)
+{
+	return GetNodeAttributeValue(NodeIdentity, AttributeName, OutValue);
+}
+
+bool UInterchangeBaseNodeContainerAdapter::GetNodeAttributeValueVector(const FName NodeIdentity, const FName AttributeName, FVector& OutValue)
+{
+	return GetNodeAttributeValue(NodeIdentity, AttributeName, OutValue);
+}
+
+bool UInterchangeBaseNodeContainerAdapter::GetNodeAttributeValueVector2D(const FName NodeIdentity, const FName AttributeName, FVector2D& OutValue)
+{
+	return GetNodeAttributeValue(NodeIdentity, AttributeName, OutValue);
+}
+
+bool UInterchangeBaseNodeContainerAdapter::GetNodeAttributeValueMatrix(const FName NodeIdentity, const FName AttributeName, FMatrix& OutValue)
+{
+	return GetNodeAttributeValue(NodeIdentity, AttributeName, OutValue);
+}
+
+bool UInterchangeBaseNodeContainerAdapter::GetNodeAttributeValueBox(const FName NodeIdentity, const FName AttributeName, FBox& OutValue)
+{
+	return GetNodeAttributeValue(NodeIdentity, AttributeName, OutValue);
+}
+
+bool UInterchangeBaseNodeContainerAdapter::GetNodeAttributeValueBoxSphereBound(const FName NodeIdentity, const FName AttributeName, FBoxSphereBounds& OutValue)
+{
+	return GetNodeAttributeValue(NodeIdentity, AttributeName, OutValue);
+}
+
+bool UInterchangeBaseNodeContainerAdapter::GetNodeAttributeValueArray(const FName NodeIdentity, const FName AttributeName, TArray<uint8>& OutValue)
+{
+	return GetNodeAttributeValue(NodeIdentity, AttributeName, OutValue);
+}
+
+bool UInterchangeBaseNodeContainerAdapter::GetNodeAttributeValueArray64(const FName NodeIdentity, const FName AttributeName, TArray64<uint8>& OutValue)
+{
+	return GetNodeAttributeValue(NodeIdentity, AttributeName, OutValue);
+}
+
+bool UInterchangeBaseNodeContainerAdapter::GetNodeAttributeValueColor(const FName NodeIdentity, const FName AttributeName, FColor& OutValue)
+{
+	return GetNodeAttributeValue(NodeIdentity, AttributeName, OutValue);
+}
+
+bool UInterchangeBaseNodeContainerAdapter::GetNodeAttributeValueDateTime(const FName NodeIdentity, const FName AttributeName, FDateTime& OutValue)
+{
+	return GetNodeAttributeValue(NodeIdentity, AttributeName, OutValue);
+}
+
+bool UInterchangeBaseNodeContainerAdapter::GetNodeAttributeValueEnum(const FName NodeIdentity, const FName AttributeName, uint8& OutValue)
+{
+	return GetNodeAttributeValue(NodeIdentity, AttributeName, OutValue);
+}
+
+bool UInterchangeBaseNodeContainerAdapter::GetNodeAttributeValueGuid(const FName NodeIdentity, const FName AttributeName, FGuid& OutValue)
+{
+	return GetNodeAttributeValue(NodeIdentity, AttributeName, OutValue);
+}
+
+bool UInterchangeBaseNodeContainerAdapter::GetNodeAttributeValueIntPoint(const FName NodeIdentity, const FName AttributeName, FIntPoint& OutValue)
+{
+	return GetNodeAttributeValue(NodeIdentity, AttributeName, OutValue);
+}
+
+bool UInterchangeBaseNodeContainerAdapter::GetNodeAttributeValueIntVecor(const FName NodeIdentity, const FName AttributeName, FIntVector& OutValue)
+{
+	return GetNodeAttributeValue(NodeIdentity, AttributeName, OutValue);
+}
+
+bool UInterchangeBaseNodeContainerAdapter::GetNodeAttributeValueLinearColor(const FName NodeIdentity, const FName AttributeName, FLinearColor& OutValue)
+{
+	return GetNodeAttributeValue(NodeIdentity, AttributeName, OutValue);
+}
+
+bool UInterchangeBaseNodeContainerAdapter::GetNodeAttributeValuePlane(const FName NodeIdentity, const FName AttributeName, FPlane& OutValue)
+{
+	return GetNodeAttributeValue(NodeIdentity, AttributeName, OutValue);
+}
+
+bool UInterchangeBaseNodeContainerAdapter::GetNodeAttributeValueQuat(const FName NodeIdentity, const FName AttributeName, FQuat& OutValue)
+{
+	return GetNodeAttributeValue(NodeIdentity, AttributeName, OutValue);
+}
+
+bool UInterchangeBaseNodeContainerAdapter::GetNodeAttributeValueName(const FName NodeIdentity, const FName AttributeName, FName& OutValue)
+{
+	return GetNodeAttributeValue(NodeIdentity, AttributeName, OutValue);
+}
+
+bool UInterchangeBaseNodeContainerAdapter::GetNodeAttributeValueRandomStream(const FName NodeIdentity, const FName AttributeName, FRandomStream& OutValue)
+{
+	return GetNodeAttributeValue(NodeIdentity, AttributeName, OutValue);
+}
+
+bool UInterchangeBaseNodeContainerAdapter::GetNodeAttributeValueRotator(const FName NodeIdentity, const FName AttributeName, FRotator& OutValue)
+{
+	return GetNodeAttributeValue(NodeIdentity, AttributeName, OutValue);
+}
+
+bool UInterchangeBaseNodeContainerAdapter::GetNodeAttributeValueString(const FName NodeIdentity, const FName AttributeName, FString& OutValue)
+{
+	return GetNodeAttributeValue(NodeIdentity, AttributeName, OutValue);
+}
+
+bool UInterchangeBaseNodeContainerAdapter::GetNodeAttributeValueTimespan(const FName NodeIdentity, const FName AttributeName, FTimespan& OutValue)
+{
+	return GetNodeAttributeValue(NodeIdentity, AttributeName, OutValue);
+}
+
+bool UInterchangeBaseNodeContainerAdapter::GetNodeAttributeValueTransform(const FName NodeIdentity, const FName AttributeName, FTransform& OutValue)
+{
+	return GetNodeAttributeValue(NodeIdentity, AttributeName, OutValue);
+}
+
+bool UInterchangeBaseNodeContainerAdapter::GetNodeAttributeValueTwoVectors(const FName NodeIdentity, const FName AttributeName, FTwoVectors& OutValue)
+{
+	return GetNodeAttributeValue(NodeIdentity, AttributeName, OutValue);
+}
+
+bool UInterchangeBaseNodeContainerAdapter::GetNodeAttributeValueVector4(const FName NodeIdentity, const FName AttributeName, FVector4& OutValue)
+{
+	return GetNodeAttributeValue(NodeIdentity, AttributeName, OutValue);
+}
+
+
+
+
+
+
+
+
+
+bool UInterchangeBaseNodeContainerAdapter::SetNodeAttributeValueBool(const FName NodeIdentity, const FName AttributeName, const bool OutValue)
+{
+	return SetNodeAttributeValue(NodeIdentity, AttributeName, OutValue);
+}
+
+bool UInterchangeBaseNodeContainerAdapter::SetNodeAttributeValueFloat(const FName NodeIdentity, const FName AttributeName, const float OutValue)
+{
+	return SetNodeAttributeValue(NodeIdentity, AttributeName, OutValue);
+}
+
+bool UInterchangeBaseNodeContainerAdapter::SetNodeAttributeValueDouble(const FName NodeIdentity, const FName AttributeName, const double OutValue)
+{
+	return SetNodeAttributeValue(NodeIdentity, AttributeName, OutValue);
+}
+
+bool UInterchangeBaseNodeContainerAdapter::SetNodeAttributeValueInt8(const FName NodeIdentity, const FName AttributeName, const int8 OutValue)
+{
+	return SetNodeAttributeValue(NodeIdentity, AttributeName, OutValue);
+}
+
+bool UInterchangeBaseNodeContainerAdapter::SetNodeAttributeValueInt16(const FName NodeIdentity, const FName AttributeName, const int16 OutValue)
+{
+	return SetNodeAttributeValue(NodeIdentity, AttributeName, OutValue);
+}
+
+bool UInterchangeBaseNodeContainerAdapter::SetNodeAttributeValueInt32(const FName NodeIdentity, const FName AttributeName, const int32 OutValue)
+{
+	return SetNodeAttributeValue(NodeIdentity, AttributeName, OutValue);
+}
+
+bool UInterchangeBaseNodeContainerAdapter::SetNodeAttributeValueInt64(const FName NodeIdentity, const FName AttributeName, const int64 OutValue)
+{
+	return SetNodeAttributeValue(NodeIdentity, AttributeName, OutValue);
+}
+
+bool UInterchangeBaseNodeContainerAdapter::SetNodeAttributeValueUint8(const FName NodeIdentity, const FName AttributeName, const uint8 OutValue)
+{
+	return SetNodeAttributeValue(NodeIdentity, AttributeName, OutValue);
+}
+
+bool UInterchangeBaseNodeContainerAdapter::SetNodeAttributeValueUint16(const FName NodeIdentity, const FName AttributeName, const uint16 OutValue)
+{
+	return SetNodeAttributeValue(NodeIdentity, AttributeName, OutValue);
+}
+
+bool UInterchangeBaseNodeContainerAdapter::SetNodeAttributeValueUint32(const FName NodeIdentity, const FName AttributeName, const uint32 OutValue)
+{
+	return SetNodeAttributeValue(NodeIdentity, AttributeName, OutValue);
+}
+
+bool UInterchangeBaseNodeContainerAdapter::SetNodeAttributeValueUint64(const FName NodeIdentity, const FName AttributeName, const uint64 OutValue)
+{
+	return SetNodeAttributeValue(NodeIdentity, AttributeName, OutValue);
+}
+
+bool UInterchangeBaseNodeContainerAdapter::SetNodeAttributeValueVector(const FName NodeIdentity, const FName AttributeName, const FVector& OutValue)
+{
+	return SetNodeAttributeValue(NodeIdentity, AttributeName, OutValue);
+}
+
+bool UInterchangeBaseNodeContainerAdapter::SetNodeAttributeValueVector2D(const FName NodeIdentity, const FName AttributeName, const FVector2D& OutValue)
+{
+	return SetNodeAttributeValue(NodeIdentity, AttributeName, OutValue);
+}
+
+bool UInterchangeBaseNodeContainerAdapter::SetNodeAttributeValueMatrix(const FName NodeIdentity, const FName AttributeName, const FMatrix& OutValue)
+{
+	return SetNodeAttributeValue(NodeIdentity, AttributeName, OutValue);
+}
+
+bool UInterchangeBaseNodeContainerAdapter::SetNodeAttributeValueBox(const FName NodeIdentity, const FName AttributeName, const FBox& OutValue)
+{
+	return SetNodeAttributeValue(NodeIdentity, AttributeName, OutValue);
+}
+
+bool UInterchangeBaseNodeContainerAdapter::SetNodeAttributeValueBoxSphereBound(const FName NodeIdentity, const FName AttributeName, const FBoxSphereBounds& OutValue)
+{
+	return SetNodeAttributeValue(NodeIdentity, AttributeName, OutValue);
+}
+
+bool UInterchangeBaseNodeContainerAdapter::SetNodeAttributeValueArray(const FName NodeIdentity, const FName AttributeName, const TArray<uint8>& OutValue)
+{
+	return SetNodeAttributeValue(NodeIdentity, AttributeName, OutValue);
+}
+
+bool UInterchangeBaseNodeContainerAdapter::SetNodeAttributeValueArray64(const FName NodeIdentity, const FName AttributeName, const TArray64<uint8>& OutValue)
+{
+	return SetNodeAttributeValue(NodeIdentity, AttributeName, OutValue);
+}
+
+bool UInterchangeBaseNodeContainerAdapter::SetNodeAttributeValueColor(const FName NodeIdentity, const FName AttributeName, const FColor& OutValue)
+{
+	return SetNodeAttributeValue(NodeIdentity, AttributeName, OutValue);
+}
+
+bool UInterchangeBaseNodeContainerAdapter::SetNodeAttributeValueDateTime(const FName NodeIdentity, const FName AttributeName, const FDateTime& OutValue)
+{
+	return SetNodeAttributeValue(NodeIdentity, AttributeName, OutValue);
+}
+
+bool UInterchangeBaseNodeContainerAdapter::SetNodeAttributeValueEnum(const FName NodeIdentity, const FName AttributeName, const uint8 OutValue)
+{
+	return SetNodeAttributeValue(NodeIdentity, AttributeName, OutValue);
+}
+
+bool UInterchangeBaseNodeContainerAdapter::SetNodeAttributeValueGuid(const FName NodeIdentity, const FName AttributeName, const FGuid& OutValue)
+{
+	return SetNodeAttributeValue(NodeIdentity, AttributeName, OutValue);
+}
+
+bool UInterchangeBaseNodeContainerAdapter::SetNodeAttributeValueIntPoint(const FName NodeIdentity, const FName AttributeName, const FIntPoint& OutValue)
+{
+	return SetNodeAttributeValue(NodeIdentity, AttributeName, OutValue);
+}
+
+bool UInterchangeBaseNodeContainerAdapter::SetNodeAttributeValueIntVecor(const FName NodeIdentity, const FName AttributeName, const FIntVector& OutValue)
+{
+	return SetNodeAttributeValue(NodeIdentity, AttributeName, OutValue);
+}
+
+bool UInterchangeBaseNodeContainerAdapter::SetNodeAttributeValueLinearColor(const FName NodeIdentity, const FName AttributeName, const FLinearColor& OutValue)
+{
+	return SetNodeAttributeValue(NodeIdentity, AttributeName, OutValue);
+}
+
+bool UInterchangeBaseNodeContainerAdapter::SetNodeAttributeValuePlane(const FName NodeIdentity, const FName AttributeName, const FPlane& OutValue)
+{
+	return SetNodeAttributeValue(NodeIdentity, AttributeName, OutValue);
+}
+
+bool UInterchangeBaseNodeContainerAdapter::SetNodeAttributeValueQuat(const FName NodeIdentity, const FName AttributeName, const FQuat& OutValue)
+{
+	return SetNodeAttributeValue(NodeIdentity, AttributeName, OutValue);
+}
+
+bool UInterchangeBaseNodeContainerAdapter::SetNodeAttributeValueName(const FName NodeIdentity, const FName AttributeName, const FName OutValue)
+{
+	return SetNodeAttributeValue(NodeIdentity, AttributeName, OutValue);
+}
+
+bool UInterchangeBaseNodeContainerAdapter::SetNodeAttributeValueRandomStream(const FName NodeIdentity, const FName AttributeName, const FRandomStream& OutValue)
+{
+	return SetNodeAttributeValue(NodeIdentity, AttributeName, OutValue);
+}
+
+bool UInterchangeBaseNodeContainerAdapter::SetNodeAttributeValueRotator(const FName NodeIdentity, const FName AttributeName, const FRotator& OutValue)
+{
+	return SetNodeAttributeValue(NodeIdentity, AttributeName, OutValue);
+}
+
+bool UInterchangeBaseNodeContainerAdapter::SetNodeAttributeValueString(const FName NodeIdentity, const FName AttributeName, const FString& OutValue)
+{
+	return SetNodeAttributeValue(NodeIdentity, AttributeName, OutValue);
+}
+
+bool UInterchangeBaseNodeContainerAdapter::SetNodeAttributeValueTimespan(const FName NodeIdentity, const FName AttributeName, const FTimespan& OutValue)
+{
+	return SetNodeAttributeValue(NodeIdentity, AttributeName, OutValue);
+}
+
+bool UInterchangeBaseNodeContainerAdapter::SetNodeAttributeValueTransform(const FName NodeIdentity, const FName AttributeName, const FTransform& OutValue)
+{
+	return SetNodeAttributeValue(NodeIdentity, AttributeName, OutValue);
+}
+
+bool UInterchangeBaseNodeContainerAdapter::SetNodeAttributeValueTwoVectors(const FName NodeIdentity, const FName AttributeName, const FTwoVectors& OutValue)
+{
+	return SetNodeAttributeValue(NodeIdentity, AttributeName, OutValue);
+}
+
+bool UInterchangeBaseNodeContainerAdapter::SetNodeAttributeValueVector4(const FName NodeIdentity, const FName AttributeName, const FVector4& OutValue)
+{
+	return SetNodeAttributeValue(NodeIdentity, AttributeName, OutValue);
+}
\ No newline at end of file
diff --git a/Engine/Source/Runtime/Interchange/Core/Private/Tasks/InterchangeTaskCompletion.cpp b/Engine/Source/Runtime/Interchange/Core/Private/Tasks/InterchangeTaskCompletion.cpp
new file mode 100644
--- /dev/null
+++ b/Engine/Source/Runtime/Interchange/Core/Private/Tasks/InterchangeTaskCompletion.cpp
@@ -0,0 +1,60 @@
+// Copyright Epic Games, Inc. All Rights Reserved.
+#include "InterchangeTaskCompletion.h"
+
+#include "AssetRegistryModule.h"
+#include "Async/TaskGraphInterfaces.h"
+#include "CoreMinimal.h"
+#include "InterchangeFactoryBase.h"
+#include "InterchangeManager.h"
+#include "Stats/Stats.h"
+#include "Templates/SharedPointer.h"
+#include "UObject/Object.h"
+#include "UObject/ObjectMacros.h"
+#include "UObject/WeakObjectPtrTemplates.h"
+
+void Interchange::FTaskCompletion::DoTask(ENamedThreads::Type CurrentThread, const FGraphEventRef& MyCompletionGraphEvent)
+{
+	TSharedPtr<Interchange::FImportAsyncHelper> AsyncHelper = WeakAsyncHelper.Pin();
+	check(AsyncHelper.IsValid());
+
+	for(TPair<int32, TArray<Interchange::FImportAsyncHelper::FImportedAssetInfo>>& AssetInfosPerSourceIndexPair : AsyncHelper->ImportedAssetsPerSourceIndex)
+	{
+		const int32 SourceIndex = AssetInfosPerSourceIndexPair.Key;
+		const bool bCallPostImportGameThreadCallback = ensure(AsyncHelper->SourceDatas.IsValidIndex(SourceIndex));
+		for (const Interchange::FImportAsyncHelper::FImportedAssetInfo& AssetInfo : AssetInfosPerSourceIndexPair.Value)
+		{
+			UObject* Asset = AssetInfo.ImportAsset;
+			//In case Some factory code cannot run outside of the main thread we offer this callback to finish the work before calling post edit change (building the asset)
+			if(bCallPostImportGameThreadCallback && AssetInfo.Factory)
+			{
+				UInterchangeFactoryBase::FPostImportGameThreadCallbackParams Arguments;
+				Arguments.ReimportObject = Asset;
+				Arguments.SourceData = AsyncHelper->SourceDatas[SourceIndex];
+				AssetInfo.Factory->PostImportGameThreadCallback(Arguments);
+			}
+
+			//Clear any async flag from the created asset
+			Asset->ClearInternalFlags(EInternalObjectFlags::Async);
+			//Make sure the package is dirty
+			Asset->MarkPackageDirty();
+#if WITH_EDITOR
+			//Make sure the asset is built correctly
+			Asset->PostEditChange();
+#endif
+			//Post import broadcast
+			if (AsyncHelper->TaskData.ReimportObject)
+			{
+				InterchangeManager->OnAssetPostReimport.Broadcast(Asset);
+			}
+			else
+			{
+				InterchangeManager->OnAssetPostImport.Broadcast(Asset);
+			}
+			//Notify the asset registry
+			FAssetRegistryModule::AssetCreated(Asset);
+		}
+	}
+	//Release the async helper
+	AsyncHelper = nullptr;
+	InterchangeManager->ReleaseAsyncHelper(WeakAsyncHelper);
+}
\ No newline at end of file
diff --git a/Engine/Source/Runtime/Interchange/Core/Private/Tasks/InterchangeTaskCompletion.h b/Engine/Source/Runtime/Interchange/Core/Private/Tasks/InterchangeTaskCompletion.h
new file mode 100644
--- /dev/null
+++ b/Engine/Source/Runtime/Interchange/Core/Private/Tasks/InterchangeTaskCompletion.h
@@ -0,0 +1,44 @@
+// Copyright Epic Games, Inc. All Rights Reserved.
+
+#pragma once
+
+#include "Async/TaskGraphInterfaces.h"
+#include "CoreMinimal.h"
+#include "InterchangeManager.h"
+#include "Stats/Stats.h"
+#include "UObject/WeakObjectPtrTemplates.h"
+
+namespace Interchange
+{
+
+class FTaskCompletion
+{
+private:
+	UInterchangeManager* InterchangeManager;
+	TWeakPtr<Interchange::FImportAsyncHelper> WeakAsyncHelper;
+public:
+	FTaskCompletion(UInterchangeManager* InInterchangeManager, TWeakPtr<Interchange::FImportAsyncHelper> InAsyncHelper)
+		: InterchangeManager(InInterchangeManager)
+		, WeakAsyncHelper(InAsyncHelper)
+	{
+	}
+
+	static FORCEINLINE ENamedThreads::Type GetDesiredThread()
+	{
+		return ENamedThreads::GameThread;
+	}
+	static FORCEINLINE ESubsequentsMode::Type GetSubsequentsMode()
+	{
+		return ESubsequentsMode::FireAndForget;
+	}
+
+	FORCEINLINE TStatId GetStatId() const
+	{
+		RETURN_QUICK_DECLARE_CYCLE_STAT(FTaskCompletion, STATGROUP_TaskGraphTasks);
+	}
+
+	void DoTask(ENamedThreads::Type CurrentThread, const FGraphEventRef& MyCompletionGraphEvent);
+};
+
+} //end namespace Interchange
+
diff --git a/Engine/Source/Runtime/Interchange/Core/Private/Tasks/InterchangeTaskCreateAsset.cpp b/Engine/Source/Runtime/Interchange/Core/Private/Tasks/InterchangeTaskCreateAsset.cpp
new file mode 100644
--- /dev/null
+++ b/Engine/Source/Runtime/Interchange/Core/Private/Tasks/InterchangeTaskCreateAsset.cpp
@@ -0,0 +1,150 @@
+// Copyright Epic Games, Inc. All Rights Reserved.
+#include "InterchangeTaskCreateAsset.h"
+
+#include "AssetRegistryModule.h"
+#include "Async/TaskGraphInterfaces.h"
+#include "CoreMinimal.h"
+#include "InterchangeFactoryBase.h"
+#include "InterchangeManager.h"
+#include "InterchangeSourceData.h"
+#include "InterchangeTranslatorBase.h"
+#include "Misc/Paths.h"
+#include "PackageUtils/PackageUtils.h"
+#include "Stats/Stats.h"
+#include "Templates/SharedPointer.h"
+#include "UObject/Object.h"
+#include "UObject/UObjectGlobals.h"
+#include "UObject/ObjectMacros.h"
+#include "UObject/WeakObjectPtrTemplates.h"
+#include "Nodes/BaseNode.h"
+
+namespace Interchange_InternalImplementation
+{
+	void InternalGetPackageName(TSharedPtr<Interchange::FImportAsyncHelper> AsyncHelper, const int32 SourceIndex, const FString& PackageBasePath, const Interchange::FBaseNode* Node, FString& OutPackageName, FString& OutAssetName)
+	{
+		const UInterchangeSourceData* SourceData = AsyncHelper->SourceDatas[SourceIndex];
+		check(SourceData);
+		FString NodeDisplayName = Node->GetDisplayLabel().ToString();
+		const FString BaseFileName = FPaths::GetBaseFilename(SourceData->GetFilename());
+
+		//Set the asset name and the package name
+		if (NodeDisplayName.Equals(BaseFileName) || BaseFileName.IsEmpty())
+		{
+			OutAssetName = NodeDisplayName;
+		}
+		else
+		{
+			OutAssetName = BaseFileName.IsEmpty() ? NodeDisplayName : BaseFileName + TEXT("_") + NodeDisplayName;
+		}
+		OutPackageName = FPaths::Combine(*PackageBasePath, *OutAssetName);
+		
+		//Sanitize only the package name
+		Interchange::SanitizeInvalidChar(OutPackageName);
+	}
+}
+
+void Interchange::FTaskCreatePackage::DoTask(ENamedThreads::Type CurrentThread, const FGraphEventRef& MyCompletionGraphEvent)
+{
+	TSharedPtr<Interchange::FImportAsyncHelper> AsyncHelper = WeakAsyncHelper.Pin();
+	check(AsyncHelper.IsValid());
+
+	//The create package thread must always execute on the game thread
+	check(IsInGameThread());
+
+	UPackage* Pkg = nullptr;
+	FString PackageName;
+	FString AssetName;
+	//If we do a reimport no need to create a package
+	if (AsyncHelper->TaskData.ReimportObject)
+	{
+		Pkg = AsyncHelper->TaskData.ReimportObject->GetPackage();
+		PackageName = Pkg->GetPathName();
+	}
+	else
+	{
+		Interchange_InternalImplementation::InternalGetPackageName(AsyncHelper, SourceIndex, PackageBasePath, Node, PackageName, AssetName);
+		// We can not create assets that share the name of a map file in the same location
+		if (Interchange::FPackageUtils::IsMapPackageAsset(PackageName))
+		{
+			const FText Message = FText::Format(NSLOCTEXT("Interchange", "AssetNameInUseByMap", "You can not create an asset named '{0}' because there is already a map file with this name in this folder."), FText::FromString(AssetName));
+			UE_LOG(LogInterchangeCore, Warning, TEXT("%s"), *Message.ToString());
+			//Skip this asset
+			return;
+		}
+
+		Pkg = CreatePackage(nullptr, *PackageName);
+		if (Pkg == nullptr)
+		{
+			const FText Message = FText::Format(NSLOCTEXT("Interchange", "CannotCreatePackageErrorMsg", "Cannot create package named '{0}', will not import asset {1}."), FText::FromString(PackageName), FText::FromString(AssetName));
+			UE_LOG(LogInterchangeCore, Warning, TEXT("%s"), *Message.ToString());
+			//Skip this asset
+			return;
+		}
+
+		//Import Asset describe by the node
+		UInterchangeFactoryBase::FCreateAssetParams CreateAssetParams;
+		CreateAssetParams.AssetName = AssetName;
+		CreateAssetParams.AssetNode = Node;
+		CreateAssetParams.Parent = Pkg;
+		CreateAssetParams.SourceData = AsyncHelper->SourceDatas[SourceIndex];
+		CreateAssetParams.Translator = nullptr;
+		CreateAssetParams.ReimportObject = AsyncHelper->TaskData.ReimportObject;
+		//Make sure the asset UObject is created with the correct type on the main thread
+		Factory->CreateEmptyAsset(CreateAssetParams);
+	}
+	// Make sure the destination package is loaded
+	Pkg->FullyLoad();
+
+	AsyncHelper->CreatedPackages.Add(PackageName, Pkg);
+}
+
+void Interchange::FTaskCreateAsset::DoTask(ENamedThreads::Type CurrentThread, const FGraphEventRef& MyCompletionGraphEvent)
+{
+	check(WeakAsyncHelper.IsValid());
+	TSharedPtr<Interchange::FImportAsyncHelper> AsyncHelper = WeakAsyncHelper.Pin();
+
+	FString PackageName;
+	FString AssetName;
+	Interchange_InternalImplementation::InternalGetPackageName(AsyncHelper, SourceIndex, PackageBasePath, Node, PackageName, AssetName);
+	if (AsyncHelper->TaskData.ReimportObject)
+	{
+		UPackage* Pkg = AsyncHelper->TaskData.ReimportObject->GetPackage();
+		PackageName = Pkg->GetPathName();
+	}
+
+	UPackage** PkgPtr = AsyncHelper->CreatedPackages.Find(PackageName);
+	if (!PkgPtr || !(*PkgPtr))
+	{
+		const FText Message = FText::Format(NSLOCTEXT("Interchange", "CannotCreateAssetNoPackageErrorMsg", "Cannot create asset named '{1}', package '{0}'was not created properly."), FText::FromString(PackageName), FText::FromString(AssetName));
+		UE_LOG(LogInterchangeCore, Warning, TEXT("%s"), *Message.ToString());
+		return;
+	}
+
+	if (!AsyncHelper->SourceDatas.IsValidIndex(SourceIndex) || !AsyncHelper->Translators.IsValidIndex(SourceIndex))
+	{
+		const FText Message = FText::Format(NSLOCTEXT("Interchange", "CannotCreateAssetMissingDataErrorMsg", "Cannot create asset named '{0}', Source data or translator is invalid."), FText::FromString(AssetName));
+		UE_LOG(LogInterchangeCore, Warning, TEXT("%s"), *Message.ToString());
+		return;
+	}
+
+	UPackage* Pkg = *PkgPtr;
+
+	UInterchangeTranslatorBase* Translator = AsyncHelper->Translators[SourceIndex];
+	//Import Asset describe by the node
+	UInterchangeFactoryBase::FCreateAssetParams CreateAssetParams;
+	CreateAssetParams.AssetName = AssetName;
+	CreateAssetParams.AssetNode = Node;
+	CreateAssetParams.Parent = Pkg;
+	CreateAssetParams.SourceData = AsyncHelper->SourceDatas[SourceIndex];
+	CreateAssetParams.Translator = Translator;
+	CreateAssetParams.ReimportObject = AsyncHelper->TaskData.ReimportObject;
+
+	UObject* NodeAsset = Factory->CreateAsset(CreateAssetParams);
+	if (NodeAsset)
+	{
+		TArray<Interchange::FImportAsyncHelper::FImportedAssetInfo>& ImportedInfos = AsyncHelper->ImportedAssetsPerSourceIndex.FindOrAdd(SourceIndex);
+		Interchange::FImportAsyncHelper::FImportedAssetInfo& AssetInfo = ImportedInfos.AddZeroed_GetRef();
+		AssetInfo.ImportAsset = NodeAsset;
+		AssetInfo.Factory = Factory;
+	}
+}
\ No newline at end of file
diff --git a/Engine/Source/Runtime/Interchange/Core/Private/Tasks/InterchangeTaskCreateAsset.h b/Engine/Source/Runtime/Interchange/Core/Private/Tasks/InterchangeTaskCreateAsset.h
new file mode 100644
--- /dev/null
+++ b/Engine/Source/Runtime/Interchange/Core/Private/Tasks/InterchangeTaskCreateAsset.h
@@ -0,0 +1,100 @@
+// Copyright Epic Games, Inc. All Rights Reserved.
+
+#pragma once
+
+#include "Async/TaskGraphInterfaces.h"
+#include "CoreMinimal.h"
+#include "InterchangeFactoryBase.h"
+#include "InterchangeManager.h"
+#include "Stats/Stats.h"
+#include "UObject/WeakObjectPtrTemplates.h"
+#include "Nodes/BaseNode.h"
+
+namespace Interchange
+{
+
+/**
+ * This task create package, Cook::PackageTracker::NotifyUObjectCreated is not thread safe, so we need to create the packages on the main thread
+ */
+class FTaskCreatePackage
+{
+private:
+	FString PackageBasePath;
+	int32 SourceIndex;
+	TWeakPtr<Interchange::FImportAsyncHelper> WeakAsyncHelper;
+	const FBaseNode* Node;
+	UInterchangeFactoryBase* Factory;
+
+public:
+	FTaskCreatePackage(const FString& InPackageBasePath, const int32 InSourceIndex, TWeakPtr<Interchange::FImportAsyncHelper> InAsyncHelper, const Interchange::FBaseNode* InNode, UInterchangeFactoryBase* InFactory)
+		: PackageBasePath(InPackageBasePath)
+		, SourceIndex(InSourceIndex)
+		, WeakAsyncHelper(InAsyncHelper)
+		, Node(InNode)
+		, Factory(InFactory)
+	{
+		check(Node);
+		check(Factory);
+	}
+
+	FORCEINLINE ENamedThreads::Type GetDesiredThread()
+	{
+		if (WeakAsyncHelper.IsValid() && WeakAsyncHelper.Pin()->TaskData.ReimportObject)
+		{
+			//When doing a reimport the package already exist, so we can get it outside of the main thread
+			return ENamedThreads::AnyBackgroundThreadNormalTask;
+		}
+		return ENamedThreads::GameThread;
+	}
+	static FORCEINLINE ESubsequentsMode::Type GetSubsequentsMode()
+	{
+		return ESubsequentsMode::TrackSubsequents;
+	}
+
+	FORCEINLINE TStatId GetStatId() const
+	{
+		RETURN_QUICK_DECLARE_CYCLE_STAT(FTaskCreatePackage, STATGROUP_TaskGraphTasks);
+	}
+
+	void DoTask(ENamedThreads::Type CurrentThread, const FGraphEventRef& MyCompletionGraphEvent);
+};
+
+class FTaskCreateAsset
+{
+private:
+	FString PackageBasePath;
+	int32 SourceIndex;
+	TWeakPtr<Interchange::FImportAsyncHelper> WeakAsyncHelper;
+	const FBaseNode* Node;
+	UInterchangeFactoryBase* Factory;
+public:
+	FTaskCreateAsset(const FString& InPackageBasePath, const int32 InSourceIndex, TWeakPtr<Interchange::FImportAsyncHelper> InAsyncHelper, const Interchange::FBaseNode* InNode, UInterchangeFactoryBase* InFactory)
+		: PackageBasePath(InPackageBasePath)
+		, SourceIndex(InSourceIndex)
+		, WeakAsyncHelper(InAsyncHelper)
+		, Node(InNode)
+		, Factory(InFactory)
+	{
+		check(Node);
+		check(Factory);
+	}
+
+	FORCEINLINE ENamedThreads::Type GetDesiredThread()
+	{
+		return Factory->CanExecuteOnAnyThread() ? ENamedThreads::AnyBackgroundThreadNormalTask : ENamedThreads::GameThread;
+	}
+	static FORCEINLINE ESubsequentsMode::Type GetSubsequentsMode()
+	{
+		return ESubsequentsMode::TrackSubsequents;
+	}
+
+	FORCEINLINE TStatId GetStatId() const
+	{
+		RETURN_QUICK_DECLARE_CYCLE_STAT(FTaskCreateAsset, STATGROUP_TaskGraphTasks);
+	}
+
+	void DoTask(ENamedThreads::Type CurrentThread, const FGraphEventRef& MyCompletionGraphEvent);
+};
+
+
+} // End namespace Interchange
diff --git a/Engine/Source/Runtime/Interchange/Core/Private/Tasks/InterchangeTaskParsing.cpp b/Engine/Source/Runtime/Interchange/Core/Private/Tasks/InterchangeTaskParsing.cpp
new file mode 100644
--- /dev/null
+++ b/Engine/Source/Runtime/Interchange/Core/Private/Tasks/InterchangeTaskParsing.cpp
@@ -0,0 +1,107 @@
+// Copyright Epic Games, Inc. All Rights Reserved.
+#include "InterchangeTaskParsing.h"
+
+#include "AssetRegistryModule.h"
+#include "Async/TaskGraphInterfaces.h"
+#include "CoreMinimal.h"
+#include "InterchangeFactoryBase.h"
+#include "InterchangeManager.h"
+#include "InterchangeSourceData.h"
+#include "InterchangeTaskCompletion.h"
+#include "InterchangeTaskCreateAsset.h"
+#include "InterchangeTranslatorBase.h"
+#include "Misc/Paths.h"
+#include "PackageUtils/PackageUtils.h"
+#include "Stats/Stats.h"
+#include "Templates/SharedPointer.h"
+#include "UObject/Object.h"
+#include "UObject/UObjectGlobals.h"
+#include "UObject/ObjectMacros.h"
+#include "UObject/WeakObjectPtrTemplates.h"
+#include "Nodes/BaseNode.h"
+#include "Nodes/BaseNodeContainer.h"
+
+void Interchange::FTaskParsing::DoTask(ENamedThreads::Type CurrentThread, const FGraphEventRef& MyCompletionGraphEvent)
+{
+	TSharedPtr<Interchange::FImportAsyncHelper> AsyncHelper = WeakAsyncHelper.Pin();
+	check(AsyncHelper.IsValid());
+
+	//Parse each graph and prepare import task data, we will then be able to create all the task with the correct dependencies
+	struct FTaskData
+	{
+		Interchange::FNodeUniqueID UniqueID;
+		int32 SourceIndex;
+		const Interchange::FBaseNode* Node;
+		TArray<Interchange::FNodeUniqueID> Dependencies;
+		FGraphEventRef GraphEventRef;
+		FGraphEventArray Prerequistes;
+		UInterchangeFactoryBase* Factory;
+	};
+
+
+
+	FGraphEventArray CompletionPrerequistes;
+	TArray<FTaskData> TaskDatas;
+	for (int32 SourceIndex = 0; SourceIndex < AsyncHelper->SourceDatas.Num(); ++SourceIndex)
+	{
+		AsyncHelper->BaseNodeContainers[SourceIndex].IterateNodes([&](const Interchange::FNodeUniqueID NodeUID, const Interchange::FBaseNode* Node)
+		{
+			if (Node->GetAssetClass() != nullptr)
+			{
+				UInterchangeFactoryBase* NodeFactory = InterchangeManager->GetRegisterFactory(Node->GetAssetClass());
+				if (!NodeFactory)
+				{
+					//nothing we can import from this element
+					return;
+				}
+				FTaskData& NodeTaskData = TaskDatas.AddDefaulted_GetRef();
+				NodeTaskData.UniqueID = Node->GetUniqueID();
+				NodeTaskData.SourceIndex = SourceIndex;
+				NodeTaskData.Node = Node;
+				Node->GetDependecies(NodeTaskData.Dependencies);
+				NodeTaskData.Factory = NodeFactory;
+			}
+		});
+	}
+
+	//Sort per dependencies
+	TaskDatas.Sort([](const FTaskData& A, const FTaskData& B)
+	{
+		//if A is a dependency of B then return true to do A before B
+		if (B.Dependencies.Contains(A.UniqueID))
+		{
+			return true;
+		}
+		if (A.Dependencies.Contains(B.UniqueID))
+		{
+			return false;
+		}
+		return A.Dependencies.Num() <= B.Dependencies.Num();
+	});
+
+	for (int32 TaskIndex = 0; TaskIndex < TaskDatas.Num(); ++TaskIndex)
+	{
+		if (TaskDatas[TaskIndex].Dependencies.Num() > 0)
+		{
+			//Search the previous node to find the dependence
+			for (int32 DepTaskIndex = 0; DepTaskIndex < TaskIndex; ++DepTaskIndex)
+			{
+				if (TaskDatas[TaskIndex].Dependencies.Contains(TaskDatas[DepTaskIndex].UniqueID))
+				{
+					//Add has prerequisite
+					TaskDatas[TaskIndex].Prerequistes.Add(TaskDatas[DepTaskIndex].GraphEventRef);
+				}
+			}
+		}
+
+		//Add create package task has a prerequisite of FTaskCreateAsset. Create package task is a game thread task
+		FGraphEventArray CreatePackagePrerequistes;
+		CreatePackagePrerequistes.Add(TGraphTask<FTaskCreatePackage>::CreateTask(&(TaskDatas[TaskIndex].Prerequistes)).ConstructAndDispatchWhenReady(PackageBasePath, TaskDatas[TaskIndex].SourceIndex, WeakAsyncHelper, TaskDatas[TaskIndex].Node, TaskDatas[TaskIndex].Factory));
+
+		int32 CreateTaskIndex = AsyncHelper->CreateAssetTasks.Add(TGraphTask<FTaskCreateAsset>::CreateTask(&(CreatePackagePrerequistes)).ConstructAndDispatchWhenReady(PackageBasePath, TaskDatas[TaskIndex].SourceIndex, WeakAsyncHelper, TaskDatas[TaskIndex].Node, TaskDatas[TaskIndex].Factory));
+		TaskDatas[TaskIndex].GraphEventRef = AsyncHelper->CreateAssetTasks[CreateTaskIndex];
+		CompletionPrerequistes.Add(TaskDatas[TaskIndex].GraphEventRef);
+	}
+
+	AsyncHelper->CompletionTask = TGraphTask<FTaskCompletion>::CreateTask(&CompletionPrerequistes).ConstructAndDispatchWhenReady(InterchangeManager, WeakAsyncHelper);
+}
diff --git a/Engine/Source/Runtime/Interchange/Core/Private/Tasks/InterchangeTaskParsing.h b/Engine/Source/Runtime/Interchange/Core/Private/Tasks/InterchangeTaskParsing.h
new file mode 100644
--- /dev/null
+++ b/Engine/Source/Runtime/Interchange/Core/Private/Tasks/InterchangeTaskParsing.h
@@ -0,0 +1,47 @@
+// Copyright Epic Games, Inc. All Rights Reserved.
+
+#pragma once
+
+#include "Async/TaskGraphInterfaces.h"
+#include "CoreMinimal.h"
+#include "InterchangeManager.h"
+#include "Stats/Stats.h"
+#include "UObject/WeakObjectPtrTemplates.h"
+
+namespace Interchange
+{
+
+class FTaskParsing
+{
+private:
+	UInterchangeManager* InterchangeManager;
+	TWeakPtr<Interchange::FImportAsyncHelper> WeakAsyncHelper;
+	FString PackageBasePath;
+public:
+	FTaskParsing(UInterchangeManager* InInterchangeManager, const FString& InPackageBasePath, TWeakPtr<Interchange::FImportAsyncHelper> InAsyncHelper)
+		: InterchangeManager(InInterchangeManager)
+		, WeakAsyncHelper(InAsyncHelper)
+		, PackageBasePath(InPackageBasePath)
+	{
+		check(InterchangeManager);
+	}
+
+	static FORCEINLINE ENamedThreads::Type GetDesiredThread()
+	{
+		return ENamedThreads::AnyBackgroundThreadNormalTask;
+	}
+	static FORCEINLINE ESubsequentsMode::Type GetSubsequentsMode()
+	{
+		
+		return ESubsequentsMode::TrackSubsequents;
+	}
+
+	FORCEINLINE TStatId GetStatId() const
+	{
+		RETURN_QUICK_DECLARE_CYCLE_STAT(FTaskParsing, STATGROUP_TaskGraphTasks);
+	}
+
+	void DoTask(ENamedThreads::Type CurrentThread, const FGraphEventRef& MyCompletionGraphEvent);
+};
+
+} // End namespace Interchange
diff --git a/Engine/Source/Runtime/Interchange/Core/Private/Tasks/InterchangeTaskPipeline.cpp b/Engine/Source/Runtime/Interchange/Core/Private/Tasks/InterchangeTaskPipeline.cpp
new file mode 100644
--- /dev/null
+++ b/Engine/Source/Runtime/Interchange/Core/Private/Tasks/InterchangeTaskPipeline.cpp
@@ -0,0 +1,26 @@
+// Copyright Epic Games, Inc. All Rights Reserved.
+#include "InterchangeTaskPipeline.h"
+
+#include "Async/TaskGraphInterfaces.h"
+#include "CoreMinimal.h"
+#include "InterchangeManager.h"
+#include "InterchangePipelineBase.h"
+#include "Stats/Stats.h"
+#include "Templates/SharedPointer.h"
+#include "UObject/WeakObjectPtrTemplates.h"
+#include "Nodes/BaseNodeContainer.h"
+
+void Interchange::FTaskPipeline::DoTask(ENamedThreads::Type CurrentThread, const FGraphEventRef& MyCompletionGraphEvent)
+{
+	TSharedPtr<Interchange::FImportAsyncHelper> AsyncHelper = WeakAsyncHelper.Pin();
+	check(AsyncHelper.IsValid());
+
+	if (UInterchangePipelineBase* Pipeline = PipelineBase.Get())
+	{
+		for (int32 GraphIndex = 0; GraphIndex < AsyncHelper->BaseNodeContainers.Num(); ++GraphIndex)
+		{
+			check(AsyncHelper->BaseNodeContainerAdapters[GraphIndex].IsValid());
+			Pipeline->ScriptedExecuteImportPipeline(AsyncHelper->BaseNodeContainerAdapters[GraphIndex].Get());
+		}
+	}
+}
diff --git a/Engine/Source/Runtime/Interchange/Core/Private/Tasks/InterchangeTaskPipeline.h b/Engine/Source/Runtime/Interchange/Core/Private/Tasks/InterchangeTaskPipeline.h
new file mode 100644
--- /dev/null
+++ b/Engine/Source/Runtime/Interchange/Core/Private/Tasks/InterchangeTaskPipeline.h
@@ -0,0 +1,46 @@
+// Copyright Epic Games, Inc. All Rights Reserved.
+
+#pragma once
+
+#include "Async/TaskGraphInterfaces.h"
+#include "CoreMinimal.h"
+#include "InterchangeManager.h"
+#include "InterchangePipelineBase.h"
+#include "Stats/Stats.h"
+#include "UObject/WeakObjectPtrTemplates.h"
+
+namespace Interchange
+{
+
+class FTaskPipeline
+{
+private:
+	TWeakObjectPtr<UInterchangePipelineBase> PipelineBase;
+	TWeakPtr<Interchange::FImportAsyncHelper> WeakAsyncHelper;
+public:
+	FTaskPipeline(TWeakObjectPtr<UInterchangePipelineBase> InPipelineBase, TWeakPtr<Interchange::FImportAsyncHelper> InAsyncHelper)
+		: PipelineBase(InPipelineBase)
+		, WeakAsyncHelper(InAsyncHelper)
+	{
+	}
+
+	ENamedThreads::Type GetDesiredThread()
+	{
+		return PipelineBase.Get()->ScriptedCanExecuteOnAnyThread() ? ENamedThreads::AnyBackgroundThreadNormalTask : ENamedThreads::GameThread;
+	}
+	
+	static ESubsequentsMode::Type GetSubsequentsMode()
+	{
+		return ESubsequentsMode::TrackSubsequents;
+	}
+
+	FORCEINLINE TStatId GetStatId() const
+	{
+		RETURN_QUICK_DECLARE_CYCLE_STAT(FTaskPipeline, STATGROUP_TaskGraphTasks);
+	}
+
+	void DoTask(ENamedThreads::Type CurrentThread, const FGraphEventRef& MyCompletionGraphEvent);
+};
+
+
+} // End namespace Interchange
diff --git a/Engine/Source/Runtime/Interchange/Core/Private/Tasks/InterchangeTaskTranslator.cpp b/Engine/Source/Runtime/Interchange/Core/Private/Tasks/InterchangeTaskTranslator.cpp
new file mode 100644
--- /dev/null
+++ b/Engine/Source/Runtime/Interchange/Core/Private/Tasks/InterchangeTaskTranslator.cpp
@@ -0,0 +1,38 @@
+// Copyright Epic Games, Inc. All Rights Reserved.
+#include "InterchangeTaskTranslator.h"
+
+#include "Async/TaskGraphInterfaces.h"
+#include "CoreMinimal.h"
+#include "InterchangeManager.h"
+#include "InterchangeSourceData.h"
+#include "InterchangeTranslatorBase.h"
+#include "Stats/Stats.h"
+#include "Templates/SharedPointer.h"
+#include "UObject/WeakObjectPtrTemplates.h"
+#include "Nodes/BaseNodeContainer.h"
+
+void Interchange::FTaskTranslator::DoTask(ENamedThreads::Type CurrentThread, const FGraphEventRef& MyCompletionGraphEvent)
+{
+	TSharedPtr<Interchange::FImportAsyncHelper> AsyncHelper = WeakAsyncHelper.Pin();
+	check(AsyncHelper.IsValid());
+
+	UInterchangeTranslatorBase* Translator = AsyncHelper->Translators.IsValidIndex(SourceIndex) ? AsyncHelper->Translators[SourceIndex] : nullptr;
+	if (!Translator)
+	{
+		return;
+	}
+	const UInterchangeSourceData* SourceData = AsyncHelper->SourceDatas.IsValidIndex(SourceIndex) ? AsyncHelper->SourceDatas[SourceIndex] : nullptr;
+	if (!SourceData)
+	{
+		return;
+	}
+
+	if (!AsyncHelper->BaseNodeContainers.IsValidIndex(SourceIndex))
+	{
+		return;
+	}
+
+	//Translate the source data
+	Interchange::FBaseNodeContainer& BaseNodeContainer = AsyncHelper->BaseNodeContainers[SourceIndex];
+	Translator->Translate(SourceData, BaseNodeContainer);
+}
diff --git a/Engine/Source/Runtime/Interchange/Core/Private/Tasks/InterchangeTaskTranslator.h b/Engine/Source/Runtime/Interchange/Core/Private/Tasks/InterchangeTaskTranslator.h
new file mode 100644
--- /dev/null
+++ b/Engine/Source/Runtime/Interchange/Core/Private/Tasks/InterchangeTaskTranslator.h
@@ -0,0 +1,45 @@
+// Copyright Epic Games, Inc. All Rights Reserved.
+
+#pragma once
+
+#include "Async/TaskGraphInterfaces.h"
+#include "CoreMinimal.h"
+#include "InterchangeManager.h"
+#include "Stats/Stats.h"
+#include "UObject/WeakObjectPtrTemplates.h"
+
+namespace Interchange
+{
+
+class FTaskTranslator
+{
+private:
+	int32 SourceIndex = INDEX_NONE;
+	TWeakPtr<Interchange::FImportAsyncHelper> WeakAsyncHelper;
+
+public:
+	FTaskTranslator(int32 InSourceIndex, TWeakPtr<Interchange::FImportAsyncHelper> InAsyncHelper)
+		: SourceIndex(InSourceIndex)
+		, WeakAsyncHelper(InAsyncHelper)
+	{
+	}
+
+	static FORCEINLINE ENamedThreads::Type GetDesiredThread()
+	{
+		return ENamedThreads::AnyBackgroundThreadNormalTask;
+	}
+	static FORCEINLINE ESubsequentsMode::Type GetSubsequentsMode()
+	{
+		return ESubsequentsMode::TrackSubsequents;
+	}
+
+	FORCEINLINE TStatId GetStatId() const
+	{
+		RETURN_QUICK_DECLARE_CYCLE_STAT(FTaskTranslator, STATGROUP_TaskGraphTasks);
+	}
+
+	void DoTask(ENamedThreads::Type CurrentThread, const FGraphEventRef& MyCompletionGraphEvent);
+};
+
+
+} // End namespace Interchange
diff --git a/Engine/Source/Runtime/Interchange/Core/Private/Tests/StorageTest.cpp b/Engine/Source/Runtime/Interchange/Core/Private/Tests/StorageTest.cpp
new file mode 100644
--- /dev/null
+++ b/Engine/Source/Runtime/Interchange/Core/Private/Tests/StorageTest.cpp
@@ -0,0 +1,498 @@
+// Copyright Epic Games, Inc. All Rights Reserved.
+
+#include "CoreMinimal.h"
+#include "Async/ParallelFor.h"
+#include "Misc/AutomationTest.h"
+#include "Serialization/MemoryReader.h"
+#include "Templates/Atomic.h"
+#include "Types/AttributeStorage.h"
+
+#if WITH_DEV_AUTOMATION_TESTS
+
+IMPLEMENT_SIMPLE_AUTOMATION_TEST(FAttributeStorageTest, "System.Runtime.Interchange.AttributeStorage", EAutomationTestFlags::ApplicationContextMask | EAutomationTestFlags::SmokeFilter)
+
+bool FAttributeStorageTest::RunTest(const FString& Parameters)
+{
+	using namespace Interchange;
+
+	FString KeyPrefixString = TEXT("TestKey");
+	FName KeyPrefix(*KeyPrefixString);
+	TAtomic<int32> KeySuffixCounter(1);
+	auto CreateUniqueKey = [&KeyPrefix, &KeySuffixCounter]()->FAttributeKey
+	{
+		KeyPrefix.SetNumber(KeySuffixCounter.IncrementExchange());
+		const FAttributeKey UniqueKeyRet = FAttributeKey(KeyPrefix);
+		return UniqueKeyRet;
+	};
+
+	FAttributeStorage TestStorage;
+	//This seed make the tests deterministic
+	FRandomStream RandomStream(564389);
+	const FAttributeKey TestInt32KeyName = CreateUniqueKey();
+	const FAttributeKey BigArrayKey = CreateUniqueKey();
+	const int32 NegativeValueRef = -1;
+	const int32 OverrideValueRef = 3327;
+	const FAttributeKey RandomStreamKey = CreateUniqueKey();
+	int32 TestStoredSeed = 0;
+	const int32 RandHelperUInt8 = 255;
+
+	//Test Adding/Reading uint8 attribute with 0 value
+	{
+		const uint8 RefValue = 0;
+		FAttributeKey Key = CreateUniqueKey();
+		if (!IsAttributeStorageResultSuccess(TestStorage.RegisterAttribute(Key, RefValue)))
+		{
+			AddError(FString(TEXT("`AttributeStorage` must handle adding uint8 attribute")));
+		}
+		uint8 StoredValue = 0;
+		TestStorage.GetAttributeHandle<uint8>(Key).Get(StoredValue);
+		TestEqual(TEXT("`AttributeStorage` must handle add and retrieve uint8 attribute"), static_cast<int32>(StoredValue), static_cast<int32>(RefValue));
+	}
+
+	//Test Adding/Reading uint8 attribute with a value different then the default value
+	{
+		const uint8 RefValue = 43;
+		FAttributeKey Key = CreateUniqueKey();
+		if (!IsAttributeStorageResultSuccess(TestStorage.RegisterAttribute(Key, RefValue)))
+		{
+			AddError(FString(TEXT("`AttributeStorage` must handle adding uint8 attribute")));
+		}
+		uint8 StoredValue = 0;
+		TestStorage.GetAttributeHandle<uint8>(Key).Get(StoredValue);
+		TestEqual(TEXT("`AttributeStorage` must handle add and retrieve uint8 attribute"), static_cast<int32>(StoredValue), static_cast<int32>(RefValue));
+	}
+
+	//Test Adding/Reading simple attribute with a value different then the default value
+	{
+		if (!IsAttributeStorageResultSuccess(TestStorage.RegisterAttribute(TestInt32KeyName, NegativeValueRef)))
+		{
+			AddError(FString(TEXT("`AttributeStorage` must handle adding int32 attribute")));
+		}
+		int32 StoredValue = 0;
+		TestStorage.GetAttributeHandle<int32>(TestInt32KeyName).Get(StoredValue);
+		TestEqual(TEXT("`AttributeStorage` must handle add and retrieve int32 attribute"), StoredValue, NegativeValueRef);
+	}
+
+	//Adding many FVector
+	for (int32 AddedIndex = 0; AddedIndex < 2; ++AddedIndex)
+	{
+		const FVector ValueRef = RandomStream.VRand();
+		FAttributeKey Key = CreateUniqueKey();
+		if (!IsAttributeStorageResultSuccess(TestStorage.RegisterAttribute(Key, ValueRef)))
+		{
+			AddError(FString(TEXT("`AttributeStorage` must handle adding FVector attribute")));
+			break;
+		}
+		FVector StoredValue(0);
+		TestStorage.GetAttributeHandle<FVector>(Key).Get(StoredValue);
+		if (!TestEqual(TEXT("`AttributeStorage` must handle add and retrieve FVector attribute"), StoredValue.X, ValueRef.X))
+		{
+			break;
+		}
+	}
+
+	//Adding FName
+	{
+		const FName RefValue = TEXT("The magic carpet ride!");
+		FAttributeKey Key = CreateUniqueKey();
+		if (!IsAttributeStorageResultSuccess(TestStorage.RegisterAttribute(Key, RefValue)))
+		{
+			AddError(FString(TEXT("`AttributeStorage` must handle adding FAttributeKey attribute")));
+		}
+		FName StoredValue;
+		TestStorage.GetAttributeHandle<FName>(Key).Get(StoredValue);
+		if(StoredValue != RefValue)
+		{
+			TestEqual(TEXT("`AttributeStorage` must handle add and retrieve FAttributeKey attribute"), StoredValue.ToString(), RefValue.ToString());
+		}
+	}
+
+	//Adding one big TArray64<uint8> outside of the Hash
+	{
+		TArray64<uint8> ValueRef;
+		uint64 ArrayNum = 50;
+		ValueRef.Reserve(ArrayNum);
+		for (uint64 ArrayIndex = 0; ArrayIndex < ArrayNum; ++ArrayIndex)
+		{
+			ValueRef.Add(static_cast<uint8>(RandomStream.RandHelper(RandHelperUInt8)));
+		}
+
+		if (!IsAttributeStorageResultSuccess(TestStorage.RegisterAttribute(BigArrayKey, ValueRef, Interchange::EAttributeProperty::NoHash)))
+		{
+			AddError(FString(TEXT("`AttributeStorage` must handle adding TArray<uint8> attribute")));
+		}
+		TArray64<uint8> StoredValue;
+		TestStorage.GetAttributeHandle<TArray64<uint8> >(BigArrayKey).Get(StoredValue);
+		int32 TestArrayIndex = RandomStream.RandHelper(ArrayNum - 1);
+		TestEqual(TEXT("`AttributeStorage` must handle add and retrieve TArray64<uint8> attribute"), StoredValue[TestArrayIndex], ValueRef[TestArrayIndex]);
+	}
+
+	//Adding many TArray<uint8>
+	for (int32 AddedIndex = 0; AddedIndex < 5; ++AddedIndex)
+	{
+		TArray<uint8> ValueRef;
+		int32 ArrayNum = 2 * (AddedIndex + 1);
+		for (int32 ArrayIndex = 0; ArrayIndex < ArrayNum; ++ArrayIndex)
+		{
+			ValueRef.Add(static_cast<uint8>(RandomStream.RandHelper(RandHelperUInt8)));
+		}
+		FAttributeKey Key = CreateUniqueKey();
+		if (!IsAttributeStorageResultSuccess(TestStorage.RegisterAttribute(Key, ValueRef)))
+		{
+			AddError(FString(TEXT("`AttributeStorage` must handle adding TArray<uint8> attribute")));
+			break;
+		}
+		TArray<uint8> StoredValue;
+		TestStorage.GetAttributeHandle<TArray<uint8> >(Key).Get(StoredValue);
+		int32 TestArrayIndex = RandomStream.RandHelper(ArrayNum-1);
+		if (!TestEqual(TEXT("`AttributeStorage` must handle add and retrieve TArray<uint8> attribute"), StoredValue[TestArrayIndex], ValueRef[TestArrayIndex]))
+		{
+			break;
+		}
+	}
+
+	//Store a FString
+	{
+		const FString ValueRef = TEXT("The quick brown fox jumped over the lazy dogs");
+		const FAttributeKey Key = CreateUniqueKey();
+		if (!IsAttributeStorageResultSuccess(TestStorage.RegisterAttribute(Key, ValueRef)))
+		{
+			AddError(FString(TEXT("`AttributeStorage` must handle adding FString attribute")));
+		}
+		FString StoredValue;
+		if (!IsAttributeStorageResultSuccess(TestStorage.GetAttributeHandle<FString >(Key).Get(StoredValue)))
+		{
+			AddError(FString(TEXT("`AttributeStorage` must handle retrieve FString attribute")));
+		}
+		//Equal random stream should give the same result when asking fraction
+		TestEqual(TEXT("`AttributeStorage` must handle add and retrieve FString attribute"), StoredValue, ValueRef);
+	}
+
+	//Store the random stream
+	{
+		if (!IsAttributeStorageResultSuccess(TestStorage.RegisterAttribute(RandomStreamKey, RandomStream)))
+		{
+			AddError(FString(TEXT("`AttributeStorage` must handle adding FRandomStream attribute")));
+		}
+		FRandomStream StoredRandomStream;
+		TestStorage.GetAttributeHandle<FRandomStream>(RandomStreamKey).Get(StoredRandomStream);
+		//Equal random stream should give the same result when asking fraction
+		TestStoredSeed = RandomStream.RandHelper(RandHelperUInt8);
+		TestEqual(TEXT("`AttributeStorage` must handle add and retrieve FRandomStream attribute"), StoredRandomStream.RandHelper(RandHelperUInt8), TestStoredSeed);
+	}
+
+	//Get the Hash we will use it later to do some test
+	FGuid HashGuidRef = TestStorage.GetStorageHash();
+
+	//Test hashing twice give the same result
+	{
+		FGuid HashGuidTest = TestStorage.GetStorageHash();
+		if (HashGuidRef != HashGuidTest)
+		{
+			AddError(FString(TEXT("`AttributeStorage` hash must be deterministic when calculating it twice.")));
+		}
+	}
+
+	//Removing...
+	{
+		if (!IsAttributeStorageResultSuccess(TestStorage.UnregisterAttribute(TestInt32KeyName)))
+		{
+			AddError(FString(TEXT("`AttributeStorage` must handle removing attribute")));
+		}
+
+		//This will kick a defrag of the storage
+		if (!IsAttributeStorageResultSuccess(TestStorage.UnregisterAttribute(BigArrayKey)))
+		{
+			AddError(FString(TEXT("`AttributeStorage` must handle removing attribute with defrag")));
+		}
+	}
+
+	//Re-adding the negative integer value
+	{
+		if (!IsAttributeStorageResultSuccess(TestStorage.RegisterAttribute(TestInt32KeyName, NegativeValueRef)))
+		{
+			AddError(FString(TEXT("`AttributeStorage` must handle adding int32 attribute")));
+		}
+		int32 StoredValue = 0;
+		TestStorage.GetAttributeHandle<int32>(TestInt32KeyName).Get(StoredValue);
+		TestEqual(TEXT("`AttributeStorage` must handle add and retrieve int32 attribute"), StoredValue, NegativeValueRef);
+	}
+
+	//Test hash must be deterministic even if the order of attributes have change
+	{
+		FGuid HashGuidTest = TestStorage.GetStorageHash();
+		if (HashGuidRef != HashGuidTest)
+		{
+			AddError(FString(TEXT("`AttributeStorage` hash must be deterministic even if the attributes order differ.")));
+		}
+	}
+
+	//Get the random stream after the defrag and re-adding
+	{
+		FRandomStream StoredRandomStream;
+		TestStorage.GetAttributeHandle<FRandomStream>(RandomStreamKey).Get(StoredRandomStream);
+		//Equal random stream should give the same result when asking fraction
+		TestEqual(TEXT("`AttributeStorage` must handle add and retrieve FRandomStream attribute"), StoredRandomStream.RandHelper(RandHelperUInt8), TestStoredSeed);
+	}
+
+	//Overriding value
+	{
+		int32 StoredValue = 0;
+		FAttributeStorage::TAttributeHandle<int32> TestInt32Handle = TestStorage.GetAttributeHandle<int32>(TestInt32KeyName);
+		TestInt32Handle.Get(StoredValue);
+		TestEqual(TEXT("`AttributeStorage` must handle add and retrieve uint8 attribute"), StoredValue, NegativeValueRef);
+
+		if (!IsAttributeStorageResultSuccess(TestInt32Handle.Set(OverrideValueRef)))
+		{
+			AddError(FString(TEXT("`AttributeStorage` must handle overriding int32 attribute")));
+		}
+		StoredValue = 0;
+		TestInt32Handle.Get(StoredValue);
+		TestEqual(TEXT("`AttributeStorage` must handle overriding uint32 attribute"), StoredValue, OverrideValueRef);
+	}
+
+	//Multi thread test
+	{
+		//Add multithread (add BatchCount*BatchSize attributes)
+		const int32 BatchSize = 10;
+		const int32 BatchCount = 50;
+		TArray<FAttributeKey> Keys;
+		Keys.AddZeroed(BatchSize*BatchCount);
+		bool ThreadError = false;
+		FAttributeStorage::TAttributeHandle<FRandomStream> RandomStreamHandle = TestStorage.GetAttributeHandle<FRandomStream>(RandomStreamKey);
+		//Iterate all vertex to compute normals for all vertex instance
+		ParallelFor(BatchCount,
+			[&TestStorage, &RandomStreamHandle, &BatchSize, &Keys, &CreateUniqueKey, &ThreadError](const int32 BatchIndex)
+		{
+			//Get the random stream from the storage so all thread add the number of attribute
+			FRandomStream RandomStream;
+			RandomStreamHandle.Get(RandomStream);
+			for (int32 AttributeIndex = 0; AttributeIndex < BatchSize; ++AttributeIndex)
+			{
+				const uint64 RealIndex = (BatchIndex * BatchSize) + AttributeIndex;
+				const FVector ValueRef = RandomStream.VRand();
+				Keys[RealIndex] = CreateUniqueKey();
+				if (!IsAttributeStorageResultSuccess(TestStorage.RegisterAttribute(Keys[RealIndex], ValueRef)))
+				{
+					ThreadError = true;
+					continue;
+				}
+				FVector StoredValue(0);
+				TestStorage.GetAttributeHandle<FVector>(Keys[RealIndex]).Get(StoredValue);
+				if (StoredValue.X != ValueRef.X)
+				{
+					ThreadError = true;
+				}
+			}
+		});
+		
+		if (ThreadError)
+		{
+			AddError(FString(TEXT("`AttributeStorage` Fail adding attributes in multi thread.")));
+		}
+
+		//Push the defrag ratio to 0.5 to avoid calling defrag to often		
+		TestStorage.SetDefragRatio(0.5);
+		
+		//Now remove in multithread the just added attributes, It should trigger a defrag at some point
+		ParallelFor(BatchCount,
+			[&TestStorage, &BatchSize, &Keys, &ThreadError](int32 BatchIndex)
+		{
+			for (int32 AttributeIndex = 0; AttributeIndex < BatchSize; ++AttributeIndex)
+			{
+				const uint64 RealIndex = BatchIndex * BatchSize + AttributeIndex;
+				if (!IsAttributeStorageResultSuccess(TestStorage.UnregisterAttribute(Keys[RealIndex])))
+				{
+					ThreadError = true;
+				}
+			}
+		});
+
+		if (ThreadError )
+		{
+			AddError(FString(TEXT("`AttributeStorage` Fail removing attributes in multi thread.")));
+		}
+	}
+
+	//Test compare storage
+	{
+		FAttributeStorage BaseStorage;
+		FAttributeStorage VersionStorage;
+		
+		//Adding many FVector
+		for (int32 AddedIndex = 0; AddedIndex < 3; ++AddedIndex)
+		{
+			const FVector ValueRef = RandomStream.VRand();
+			FAttributeKey Key = CreateUniqueKey();
+			if (!IsAttributeStorageResultSuccess(BaseStorage.RegisterAttribute(Key, ValueRef)))
+			{
+				AddError(FString(TEXT("`AttributeStorage` must handle adding FVector attribute")));
+				break;
+			}
+			if (!IsAttributeStorageResultSuccess(VersionStorage.RegisterAttribute(Key, ValueRef)))
+			{
+				AddError(FString(TEXT("`AttributeStorage` must handle adding FVector attribute")));
+				break;
+			}
+		}
+
+		//Adding many int32
+		for (int32 AddedIndex = 0; AddedIndex < 3; ++AddedIndex)
+		{
+			const int32 ValueRef = RandomStream.RandHelper(RandHelperUInt8);
+			FAttributeKey Key = CreateUniqueKey();
+			if (!IsAttributeStorageResultSuccess(BaseStorage.RegisterAttribute(Key, ValueRef)))
+			{
+				AddError(FString(TEXT("`AttributeStorage` must handle adding int32 attribute")));
+				break;
+			}
+			if (!IsAttributeStorageResultSuccess(VersionStorage.RegisterAttribute(Key, ValueRef)))
+			{
+				AddError(FString(TEXT("`AttributeStorage` must handle adding int32 attribute")));
+				break;
+			}
+		}
+		TArray<FAttributeKey> RemovedKeys;
+		TArray<FAttributeKey> AddedKeys;
+		TArray<FAttributeKey> ModifiedKeys;
+		FAttributeStorage::CompareStorage(BaseStorage, VersionStorage, RemovedKeys, AddedKeys, ModifiedKeys);
+		if (RemovedKeys.Num() > 0 || AddedKeys.Num() > 0 || ModifiedKeys.Num() > 0)
+		{
+			AddError(FString(TEXT("`AttributeStorage` Compare storage shoukld have found no difference.")));
+		}
+
+		TArray<FAttributeKey> VersionKeys;
+		VersionStorage.GetAttributeKeys(VersionKeys);
+
+		//Add a integer key to the version storage
+		{
+			const int32 ValueRef = RandomStream.RandHelper(RandHelperUInt8);
+			FAttributeKey Key(TEXT("VersionAddedInteger"));
+			if (!IsAttributeStorageResultSuccess(VersionStorage.RegisterAttribute(Key, ValueRef)))
+			{
+				AddError(FString(TEXT("`AttributeStorage` must handle adding int32 attribute")));
+			}
+		}
+
+		//Modify a key to the version storage
+		if (VersionKeys.IsValidIndex(0))
+		{
+			const FVector ValueRef = RandomStream.VRand();
+			if (!IsAttributeStorageResultSuccess(VersionStorage.GetAttributeHandle<FVector>(VersionKeys[0]).Set(ValueRef)))
+			{
+				AddError(FString(TEXT("`AttributeStorage` must handle overriding FVector attribute")));
+			}
+		}
+
+		//Remove a key to the version storage
+		if(VersionKeys.IsValidIndex(3))
+		{
+			if (!IsAttributeStorageResultSuccess(VersionStorage.UnregisterAttribute(VersionKeys[3])))
+			{
+				AddError(FString(TEXT("`AttributeStorage` must handle removing int32 attribute")));
+			}
+		}
+
+		FAttributeStorage::CompareStorage(BaseStorage, VersionStorage, RemovedKeys, AddedKeys, ModifiedKeys);
+		if (RemovedKeys.Num() != 1 || AddedKeys.Num() != 1 || ModifiedKeys.Num() != 1)
+		{
+			AddError(FString(TEXT("`AttributeStorage` Compare storage should have found one difference for each modify/add/remove attruibutes.")));
+		}
+	}
+
+	//Serialization test
+	{
+		TArray<uint8> MemoryMocked;
+		const bool bMemoryPersistent = true;
+		FMemoryWriter Ar(MemoryMocked, bMemoryPersistent);
+		Ar << TestStorage;
+
+		FAttributeStorage FromMemoryMocked;
+		FMemoryReader ArRead(MemoryMocked, bMemoryPersistent);
+		ArRead << FromMemoryMocked;
+
+		TArray<FAttributeKey> RemovedKeys;
+		TArray<FAttributeKey> AddedKeys;
+		TArray<FAttributeKey> ModifiedKeys;
+
+		if(TestStorage.GetStorageHash() != FromMemoryMocked.GetStorageHash())
+		{
+			AddError(FString(TEXT("`AttributeStorage` serialization Write then Readback give different storage hash value.")));
+		}
+		
+		FAttributeStorage::CompareStorage(TestStorage, FromMemoryMocked, RemovedKeys, AddedKeys, ModifiedKeys);
+		if (RemovedKeys.Num() != 0 || AddedKeys.Num() != 0 || ModifiedKeys.Num() != 0)
+		{
+			AddError(FString(TEXT("`AttributeStorage` Compare storage should have found no difference after a serialization write follow by a read.")));
+		}
+	}
+
+	//Performance Test, add two million entry, 1 million int32 and 1 million FVector
+	{
+		FGuid StorageHash = TestStorage.GetStorageHash();
+		TArray<FAttributeKey> AttributeKeys;
+		const int32 TestCount = 10000;
+		AttributeKeys.AddZeroed(2*TestCount);
+		//Reserve 2x Testcount items (Testcount int32 and Testcount FVector) and reserve the data size
+		TestStorage.Reserve(TestCount*2, TestCount*(sizeof(int32)+sizeof(FVector)));
+		//Start the timer after we reserve the memory
+		const double SmokeTestStartTime = FPlatformTime::Seconds();
+		for (int32 AddedIndex = 0; AddedIndex < TestCount; ++AddedIndex)
+		{
+			const int32 ValueRef = AddedIndex;
+			FAttributeKey Key = CreateUniqueKey();
+			AttributeKeys[AddedIndex*2] = Key;
+			if (!IsAttributeStorageResultSuccess(TestStorage.RegisterAttribute(Key, ValueRef)))
+			{
+				AddError(FString(TEXT("`AttributeStorage` must handle adding int32 attribute")));
+				break;
+			}
+
+			const FVector VectorRef(AddedIndex/TestCount);
+			Key = CreateUniqueKey();
+			AttributeKeys[(AddedIndex*2)+1] = Key;
+			if (!IsAttributeStorageResultSuccess(TestStorage.RegisterAttribute(Key, VectorRef)))
+			{
+				AddError(FString(TEXT("`AttributeStorage` must handle adding vector attribute")));
+				break;
+			}
+		}
+		const double SmokeTestWriteTime = FPlatformTime::Seconds();
+		int32 MaxKeys = AttributeKeys.Num();
+		for (int32 KeyIndex = 0; KeyIndex < MaxKeys; ++KeyIndex)
+		{
+			const FAttributeKey& Key = AttributeKeys[KeyIndex];
+			if(TestStorage.GetAttributeType(Key) == EAttributeTypes::Int32)
+			{
+				int32 StoredValue(0);
+				TestStorage.GetAttributeHandle<int32>(Key).Get(StoredValue);
+			}
+			else if(TestStorage.GetAttributeType(Key) == EAttributeTypes::Vector)
+			{
+				FVector StoredVector(0.0f);
+				TestStorage.GetAttributeHandle<FVector>(Key).Get(StoredVector);
+			}
+		}
+		const double SmokeTestReadTime = FPlatformTime::Seconds();
+		for (int32 KeyIndex = 0; KeyIndex < MaxKeys; ++KeyIndex)
+		{
+			const FAttributeKey& Key = AttributeKeys[KeyIndex];
+			TestStorage.UnregisterAttribute(Key);
+		}
+		const double SmokeTestStopTime = FPlatformTime::Seconds();
+		const float TimeForWrite = (SmokeTestWriteTime - SmokeTestStartTime);
+		const float TimeForRead = (SmokeTestReadTime - SmokeTestWriteTime);
+		const float TimeForUnregister = (SmokeTestStopTime - SmokeTestReadTime);
+		const float TimeForTest = (SmokeTestStopTime - SmokeTestStartTime);
+		AddInfo(FString::Format(TEXT("AttributeStorage performance test result ({4} int32 and FVector)\n\tRegister: {0}\n\tRead: {1}\n\tUnregister: {2}\n\tTotal: {3}"), {TimeForWrite, TimeForRead, TimeForUnregister, TimeForTest, TestCount}));
+		FGuid StorageHashAfter = TestStorage.GetStorageHash();
+		if (StorageHashAfter != StorageHash)
+		{
+			AddError(FString(TEXT("`AttributeStorage` hash must be deterministic even if the attributes order differ.")));
+		}
+	}
+	return true;
+}
+
+
+#endif //WITH_DEV_AUTOMATION_TESTS
diff --git a/Engine/Source/Runtime/Interchange/Core/Private/Types/AttributeStorage.cpp b/Engine/Source/Runtime/Interchange/Core/Private/Types/AttributeStorage.cpp
new file mode 100644
--- /dev/null
+++ b/Engine/Source/Runtime/Interchange/Core/Private/Types/AttributeStorage.cpp
@@ -0,0 +1,456 @@
+// Copyright Epic Games, Inc. All Rights Reserved.
+#include "Types/AttributeStorage.h"
+#include "../InterchangeLogPrivate.h"
+
+#include "CoreMinimal.h"
+
+//Interchange namespace
+namespace Interchange
+{
+	void LogAttributeStorageErrors(const EAttributeStorageResult Result, const FString OperationName, const FAttributeKey AttributeKey)
+	{
+		//////////////////////////////////////////////////////////////////////////
+		//Errors
+		if (HasAttributeStorageResult(Result, EAttributeStorageResult::Operation_Error_AttributeAllocationCorrupted))
+		{
+			UE_LOG(LogInterchangeCore, Error, TEXT("Attribute storage operation [%s] Key[%s]: Storage is corrupted."), *OperationName, *(AttributeKey.Key.ToString()));
+		}
+
+		//////////////////////////////////////////////////////////////////////////
+		//Warning
+		if (HasAttributeStorageResult(Result, EAttributeStorageResult::Operation_Error_CannotFoundKey))
+		{
+			UE_LOG(LogInterchangeCore, Warning, TEXT("Attribute storage operation [%s] Key[%s]: Cannot find attribute key."), *OperationName, *(AttributeKey.Key.ToString()));
+		}
+
+		if (HasAttributeStorageResult(Result, EAttributeStorageResult::Operation_Error_CannotOverrideAttribute))
+		{
+			UE_LOG(LogInterchangeCore, Warning, TEXT("Attribute storage operation [%s] Key[%s]: Cannot override attribute."), *OperationName, *(AttributeKey.Key.ToString()));
+		}
+
+		if (HasAttributeStorageResult(Result, EAttributeStorageResult::Operation_Error_CannotRemoveAttribute))
+		{
+			UE_LOG(LogInterchangeCore, Warning, TEXT("Attribute storage operation [%s] Key[%s]: Cannot remove an attribute."), *OperationName, *(AttributeKey.Key.ToString()));
+		}
+
+		if (HasAttributeStorageResult(Result, EAttributeStorageResult::Operation_Error_WrongSize))
+		{
+			UE_LOG(LogInterchangeCore, Warning, TEXT("Attribute storage operation [%s] Key[%s]: Stored attribute value size do not match parameter value size."), *OperationName, *(AttributeKey.Key.ToString()));
+		}
+
+		if (HasAttributeStorageResult(Result, EAttributeStorageResult::Operation_Error_WrongType))
+		{
+			UE_LOG(LogInterchangeCore, Warning, TEXT("Attribute storage operation [%s] Key[%s]: Stored attribute value type do not match parameter value type."), *OperationName, *(AttributeKey.Key.ToString()));
+		}
+
+		if (HasAttributeStorageResult(Result, EAttributeStorageResult::Operation_Error_UnsupportedType))
+		{
+			UE_LOG(LogInterchangeCore, Warning, TEXT("Attribute storage operation [%s] Key[%s]: The storage do not support this type of value, please add unsupported type to EAttributeTypes."), *OperationName, *(AttributeKey.Key.ToString()));
+		}
+
+		if (HasAttributeStorageResult(Result, EAttributeStorageResult::Operation_Error_InvalidStorage))
+		{
+			UE_LOG(LogInterchangeCore, Warning, TEXT("Attribute storage operation [%s] Key[%s]: The storage is invalid (NULL)."), *OperationName, *(AttributeKey.Key.ToString()));
+		}
+
+		if (HasAttributeStorageResult(Result, EAttributeStorageResult::Operation_Error_InvalidMultiSizeValueData))
+		{
+			UE_LOG(LogInterchangeCore, Warning, TEXT("Attribute storage operation [%s] Key[%s]: Cannot retrieve the multisize value data pointer."), *OperationName, *(AttributeKey.Key.ToString()));
+		}
+	}
+
+	FAttributeStorage& FAttributeStorage::operator=(const FAttributeStorage& Other)
+	{
+		//Lock both Storage mutex
+		FScopeLock ScopeLock(&StorageMutex);
+		FScopeLock ScopeLockOther(&Other.StorageMutex);
+
+		//Copy all value
+		AttributeAllocationTable = Other.AttributeAllocationTable;
+		AttributeStorage = Other.AttributeStorage;
+		FragmentedMemoryCost = Other.FragmentedMemoryCost;
+		DefragRatio = Other.DefragRatio;
+
+		return *this;
+	}
+
+	EAttributeStorageResult FAttributeStorage::UnregisterAttribute(const FAttributeKey& ElementAttributeKey)
+	{
+		//Lock the storage
+		FScopeLock ScopeLock(&StorageMutex);
+		
+		FAttributeAllocationInfo* AttributeAllocationInfo = AttributeAllocationTable.Find(ElementAttributeKey);
+		if (!AttributeAllocationInfo)
+		{
+			return EAttributeStorageResult::Operation_Error_CannotFoundKey;
+		}
+
+		FragmentedMemoryCost += AttributeAllocationInfo->Size;
+		if (AttributeAllocationTable.Remove(ElementAttributeKey) == 0)
+		{
+			return EAttributeStorageResult::Operation_Error_CannotRemoveAttribute;
+		}
+
+		DefragInternal();
+		
+		return EAttributeStorageResult::Operation_Success;
+	}
+	
+	EAttributeTypes FAttributeStorage::GetAttributeType(const FAttributeKey& ElementAttributeKey)
+	{
+		//Lock the storage
+		FScopeLock ScopeLock(&StorageMutex);
+		const FAttributeAllocationInfo* AttributeAllocationInfo = AttributeAllocationTable.Find(ElementAttributeKey);
+		if (AttributeAllocationInfo)
+		{
+			return AttributeAllocationInfo->Type;
+		}
+		return EAttributeTypes::None;
+	}	
+
+	bool FAttributeStorage::ContainAttribute(const FAttributeKey& ElementAttributeKey) const
+	{
+		//Lock the storage
+		FScopeLock ScopeLock(&StorageMutex);
+
+		return AttributeAllocationTable.Contains(ElementAttributeKey);
+	}
+
+	void FAttributeStorage::GetAttributeKeys(TArray<FAttributeKey>& AttributeKeys) const
+	{
+		FScopeLock ScopeLock(&StorageMutex);
+		AttributeAllocationTable.GetKeys(AttributeKeys);
+	}
+
+	FGuid FAttributeStorage::GetAttributeHash(const FAttributeKey& ElementAttributeKey) const
+	{
+		FGuid AttributeHash;
+		GetAttributeHash(ElementAttributeKey, AttributeHash);
+		return AttributeHash;
+	}
+
+	bool FAttributeStorage::GetAttributeHash(const FAttributeKey& ElementAttributeKey, FGuid& OutGuid) const
+	{
+		//Lock the storage
+		FScopeLock ScopeLock(&StorageMutex);
+		
+		const FAttributeAllocationInfo* AttributeAllocationInfo = AttributeAllocationTable.Find(ElementAttributeKey);
+		if (AttributeAllocationInfo)
+		{
+			OutGuid = AttributeAllocationInfo->Hash;
+			return true;
+		}
+		return false;
+	}
+
+	FGuid FAttributeStorage::GetStorageHash() const
+	{
+		//Lock the storage
+		FScopeLock ScopeLock(&StorageMutex);
+
+		TArray<FAttributeKey> OrderedKeys;
+		AttributeAllocationTable.GetKeys(OrderedKeys);
+		OrderedKeys.Sort();
+		FSHA1 Sha;
+		for (FAttributeKey Key : OrderedKeys)
+		{
+			const FAttributeAllocationInfo& AttributeAllocationInfo = AttributeAllocationTable.FindChecked(Key);
+			//Skip Attribute that has the no hash flag
+			if(HasAttributeProperty(AttributeAllocationInfo.Property, EAttributeProperty::NoHash))
+			{
+				continue;
+			}
+			uint32 GuidData[4];
+			GuidData[0] = AttributeAllocationInfo.Hash.A;
+			GuidData[1] = AttributeAllocationInfo.Hash.B;
+			GuidData[2] = AttributeAllocationInfo.Hash.C;
+			GuidData[3] = AttributeAllocationInfo.Hash.D;
+			Sha.Update(reinterpret_cast<uint8*>(&GuidData[0]), 16);
+		}
+		Sha.Final();
+		// Retrieve the hash and use it to construct a pseudo-GUID. 
+		uint32 Hash[5];
+		Sha.GetHash(reinterpret_cast<uint8*>(Hash));
+		return FGuid(Hash[0] ^ Hash[4], Hash[1], Hash[2], Hash[3]);
+	}
+
+	void FAttributeStorage::CompareStorage(const FAttributeStorage& BaseStorage, const FAttributeStorage& VersionStorage, TArray<FAttributeKey>& RemovedAttributes, TArray<FAttributeKey>& AddedAttributes, TArray<FAttributeKey>& ModifiedAttributes)
+	{
+		//Lock the storage
+		FScopeLock ScopeLockVersion(&VersionStorage.StorageMutex);
+		//Lock the storage
+		FScopeLock ScopeLockBase(&BaseStorage.StorageMutex);
+		
+		for (const auto KvpBase : BaseStorage.AttributeAllocationTable)
+		{
+			FAttributeKey KeyBase = KvpBase.Key;
+			const FAttributeAllocationInfo* AttributeAllocationInfoVersion = VersionStorage.AttributeAllocationTable.Find(KeyBase);
+			if (!AttributeAllocationInfoVersion)
+			{
+				//Add the attribute to RemovedAttributes
+				RemovedAttributes.Add(KeyBase);
+			}
+			else if(KvpBase.Value.Hash != AttributeAllocationInfoVersion->Hash)
+			{
+				//Add the attribute to ModifiedAttributes
+				ModifiedAttributes.Add(KeyBase);
+			}
+		}
+
+		for (const auto KvpVersion : VersionStorage.AttributeAllocationTable)
+		{
+			FAttributeKey KeyVersion = KvpVersion.Key;
+			const FAttributeAllocationInfo* AttributeAllocationInfoBase = BaseStorage.AttributeAllocationTable.Find(KeyVersion);
+			if (!AttributeAllocationInfoBase)
+			{
+				//Add the attribute to RemovedAttributes
+				AddedAttributes.Add(KeyVersion);
+			}
+		}
+	}
+
+	void FAttributeStorage::SetDefragRatio(const float InDefragRatio)
+	{
+		FScopeLock ScopeLock(&StorageMutex);
+		DefragRatio = InDefragRatio;
+		DefragInternal();
+	}
+
+	void FAttributeStorage::Reserve(int64 NewAttributeCount, int64 NewStorageSize)
+	{
+		if (NewAttributeCount > 0)
+		{
+			const int64 ReserveCount = AttributeAllocationTable.Num() + NewAttributeCount;
+			AttributeAllocationTable.Reserve(ReserveCount);
+		}
+		if (NewStorageSize > 0)
+		{
+			const int64 ReserveCount = AttributeStorage.Num() + NewStorageSize;
+			AttributeStorage.Reserve(ReserveCount);
+		}
+	}
+
+	void FAttributeStorage::ExtractFStringAttributeFromStorage(const uint8* StorageData, const FAttributeAllocationInfo* AttributeAllocationInfo, FString& OutValue) const
+	{
+		//Allocate the ByteArray
+		const uint64 NumberOfChar = AttributeAllocationInfo->Size / sizeof(TCHAR);
+		check(NumberOfChar > 0);
+
+
+		//Empty and pre allocate the FString buffer (TArray<TCHAR>)
+		OutValue.Reset(NumberOfChar - 1);
+		//SLOW! We have to add the characters 1 x 1, there is no AddZeroed or similar stuff. We do not have a string to use the operator= which do what we need here
+		//We go up to NumberOfChar -1 because there is always a null terminating character add at the end of the FString
+		for (int32 CharIndex = 0; CharIndex < NumberOfChar - 1; ++CharIndex)
+		{
+			//Adding spaces, we cannot add null character because the null character are not added
+			OutValue.AppendChar(' ');
+		}
+
+		//Copy into the buffer
+		FMemory::Memcpy((uint8*)(*OutValue), &StorageData[AttributeAllocationInfo->Offset], AttributeAllocationInfo->Size);
+	}
+
+	EAttributeStorageResult FAttributeStorage::GetAttribute(const FAttributeKey& ElementAttributeKey, FString& OutValue, TSpecializeType<FString >) const
+	{
+		//Lock the storage
+		FScopeLock ScopeLock(&StorageMutex);
+
+		if (!IsSupportedType(OutValue))
+		{
+			return EAttributeStorageResult::Operation_Error_UnsupportedType;
+		}
+
+		const EAttributeTypes ValueType = TAttributeTypeTraits<FString>::GetType();
+		const FAttributeAllocationInfo* AttributeAllocationInfo = AttributeAllocationTable.Find(ElementAttributeKey);
+		if (AttributeAllocationInfo)
+		{
+			if (AttributeAllocationInfo->Type != ValueType)
+			{
+				return EAttributeStorageResult::Operation_Error_WrongType;
+			}
+			if (!AttributeStorage.IsValidIndex(AttributeAllocationInfo->Offset))
+			{
+				return EAttributeStorageResult::Operation_Error_AttributeAllocationCorrupted;
+			}
+		}
+		else
+		{
+			//The key do not exist
+			return EAttributeStorageResult::Operation_Error_CannotFoundKey;
+		}
+
+		if (AttributeAllocationInfo->Size == 0)
+		{
+			OutValue = TEXT("");
+			return EAttributeStorageResult::Operation_Success;
+		}
+
+		ExtractFStringAttributeFromStorage(AttributeStorage.GetData(), AttributeAllocationInfo, OutValue);
+
+		return EAttributeStorageResult::Operation_Success;
+	}
+
+	EAttributeStorageResult FAttributeStorage::GetAttribute(const FAttributeKey& ElementAttributeKey, FName& OutValue, TSpecializeType<FName >) const
+	{
+		//Lock the storage
+		FScopeLock ScopeLock(&StorageMutex);
+
+		if (!IsSupportedType(OutValue))
+		{
+			return EAttributeStorageResult::Operation_Error_UnsupportedType;
+		}
+
+		const EAttributeTypes ValueType = TAttributeTypeTraits<FName>::GetType();
+
+		
+		const FAttributeAllocationInfo* AttributeAllocationInfo = AttributeAllocationTable.Find(ElementAttributeKey);
+		if (AttributeAllocationInfo)
+		{
+			if (AttributeAllocationInfo->Type != ValueType)
+			{
+				return EAttributeStorageResult::Operation_Error_WrongType;
+			}
+			if (!AttributeStorage.IsValidIndex(AttributeAllocationInfo->Offset))
+			{
+				return EAttributeStorageResult::Operation_Error_AttributeAllocationCorrupted;
+			}
+		}
+		else
+		{
+			//The key do not exist
+			return EAttributeStorageResult::Operation_Error_CannotFoundKey;
+		}
+
+		
+		if (AttributeAllocationInfo->Size == 0)
+		{
+			OutValue = NAME_None;
+			return EAttributeStorageResult::Operation_Success;
+		}
+
+		FString ValueStr;
+
+		//Share the code with FString from here
+		ExtractFStringAttributeFromStorage(AttributeStorage.GetData(), AttributeAllocationInfo, ValueStr);
+
+		//Create the FName and copy it to OutValue
+		OutValue = FName(*ValueStr);
+
+		return EAttributeStorageResult::Operation_Success;
+	}
+
+	void FAttributeStorage::DefragInternal()
+	{
+		//Defrag if the fragmented memory cost is bigger then the DefragRatio.
+		if (FragmentedMemoryCost < 11 || FragmentedMemoryCost < AttributeStorage.Num()* DefragRatio)
+		{
+			return;
+		}
+		
+		//Sorts the allocation table per offset since we want to defrag using memory block
+		//The ValueSort is also calling CompactStable which remove all holes in the TMap
+		struct FAttributeAllocationInfoSmaller
+		{
+			FORCEINLINE bool operator()(const FAttributeAllocationInfo& A, const FAttributeAllocationInfo& B) const
+			{
+				return A.Offset < B.Offset;
+			}
+		};
+		AttributeAllocationTable.ValueSort(FAttributeAllocationInfoSmaller());
+
+		const uint64 AttributeNumber = AttributeAllocationTable.Num();
+		uint8* StorageData = AttributeStorage.GetData();
+		//Number of attribute we have done
+		uint64 AttributeCount = 0;
+		//Current storage offset we want to defrag
+		uint64 CurrentOffset = 0;
+		//Are we build a defrag block
+		bool bBuildBlock = false;
+		//Defrag block start
+		uint64 MemmoveBlockStart = 0;
+		//Replace CurrentOffset when we are moving a block
+		uint64 MemmoveBlockCurrentOffset = 0;
+		for (auto& Kvp : AttributeAllocationTable)
+		{
+			AttributeCount++;
+			const bool bLast = AttributeCount == AttributeNumber;
+			
+			FAttributeAllocationInfo& AttributeAllocationInfo = Kvp.Value;
+			if (AttributeAllocationInfo.Offset == CurrentOffset)
+			{
+				CurrentOffset += AttributeAllocationInfo.Size;
+				continue;
+			}
+			check(CurrentOffset < AttributeAllocationInfo.Offset);
+			if(!bBuildBlock)
+			{
+				bBuildBlock = true;
+				MemmoveBlockStart = AttributeAllocationInfo.Offset;
+				AttributeAllocationInfo.Offset = CurrentOffset;
+				MemmoveBlockCurrentOffset = MemmoveBlockStart + AttributeAllocationInfo.Size;
+			}
+			else
+			{
+				if (AttributeAllocationInfo.Offset == MemmoveBlockCurrentOffset)
+				{
+					AttributeAllocationInfo.Offset -= MemmoveBlockStart - CurrentOffset;
+					MemmoveBlockCurrentOffset += AttributeAllocationInfo.Size;
+				}
+				else
+				{
+					const uint64 BlockSize = MemmoveBlockCurrentOffset - MemmoveBlockStart;
+					//Memmove support overlap
+					FMemory::Memmove(&StorageData[CurrentOffset], &StorageData[MemmoveBlockStart], BlockSize);
+					CurrentOffset += BlockSize;
+
+					MemmoveBlockStart = AttributeAllocationInfo.Offset;
+					AttributeAllocationInfo.Offset = CurrentOffset;
+					MemmoveBlockCurrentOffset = MemmoveBlockStart + AttributeAllocationInfo.Size;
+				}
+			}
+
+			//We have to move the block if we are the last attribute
+			if (bLast)
+			{
+				const uint64 BlockSize = MemmoveBlockCurrentOffset - MemmoveBlockStart;
+				//Memmove support overlap
+				FMemory::Memmove(&StorageData[CurrentOffset], &StorageData[MemmoveBlockStart], BlockSize);
+				const uint64 RemoveStartIndex = CurrentOffset + BlockSize;
+				const uint64 RemoveCount = AttributeStorage.Num() - RemoveStartIndex;
+				//Remove the Last items, allow shrinking so we are compact
+				AttributeStorage.RemoveAt(RemoveStartIndex, RemoveCount);
+			}
+		}
+		//Reset the fragmented cost
+		FragmentedMemoryCost = 0;
+	}
+
+	FGuid FAttributeStorage::GetValueHash(const uint8* Value, uint64 ValueSize)
+	{
+		FSHA1 Sha;
+		Sha.Update(Value, ValueSize);
+		Sha.Final();
+		// Retrieve the hash and use it to construct a pseudo-GUID. 
+		uint32 Hash[5];
+		Sha.GetHash(reinterpret_cast<uint8*>(Hash));
+		return FGuid(Hash[0] ^ Hash[4], Hash[1], Hash[2], Hash[3]);
+	}
+
+	uint64 FAttributeStorage::GetValueSize(const FString& Value, TSpecializeType<FString >)
+	{
+		//We must add the null character '/0' terminate string
+		return (Value.Len()+1) * sizeof(TCHAR);
+	}
+
+	uint64 FAttributeStorage::GetValueSize(const TArray<uint8>& Value, TSpecializeType<TArray<uint8> >)
+	{
+		return Value.Num();
+	}
+
+	uint64 FAttributeStorage::GetValueSize(const TArray64<uint8>& Value, TSpecializeType<TArray64<uint8> >)
+	{
+		return Value.Num();
+	}
+
+} //End Interchange namespace
\ No newline at end of file
diff --git a/Engine/Source/Runtime/Interchange/Core/Public/InterchangeFactoryBase.h b/Engine/Source/Runtime/Interchange/Core/Public/InterchangeFactoryBase.h
new file mode 100644
--- /dev/null
+++ b/Engine/Source/Runtime/Interchange/Core/Public/InterchangeFactoryBase.h
@@ -0,0 +1,101 @@
+// Copyright Epic Games, Inc. All Rights Reserved.
+
+#pragma once
+
+#include "CoreMinimal.h"
+#include "Nodes/BaseNode.h"
+#include "InterchangeTranslatorBase.h"
+#include "UObject/Class.h"
+#include "UObject/Object.h"
+#include "UObject/ObjectMacros.h"
+
+#include "InterchangeFactoryBase.generated.h"
+
+UCLASS(BlueprintType, Blueprintable, Abstract)
+class INTERCHANGECORE_API UInterchangeFactoryBase : public UObject
+{
+	GENERATED_BODY()
+public:
+
+	/**
+	 * return the UClass this factory can create.
+	 */
+	UFUNCTION(BlueprintCallable, Category = "Interchange | Factory")
+	virtual UClass* GetFactoryClass() const
+	{
+		return nullptr;
+	}
+
+	/**
+	 * Parameters to pass to CreateAsset function
+	 */
+	struct FCreateAssetParams
+	{
+		/** The package where to create the asset, if null it will put it in the transient package */
+		UObject* Parent = nullptr;
+
+		/** The name we want to give to the asset we will create */
+		FString AssetName = FString();
+
+		/** The base node that describe how to create the asset */
+		const Interchange::FBaseNode* AssetNode = nullptr;
+
+		/** The translator is use to retrieve the PayLoad data in case the factory need it */
+		const UInterchangeTranslatorBase* Translator = nullptr;
+
+		/** The source data, mainly use to set the asset import data file. TODO: we have to refactor UAssetImportData, the source data should be the base class for this now */
+		const UInterchangeSourceData* SourceData = nullptr;
+
+		UObject* ReimportObject = nullptr;
+	};
+
+	/**
+	 * Create an empty asset from a Node data, this function will be call in the main thread in same time has we create the package.
+	 *
+	 * @param Arguments - The structure containing all necessary arguments, see the structure definition for the documentation
+	 * @return the created UObject or nullptr if there is an error, see LOG for any error detail.
+	 *
+	 * @Note override function should verify the AssetNode in Arguments match the expected type for the factory.
+	 */
+	virtual UObject* CreateEmptyAsset(const FCreateAssetParams& Arguments) const
+	{
+		return nullptr;
+	}
+
+	/**
+	 * Create an asset from a Node data, this function must be multithread safe, it cannot use member all the data must be pass in the FCreateAssetParams structure
+	 *
+	 * @param Arguments - The structure containing all necessary arguments, see the structure definition for the documentation
+	 * @return the created UObject or nullptr if there is an error, see LOG for any error detail.
+	 *
+	 * @Note override function should verify the AssetNode in Arguments match the expected type for the factory.
+	 */
+	virtual UObject* CreateAsset(const FCreateAssetParams& Arguments) const
+	{
+		return nullptr;
+	}
+
+	/** Return true if the factory can create the asset asynchronously on any thread, false if it need to be on the main thread */
+	virtual bool CanExecuteOnAnyThread() const
+	{
+		return true;
+	}
+
+	/**
+	 * Parameters to pass to CreateAsset function
+	 */
+	struct FPostImportGameThreadCallbackParams
+	{
+		/** The source data, mainly use to set the asset import data file. TODO: we have to refactor UAssetImportData, the source data should be the base class for this now */
+		const UInterchangeSourceData* SourceData = nullptr;
+
+		/** The UObject  we want to execute code on*/
+		UObject* ReimportObject = nullptr;
+	};
+	/* This function is call in the completion task on the main thread, use it to call main thread post creation step for your assets*/
+	virtual void PostImportGameThreadCallback(const FPostImportGameThreadCallbackParams& Arguments) const
+	{
+		check(IsInGameThread());
+		return;
+	}
+};
diff --git a/Engine/Source/Runtime/Interchange/Core/Public/InterchangeManager.h b/Engine/Source/Runtime/Interchange/Core/Public/InterchangeManager.h
new file mode 100644
--- /dev/null
+++ b/Engine/Source/Runtime/Interchange/Core/Public/InterchangeManager.h
@@ -0,0 +1,330 @@
+// Copyright Epic Games, Inc. All Rights Reserved.
+
+#pragma once
+
+#include "CoreMinimal.h"
+#include "Delegates/DelegateCombinations.h"
+#include "InterchangeFactoryBase.h"
+#include "InterchangePipelineBase.h"
+#include "InterchangeSourceData.h"
+#include "InterchangeTranslatorBase.h"
+#include "InterchangeWriterBase.h"
+#include "UObject/Object.h"
+#include "UObject/ObjectMacros.h"
+#include "UObject/StrongObjectPtr.h"
+#include "Nodes/InterchangeBaseNodeContainerAdapter.h"
+
+#include "InterchangeManager.generated.h"
+
+class FAsyncTaskNotification;
+
+namespace Interchange
+{
+	class INTERCHANGECORE_API FScopedSourceData
+	{
+	public:
+		explicit FScopedSourceData(const FString& Filename);
+		UInterchangeSourceData* GetSourceData() const;
+	private:
+		TStrongObjectPtr<UInterchangeSourceData> SourceDataPtr = nullptr;
+	};
+
+	class INTERCHANGECORE_API FScopedTranslator
+	{
+	public:
+		explicit FScopedTranslator(const UInterchangeSourceData* SourceData);
+		UInterchangeTranslatorBase* GetTranslator();
+
+	private:
+		TStrongObjectPtr<UInterchangeTranslatorBase> ScopedTranslatorPtr = nullptr;
+	};
+
+	enum class EImportType : uint8
+	{
+		ImportType_None,
+		ImportType_Asset,
+		ImportType_Scene
+	};
+
+	struct FImportAsyncHelperData
+	{
+		//True if the import process is unattended. We cannot show UI  if the import is automated
+		bool bIsAutomated = false;
+
+		//We can import assets or full scene
+		EImportType ImportType = EImportType::ImportType_None;
+
+		//True if we are reimporting assets or scene
+		UObject* ReimportObject = nullptr;
+	};
+
+	class FImportAsyncHelper : public FGCObject
+	{
+	public:
+		~FImportAsyncHelper()
+		{
+			CleanUp();
+		}
+
+		/* FGCObject interface */
+		virtual void AddReferencedObjects(FReferenceCollector& Collector) override;
+
+		TArray<TStrongObjectPtr<UInterchangeBaseNodeContainerAdapter>> BaseNodeContainerAdapters;
+		TArray <FBaseNodeContainer> BaseNodeContainers;
+
+		TArray<UInterchangeSourceData* > SourceDatas;
+		TArray<UInterchangeTranslatorBase* > Translators;
+		TArray<UInterchangePipelineBase* > Pipelines;
+		TArray<UInterchangeFactoryBase* > Factories;
+
+		TArray<FGraphEventRef> TranslatorTasks;
+		TArray<FGraphEventRef> PipelineTasks;
+		FGraphEventRef ParsingTask;
+		TArray<FGraphEventRef> CreateAssetTasks;
+
+		FGraphEventRef CompletionTask;
+
+		//Create package map, Key is package name. We cannot create package asynchronously so we have to create a game thread task to do this
+		TMap<FString, UPackage*> CreatedPackages;
+
+		struct FImportedAssetInfo
+		{
+			UObject* ImportAsset;
+			UInterchangeFactoryBase* Factory;
+		};
+
+		TMap<int32, TArray<FImportedAssetInfo>> ImportedAssetsPerSourceIndex;
+
+		FImportAsyncHelperData TaskData;
+
+		void CleanUp();
+	};
+
+	void SanitizeInvalidChar(FString& String);
+
+} //End namespace interchange
+
+USTRUCT(BlueprintType)
+struct FImportAssetParameters
+{
+	GENERATED_USTRUCT_BODY()
+
+	UPROPERTY(EditAnywhere, BlueprintReadWrite, Category = "Interchange|ImportAsset")
+	UObject* ReimportAsset = nullptr;
+
+	UPROPERTY(EditAnywhere, BlueprintReadWrite, Category = "Interchange|ImportAsset")
+	bool bIsAutomated = false;
+};
+
+UCLASS(Transient, BlueprintType)
+class INTERCHANGECORE_API UInterchangeManager : public UObject
+{
+	GENERATED_BODY()
+public:
+
+	/**
+	 * Return the interchange manager singleton pointer.
+	 *
+	 * @note - We need to return a pointer to have a valid Blueprint callable function
+	 */
+	UFUNCTION(BlueprintCallable, Category = "Interchange | Import Manager")
+	static UInterchangeManager* GetInterchangeManagerScripted()
+	{
+		return &GetInterchangeManager();
+	}
+
+	/** Return the interchange manager singleton.*/
+	static UInterchangeManager& GetInterchangeManager()
+	{
+		static TStrongObjectPtr<UInterchangeManager> InterchangeManager = nullptr;
+		if (!InterchangeManager.IsValid())
+		{
+			//We cannot create a TStrongObjectPtr outside of the main thread, we also need a valid Transient package
+			check(IsInGameThread() && GetTransientPackage());
+			InterchangeManager = TStrongObjectPtr<UInterchangeManager>(NewObject<UInterchangeManager>(GetTransientPackage(), NAME_None, EObjectFlags::RF_NoFlags));
+		}
+		//When we get here we should be valid
+		check(InterchangeManager.IsValid());
+
+		return *(InterchangeManager.Get());
+	}
+
+
+	/** delegate type fired when new assets have been imported. Note: InCreatedObject can be NULL if import failed. Params: UFactory* InFactory, UObject* InCreatedObject */
+	DECLARE_MULTICAST_DELEGATE_OneParam(FInterchangeOnAssetPostImport, UObject*);
+	/** delegate type fired when new assets have been reimported. Note: InCreatedObject can be NULL if import failed. Params: UObject* InCreatedObject */
+	DECLARE_MULTICAST_DELEGATE_OneParam(FInterchangeOnAssetPostReimport, UObject*);
+	
+
+	// Delegates used to register and unregister
+
+	FInterchangeOnAssetPostImport OnAssetPostImport;
+	FInterchangeOnAssetPostReimport OnAssetPostReimport;
+
+	/**
+	 * Any translator must register to the manager
+	 * @Param Translator - The UClass of the translator you want to register
+	 * @return true if the translator class can be register false otherwise.
+	 *
+	 * @Note if you register multiple time the same class it will return true for every call
+	 */
+	bool RegisterTranslator(const UClass* TranslatorClass);
+
+	/**
+	 * Any factory must register to the manager
+	 * @Param Factory - The UClass of the factory you want to register
+	 * @return true if the factory class can be register false otherwise.
+	 *
+	 * @Note if you register multiple time the same class it will return true for every call
+	 */
+	bool RegisterFactory(const UClass* Factory);
+
+	/**
+	 * Any c++ pipeline must register to the manager
+	 * @Param pipeline - The UClass of the pipeline you want to register
+	 * @return true if the pipeline class can be register false otherwise.
+	 *
+	 * @Note if you register multiple time the same class it will return true for every call
+	 */
+	bool RegisterPipeline(const UClass* Pipeline);
+
+	/**
+	 * Any writer must register to the manager
+	 * @Param Writer - The UClass of the writer you want to register
+	 * @return true if the writer class can be register false otherwise.
+	 *
+	 * @Note if you register multiple time the same class it will return true for every call
+	 */
+	bool RegisterWriter(const UClass* Writer);
+
+	/**
+	 * Look if there is a registered translator for this source data.
+	 * This allow us to by pass the original asset tools system to import supported asset.
+	 * @Param SourceData - The source data input we want to translate to Uod
+	 * @return True if there is a registered translator that can handle handle this source data, false otherwise.
+	 */
+	bool CanTranslateSourceData(const UInterchangeSourceData* SourceData) const;
+
+	/**
+	 * Call this to start an import asset process, the caller must specify a source data.
+	 * This import process can import many different asset, but all in the game content.
+	 *
+	 * @Param ContentPath - The content path where to import the assets
+	 * @Param SourceData - The source data input we want to translate
+	 * @param ImportAssetParameters - All import asset parameter we need to pass to the import asset function
+	 * @return true if the import succeed, false otherwise.
+	 */
+	UFUNCTION(BlueprintCallable, Category = "Interchange | Import Manager")
+	bool ImportAsset(const FString& ContentPath, const UInterchangeSourceData* SourceData, const FImportAssetParameters& ImportAssetParameters);
+
+	/**
+	 * Call this to start an import scene process, the caller must specify a source data.
+	 * This import process can import many different asset and there transform (USceneComponent) and store the result in a blueprint and add the blueprint to the level.
+	 *
+	 * @Param ContentPath - The content path where to import the assets
+	 * @Param SourceData - The source data input we want to translate, this object will be duplicate to allow multithread safe operations
+	 * @param ImportAssetParameters - All import asset parameter we need to pass to the import asset function
+	 * @return true if the import succeed, false otherwise.
+	 */
+	UFUNCTION(BlueprintCallable, Category = "Interchange | Import Manager")
+	bool ImportScene(const FString& ContentPath, const UInterchangeSourceData* SourceData, bool bIsReimport = false, bool bIsAutomated = false);
+
+	/**
+	 * Call this to start an export asset process, the caller must specify a source data.
+	 * 
+	 * @Param SourceData - The source data output 
+	 * @Param bIsAutomated - If true the exporter will not show any UI or dialog
+	 * @return true if the import succeed, false otherwise.
+	 */
+	UFUNCTION(BlueprintCallable, Category = "Interchange | Export Manager")
+	bool ExportAsset(const UObject* Asset, bool bIsAutomated = false);
+
+	/**
+	 * Call this to start an export scene process, the caller must specify a source data
+	 * This import process can import many different asset and there transform (USceneComponent) and store the result in a blueprint and add the blueprint to the level.
+	 * @Param SourceData - The source data input we want to translate
+	 * @Param bIsAutomated - If true the import asset will not show any UI or dialog
+	 * @return true if the import succeed, false otherwise.
+	 */
+	UFUNCTION(BlueprintCallable, Category = "Interchange | Export Manager")
+	bool ExportScene(const UObject* World, bool bIsAutomated = false);
+
+	/*
+	* Script helper to create a source data object pointing on a file on disk
+	* @Param InFilename: Specify a file on disk
+	* @return: A new UInterchangeSourceData.
+	*/
+	UFUNCTION(BlueprintCallable, Category = "Interchange | Import Manager")
+	UInterchangeSourceData* CreateSourceData(const FString& InFileName);
+
+	/**
+	* Script helper to get a registered factory for a specified class
+	* @Param FactoryClass: The class we search a registerd factory
+	* @return: if found, we return the factory that is registered. Return NULL if nothing found.
+	*/
+	UFUNCTION(BlueprintCallable, Category = "Interchange | Import Manager")
+	UInterchangeFactoryBase* GetRegisterFactory(UClass* FactoryClass)
+	{
+		for (auto Kvp : RegisteredFactories)
+		{
+			if (FactoryClass->IsChildOf(Kvp.Key))
+			{
+				return Kvp.Value;
+			}
+		}
+		return nullptr;
+	}
+
+	/**
+	 * Return an FImportAsynHelper pointer. The pointer is deleted when ReleaseAsyncHelper is call.
+	 * @param Data - The data we want to pass to the different import tasks
+	 */
+	TWeakPtr<Interchange::FImportAsyncHelper> CreateAsyncHelper(const Interchange::FImportAsyncHelperData& Data);
+
+	/** Delete the specified AsyncHelper and remove it from the array that was holding it. */
+	void ReleaseAsyncHelper(TWeakPtr<Interchange::FImportAsyncHelper> AsyncHelper);
+
+protected:
+
+	/*
+	 * Return the first translator that can translate the source data.
+	 * @Param SourceData - The source data for which we search a translator.
+	 * @return return a matching translator or nullptr if there is no match.
+	 */
+	UInterchangeTranslatorBase* GetTranslatorForSourceData(const UInterchangeSourceData* SourceData) const;
+
+	/** Return true if we can show some UI */
+	static bool IsAttended();
+
+	/*
+	 * Find all Pipeline candidate (c++, blueprint and python).
+	 * @Param SourceData - The source data for which we search a translator.
+	 * @return return a matching translator or nullptr if there is no match.
+	 */
+	void FindPipelineCandidate(TArray<UClass*>& PipelineCandidates);
+
+private:
+	//By using pointer, there is no issue if the array get resize
+	TArray<TSharedPtr<Interchange::FImportAsyncHelper> > ImportTasks;
+
+	TSharedPtr<FAsyncTaskNotification> Notification = nullptr;
+
+	//The manager will create translator at every import, translator must be able to retrieve payload information when the factory ask for it.
+	//The translator stored has value is only use to know if we can use this type of translator.
+	UPROPERTY()
+	TMap<const UClass*, UInterchangeTranslatorBase* > RegisteredTranslators;
+	
+	//The manager will create only one pipeline per type
+	UPROPERTY()
+	TMap<const UClass*, UInterchangePipelineBase* > RegisteredPipelines;
+
+	//The manager will create only one factory per type
+	UPROPERTY()
+	TMap<const UClass*, UInterchangeFactoryBase* > RegisteredFactories;
+
+	//The manager will create only one writer per type
+	UPROPERTY()
+	TMap<const UClass*, UInterchangeWriterBase* > RegisteredWriters;
+
+	friend class Interchange::FScopedTranslator;
+};
diff --git a/Engine/Source/Runtime/Interchange/Core/Public/InterchangePipelineBase.h b/Engine/Source/Runtime/Interchange/Core/Public/InterchangePipelineBase.h
new file mode 100644
--- /dev/null
+++ b/Engine/Source/Runtime/Interchange/Core/Public/InterchangePipelineBase.h
@@ -0,0 +1,86 @@
+// Copyright Epic Games, Inc. All Rights Reserved.
+
+#pragma once
+
+#include "CoreMinimal.h"
+#include "Nodes/BaseNodeContainer.h"
+#include "InterchangeSourceData.h"
+#include "UObject/Object.h"
+#include "UObject/ObjectMacros.h"
+
+#include "InterchangePipelineBase.generated.h"
+
+UCLASS(BlueprintType, Blueprintable)
+class INTERCHANGECORE_API UInterchangePipelineBase : public UObject
+{
+	GENERATED_BODY()
+
+public:
+
+	/**
+	 * Non virtual helper to allow blueprint to implement event base function to implement a pipeline,
+	 * the Interchange manager is calling this function not the virtual one that is call by the default implementation.
+	 */
+	UFUNCTION(BlueprintNativeEvent, BlueprintCallable, Category = "Interchange | Translator")
+	bool ScriptedExecuteImportPipeline(UInterchangeBaseNodeContainerAdapter* BaseNodeContainerAdapter);
+	/** The default implementation (call if the blueprint do not have any implementation) will call the virtual ExecuteImportPipeline */
+	bool ScriptedExecuteImportPipeline_Implementation(UInterchangeBaseNodeContainerAdapter* BaseNodeContainerAdapter)
+	{
+		//By default we call the virtual import pipeline execution
+		return ExecuteImportPipeline(BaseNodeContainerAdapter);
+	}
+
+	/**
+	 * Non virtual helper to allow blueprint to implement event base function let the interchange know if it can run asynchronously.
+	 * the Interchange manager is calling this function not the virtual one that is call by the default implementation.
+	 */
+	UFUNCTION(BlueprintCallable, BlueprintNativeEvent, Category = "Interchange | Translator")
+	bool ScriptedExecuteExportPipeline(UInterchangeBaseNodeContainerAdapter* BaseNodeContainerAdapter);
+
+	/** The default implementation (call if the blueprint do not have any implementation) will call the virtual ExecuteExportPipeline */
+	bool ScriptedExecuteExportPipeline_Implementation(UInterchangeBaseNodeContainerAdapter* BaseNodeContainerAdapter)
+	{
+		return ExecuteExportPipeline(BaseNodeContainerAdapter);
+	}
+
+	/**
+	 * Non virtual helper to allow blueprint to implement event base function let the interchange know if it can run asynchronously.
+	 * the Interchange manager is calling this function not the virtual one that is call by the default implementation.
+	 */
+	UFUNCTION(BlueprintNativeEvent, BlueprintCallable, Category = "Interchange | Translator")
+	bool ScriptedCanExecuteOnAnyThread();
+
+	/** The default implementation (call if the blueprint do not have any implementation) will call the virtual CanExecuteAsync */
+	bool ScriptedCanExecuteOnAnyThread_Implementation()
+	{
+		return CanExecuteOnAnyThread();
+	}
+protected:
+
+	/**
+	 * This function can modify the BaseNodeContainer to create a pipeline that will set the graph and the nodes options has it want it to be imported by the factories
+	 * The interchange manager is not calling this function directly. It is calling the blueprint native event in case this object is a blueprint derive object.
+	 * By default the scripted implementation is calling this virtual pipeline.
+	 */
+	virtual bool ExecuteImportPipeline(UInterchangeBaseNodeContainerAdapter* BaseNodeContainerAdapter)
+	{
+		return false;
+	}
+
+	/**
+	 * This function tell the interchange manager if we can execute this pipeline in async mode. If it return false, the ScriptedExecuteImportPipeline
+	 * will be call on the main thread (GameThread), if true it will be run in a background thread and possibly in parallel. If there is multiple
+	 * import process in same time.
+	 *
+	 */
+	virtual bool CanExecuteOnAnyThread()
+	{
+		return true;
+	}
+
+	/** This function can modify the BaseNodeContainer to create a pipeline that will set/validate the graph nodes hierarchy and options.*/
+	virtual bool ExecuteExportPipeline(UInterchangeBaseNodeContainerAdapter* BaseNodeContainerAdapter)
+	{
+		return false;
+	}
+};
diff --git a/Engine/Source/Runtime/Interchange/Core/Public/InterchangeSourceData.h b/Engine/Source/Runtime/Interchange/Core/Public/InterchangeSourceData.h
new file mode 100644
--- /dev/null
+++ b/Engine/Source/Runtime/Interchange/Core/Public/InterchangeSourceData.h
@@ -0,0 +1,103 @@
+// Copyright Epic Games, Inc. All Rights Reserved.
+
+#pragma once
+
+#include "CoreMinimal.h"
+#include "Misc/SecureHash.h"
+#include "UObject/Object.h"
+#include "UObject/ObjectMacros.h"
+
+#include "InterchangeSourceData.generated.h"
+
+/*
+ * Helper class to be able to read different source data
+ * File on disk
+ * HTTP URL (TODO)
+ * Memory buffer (TODO)
+ * Stream (TODO)
+ */
+
+UCLASS(BlueprintType, Blueprintable)
+class INTERCHANGECORE_API UInterchangeSourceData : public UObject
+{
+	 GENERATED_BODY()
+public:
+	UInterchangeSourceData()
+	{
+		Filename = TEXT("");
+		FileContentHashCache.Reset();
+	}
+
+	UInterchangeSourceData(FString InFilename)
+	: Filename(InFilename)
+	{
+		FileContentHashCache.Reset();
+	}
+	
+	UFUNCTION(BlueprintCallable, Category = "Interchange | Import Manager")
+	FORCEINLINE FString GetFilename() const
+	{
+		return Filename;
+	}
+
+	/** Return the hash of the content point by Filename. */
+	FORCEINLINE TOptional<FMD5Hash> GetFileContentHash() const
+	{
+		if(!FileContentHashCache.IsSet())
+		{
+			ComputeFileContentHashCache();
+		}
+		return FileContentHashCache;
+	}
+
+	UFUNCTION(BlueprintCallable, Category = "Interchange | Import Manager")
+	FORCEINLINE bool SetFilename(const FString& InFilename)
+	{
+		Filename = InFilename;
+		//Reset the cache
+		FileContentHashCache.Reset();
+		return true;
+	}
+
+	/** Return a easy to read source description string, this is mainly use for logging or UI. */
+	FORCEINLINE FString ToDisplayString() const
+	{
+		if (!Filename.IsEmpty())
+		{
+			bool bFileExist = FPaths::FileExists(Filename);
+			FString CleanFilename = FPaths::GetCleanFilename(Filename);
+			FString BasePath = FPaths::GetPath(Filename);
+			//limit the base path to 20 character do ellipsis in the middle
+			if (BasePath.Len() > 33)
+			{
+				FString RightPath = BasePath.Right(15);
+				FString LeftPath = BasePath.Left(15);
+				BasePath = LeftPath + TEXT("...") + RightPath;
+			}
+			FString DisplayString = BasePath + CleanFilename;
+			return DisplayString;
+		}
+		return FString();
+	}
+
+private:
+	void ComputeFileContentHashCache() const
+	{
+		FileContentHashCache.Reset();
+		if (!Filename.IsEmpty())
+		{
+			FileContentHashCache = FMD5Hash::HashFile(*Filename);
+		}
+	}
+
+	UPROPERTY()
+	FString Filename;
+
+	/**
+	 * Hash cache for the file content.
+	 * We use mutable because the cache is compute when we use the get which is a const function.
+	 * The cache is compute only in the GetFileContentHash() to let the client control in which thread he will compute the cache.
+	 * It also make sure we do not waste CPU computing a cache in case no client use GetFileContentHash().
+	 */
+	mutable TOptional<FMD5Hash> FileContentHashCache;
+};
diff --git a/Engine/Source/Runtime/Interchange/Core/Public/InterchangeTranslatorBase.h b/Engine/Source/Runtime/Interchange/Core/Public/InterchangeTranslatorBase.h
new file mode 100644
--- /dev/null
+++ b/Engine/Source/Runtime/Interchange/Core/Public/InterchangeTranslatorBase.h
@@ -0,0 +1,36 @@
+// Copyright Epic Games, Inc. All Rights Reserved.
+
+#pragma once
+
+#include "CoreMinimal.h"
+#include "Nodes/BaseNodeContainer.h"
+#include "InterchangeSourceData.h"
+#include "UObject/Object.h"
+#include "UObject/ObjectMacros.h"
+
+#include "InterchangeTranslatorBase.generated.h"
+
+UCLASS(BlueprintType, Blueprintable, Abstract)
+class INTERCHANGECORE_API UInterchangeTranslatorBase : public UObject
+{
+	GENERATED_BODY()
+public:
+
+	/** return true if the translator can translate the specified file. */
+	virtual bool CanImportSourceData(const UInterchangeSourceData* SourceData) const
+	{
+		return false;
+	}
+
+	/**
+	 * Translate a source data into node(s) that are hold in the specified nodes container.
+	 *
+	 * @param SourceData - The source data containing the data to translate
+	 * @param BaseNodeContainer - The nodes container where to put the translated source data result.
+	 * @return true if the translator can translate the source data, false otherwise.
+	 */
+	virtual bool Translate(const UInterchangeSourceData* SourceData, Interchange::FBaseNodeContainer& BaseNodeContainer) const
+	{
+		return false;
+	}
+};
diff --git a/Engine/Source/Runtime/Interchange/Core/Public/InterchangeWriterBase.h b/Engine/Source/Runtime/Interchange/Core/Public/InterchangeWriterBase.h
new file mode 100644
--- /dev/null
+++ b/Engine/Source/Runtime/Interchange/Core/Public/InterchangeWriterBase.h
@@ -0,0 +1,24 @@
+// Copyright Epic Games, Inc. All Rights Reserved.
+
+#pragma once
+
+#include "CoreMinimal.h"
+#include "Nodes/BaseNodeContainer.h"
+#include "UObject/Object.h"
+#include "UObject/ObjectMacros.h"
+
+#include "InterchangeWriterBase.generated.h"
+
+UCLASS(BlueprintType, Blueprintable, Abstract)
+class INTERCHANGECORE_API UInterchangeWriterBase : public UObject
+{
+	GENERATED_BODY()
+public:
+
+	/**
+	 * Export an nodes container (graph)
+	 * @param BaseNodeContainer - The container holding nodes that describe what to export
+	 * @return true if the writer can export the Nodes, false otherwise.
+	 */
+	virtual bool Export(Interchange::FBaseNodeContainer& BaseNodeContainer) const {return false;}
+};
diff --git a/Engine/Source/Runtime/Interchange/Core/Public/Nodes/BaseNode.h b/Engine/Source/Runtime/Interchange/Core/Public/Nodes/BaseNode.h
new file mode 100644
--- /dev/null
+++ b/Engine/Source/Runtime/Interchange/Core/Public/Nodes/BaseNode.h
@@ -0,0 +1,487 @@
+// Copyright Epic Games, Inc. All Rights Reserved.
+
+#pragma once
+
+#include "CoreMinimal.h"
+#include "Types/AttributeStorage.h"
+
+
+/**
+ * Internal Helper to get set custom property for class that derive from FBaseNode. This is use by the macro IMPLEMENT_UOD_ATTRIBUTE.
+ */
+namespace InterchangePrivateNodeBase
+{
+	/**
+	 * Retrieve a custom attribute if the attribute exist
+	 *
+	 * @param Attributes - The attribute storage you want to query the custom attribute
+	 * @param AttributeKey - The storage key for the attribute
+	 * @param OperationName - The name of the operation in case there is an error
+	 * @param OutAttributeValue - This is where we store the value we retrieve from the storage
+	 *
+	 * @return - return true if the attribute exist in the storage and was query without error.
+	 *           return false if the attribute do not exist or there is an error retriving it from the Storage
+	 */
+	template<typename ValueType>
+	bool GetCustomAttribute(const Interchange::FAttributeStorage& Attributes, const Interchange::FAttributeKey& AttributeKey, const FString& OperationName, ValueType& OutAttributeValue)
+	{
+		if (!Attributes.ContainAttribute(AttributeKey))
+		{
+			return false;
+		}
+		Interchange::FAttributeStorage::TAttributeHandle<ValueType> AttributeHandle = Attributes.GetAttributeHandle<ValueType>(AttributeKey);
+		if (!AttributeHandle.IsValid())
+		{
+			return false;
+		}
+		Interchange::EAttributeStorageResult Result = AttributeHandle.Get(OutAttributeValue);
+		if (!Interchange::IsAttributeStorageResultSuccess(Result))
+		{
+			Interchange::LogAttributeStorageErrors(Result, OperationName, AttributeKey);
+			return false;
+		}
+		return true;
+	}
+
+	/**
+	 * Add or update a custom attribute value in the specified storage
+	 *
+	 * @param Attributes - The attribute storage you want to add or update the custom attribute
+	 * @param AttributeKey - The storage key for the attribute
+	 * @param OperationName - The name of the operation in case there is an error
+	 * @param AttributeValue - The value we want to add or update in the storage
+	 */
+	template<typename ValueType>
+	bool SetCustomAttribute(Interchange::FAttributeStorage& Attributes, const Interchange::FAttributeKey& AttributeKey, const FString& OperationName, const ValueType& AttributeValue)
+	{
+		Interchange::EAttributeStorageResult Result = Attributes.RegisterAttribute(AttributeKey, AttributeValue);
+		if (!Interchange::IsAttributeStorageResultSuccess(Result))
+		{
+			Interchange::LogAttributeStorageErrors(Result, OperationName, AttributeKey);
+			return false;
+		}
+		return true;
+	}
+}
+
+/**
+ * Use those macro to create Get/Set/ApplyToAsset for an attribute you want to support in your custom UOD node that derive from FBaseNode.
+ * The attribute key will be declare under private access modifier and the functions are declare under public access modifier.
+ * So the class access modifier is public after calling this macro.
+ *
+ * @note - The Get will return false if the attribute was never set.
+ * @note - The Set will add the attribute to the storage or update the value if the storage already has this attribute.
+ * @note - The Apply will set a member variable of the AssetType instance. It will apply the value only if the storage contain the attribute key
+ *
+ * @param AttributeName - The name of the Get/Set functions. Example if you pass Foo you will end up with GetFoo and SetFoo function
+ * @param AttributeType - This is to specify the type of the attribute. bool, float, FString... anything supported by the FAttributeStorage
+ * @param AssetType - This is the asset you want to apply the storage value"
+ * @param EnumType - Optional, specify it only if the AssetType member is an enum so we can type cast it in the apply function (we use uint8 to store the enum value)"
+ */
+#define IMPLEMENT_NODE_ATTRIBUTE(NodeClassName, AttributeName, AttributeType, AssetType, EnumType)										\
+private:																																\
+const FAttributeKey Macro_Custom##AttributeName##Key = Interchange::FAttributeKey(TEXT(#AttributeName));								\
+bool ApplyCustom##AttributeName##ToAsset(UObject* Asset) const																			\
+{																																		\
+	if (!Asset)																															\
+	{																																	\
+		return false;																													\
+	}																																	\
+	AssetType* TypedObject = Cast<AssetType>(Asset);																					\
+	if (!TypedObject)																													\
+	{																																	\
+		return false;																													\
+	}																																	\
+	AttributeType ValueData;																											\
+	if (GetCustom##AttributeName(ValueData))																							\
+	{																																	\
+		TypedObject->AttributeName = EnumType(ValueData);																				\
+		return true;																													\
+	}																																	\
+	return false;																														\
+}																																		\
+																																		\
+public:																																	\
+bool GetCustom##AttributeName(AttributeType& AttributeValue) const																		\
+{																																		\
+	FString OperationName = GetTypeName() + TEXT(".Get##AttributeName");																\
+	return InterchangePrivateNodeBase::GetCustomAttribute<AttributeType>(Attributes, Macro_Custom##AttributeName##Key, OperationName, AttributeValue);	\
+}																																		\
+																																		\
+bool SetCustom##AttributeName(const AttributeType& AttributeValue, bool bAddApplyDelegate = true)										\
+{																																		\
+	FString OperationName = GetTypeName() + TEXT(".Set##AttributeName");																\
+	if(InterchangePrivateNodeBase::SetCustomAttribute<AttributeType>(Attributes, Macro_Custom##AttributeName##Key, OperationName, AttributeValue))	\
+	{																																	\
+		if(bAddApplyDelegate)																											\
+		{																																\
+			TArray<FApplyAttributeToAsset>& Delegates = ApplyCustomAttributeDelegates.FindOrAdd(AssetType::StaticClass());					\
+			Delegates.Add(FApplyAttributeToAsset::CreateRaw(this, &NodeClassName::ApplyCustom##AttributeName##ToAsset));				\
+		}																																\
+		return true;																													\
+	}																																	\
+	return false;																														\
+}
+
+
+
+//Interchange namespace
+namespace Interchange
+{
+
+DECLARE_DELEGATE_RetVal_OneParam(bool, FApplyAttributeToAsset, UObject*);
+
+typedef FName FNodeUniqueID;
+
+
+/**
+ * Helper struct use to declare static const data we use in the FBaseNode
+ * Node that derive from FBaseNode can also add a struct that derive from this one to add there static data
+ * @note: The static data are mainly for holding Attribute keys. All attributes that are always available for a node should be in this class or a derived class.
+ */
+struct FBaseNodeStaticData
+{
+	static const FAttributeKey& UniqueIDKey()
+	{
+		static FAttributeKey AttributeKey(TEXT("__UNQ_ID_"));
+		return AttributeKey;
+	}
+
+	static const FAttributeKey& DisplayLabelKey()
+	{
+		static FAttributeKey AttributeKey(TEXT("__DSPL_LBL_"));
+		return AttributeKey;
+	}
+
+	static const FAttributeKey& ParentIDKey()
+	{
+		static FAttributeKey AttributeKey(TEXT("__PARENT_UID_"));
+		return AttributeKey;
+	}
+
+	static const FAttributeKey& IsEnabledKey()
+	{
+		static FAttributeKey AttributeKey(TEXT("__IS_NBLD_"));
+		return AttributeKey;
+	}
+
+	static const FAttributeKey& DependencyCountKey()
+	{
+		static FAttributeKey AttributeKey(TEXT("__DEPENDENCY_COUNT_"));
+		return AttributeKey;
+	}
+
+	static const FAttributeKey& DependencyBaseKey()
+	{
+		static FAttributeKey AttributeKey(TEXT("__DEPENDENCY_INDEX_"));
+		return AttributeKey;
+	}
+};
+
+/**
+ * This class is used to store and retrieve key value attributes. The attributes are store in a generic FAttributeStorage which serialize the value in a TArray64<uint8>
+ * See Interchange::EAttributeTypes to know the supported template types
+ * This is an abstract class. This is the base class of the interchange node graph format, all class in this format should derive from this class
+ */
+
+class FBaseNode
+{
+public:
+	virtual ~FBaseNode() = default;
+
+ 	FBaseNode(const FBaseNode& Other)
+ 	{
+ 		Attributes = Other.Attributes;
+ 	}
+
+	/**
+	 * Constructor
+	 * @param UniqueID - The uniqueId for this node
+	 * @param DisplayLabel - The name of the node
+	 *
+	 */
+	FBaseNode(const FNodeUniqueID& UniqueID, const FName& DisplayLabel)
+	{
+		EAttributeStorageResult Result = Attributes.RegisterAttribute(FBaseNodeStaticData::UniqueIDKey(), UniqueID, EAttributeProperty::NoHash);
+		if (!IsAttributeStorageResultSuccess(Result))
+		{
+			LogAttributeStorageErrors(Result, TEXT("RegisterAttribute"), FBaseNodeStaticData::UniqueIDKey());
+		}
+
+		Result = Attributes.RegisterAttribute(FBaseNodeStaticData::DisplayLabelKey(), DisplayLabel, EAttributeProperty::NoHash);
+		if (!IsAttributeStorageResultSuccess(Result))
+		{
+			LogAttributeStorageErrors(Result, TEXT("RegisterAttribute"), FBaseNodeStaticData::DisplayLabelKey());
+		}
+	}
+
+	/**
+	 * Return the node type name of the class, we use this when reporting error
+	 */
+	virtual FString GetTypeName() const
+	{
+		const FString TypeName = TEXT("BaseNode");
+		return TypeName;
+	}
+
+	/**
+	 * Add an attribute to the node
+	 * @param NodeAttributeKey - The key of the attribute
+	 * @param Value - The attribute value.
+	 *
+	 */
+	template<typename T>
+	FAttributeStorage::TAttributeHandle<T> RegisterAttribute(const FAttributeKey& NodeAttributeKey, const T& Value)
+	{
+		const EAttributeStorageResult Result = Attributes.RegisterAttribute(NodeAttributeKey, Value);
+		
+		if (IsAttributeStorageResultSuccess(Result))
+		{
+			return Attributes.GetAttributeHandle<T>(NodeAttributeKey);
+		}
+		LogAttributeStorageErrors(Result, TEXT("RegisterAttribute"), NodeAttributeKey);
+		return FAttributeStorage::TAttributeHandle<T>();
+	}
+
+	/**
+	 * Return true if the node contain an attribute with the specified Key
+	 * @param nodeAttributeKey - The key of the searched attribute
+	 *
+	 */
+	virtual bool HasAttribute(const FAttributeKey& NodeAttributeKey) const
+	{
+		return Attributes.ContainAttribute(NodeAttributeKey);
+	}
+
+
+	/**
+	 * This function returnan  attribute handle for the specified Key.
+	 * If there is an issue with the KEY or storage the method will trip a check, always make sure you have a valid key before calling this
+	 * 
+	 * @param NodeAttributeKey - The key of the attribute
+	 */
+	template<typename T>
+	FAttributeStorage::TAttributeHandle<T> GetAttributeHandle(const FAttributeKey& NodeAttributeKey) const
+	{
+		return Attributes.GetAttributeHandle<T>(NodeAttributeKey);
+	}
+
+	/**
+	 * Return the unique id pass in the constructor.
+	 *
+	 */
+	virtual FNodeUniqueID GetUniqueID() const
+	{
+		FNodeUniqueID UniqueID = NAME_None;
+		Attributes.GetAttributeHandle<FNodeUniqueID>(FBaseNodeStaticData::UniqueIDKey()).Get(UniqueID);
+		return UniqueID;
+	}
+
+	/**
+	 * Return the display label pass in the constructor.
+	 *
+	 */
+	virtual FName GetDisplayLabel() const
+	{
+		checkSlow(Attributes.ContainAttribute(FBaseNodeStaticData::DisplayLabelKey()));
+		FName DisplayLabel = NAME_None;
+		Attributes.GetAttributeHandle<FName>(FBaseNodeStaticData::DisplayLabelKey()).Get(DisplayLabel);
+		return DisplayLabel;
+	}
+
+	/**
+	 * Return the parent unique id. In case the attribute does not exist it will return InvalidNodeUID()
+	 */
+	virtual FNodeUniqueID GetParentUID() const
+	{
+		if (!Attributes.ContainAttribute(FBaseNodeStaticData::ParentIDKey()))
+		{
+			return InvalidNodeUID();
+		}
+
+		FNodeUniqueID ParentUniqueID = NAME_None;
+		FAttributeStorage::TAttributeHandle<FNodeUniqueID> Handle = Attributes.GetAttributeHandle<FNodeUniqueID>(FBaseNodeStaticData::ParentIDKey());
+		if(Handle.IsValid())
+		{
+			Handle.Get(ParentUniqueID);
+			return ParentUniqueID;
+		}
+		return InvalidNodeUID();
+	}
+
+	/**
+	 * Set the parent unique id.
+	 */
+	virtual bool SetParentUID(FNodeUniqueID ParentUID)
+	{
+		EAttributeStorageResult Result = Attributes.RegisterAttribute(FBaseNodeStaticData::ParentIDKey(), ParentUID);
+		if(IsAttributeStorageResultSuccess(Result))
+		{
+			FAttributeStorage::TAttributeHandle<FNodeUniqueID> Handle = Attributes.GetAttributeHandle<FNodeUniqueID>(FBaseNodeStaticData::ParentIDKey());
+			return Handle.IsValid();
+		}
+		return false;
+	}
+
+	/**
+	 * This function allow to retrieve the dependency for this object.
+	 * 
+	 */
+	virtual void GetDependecies(TArray<FNodeUniqueID>& OutDependencies ) const
+	{
+		OutDependencies.Reset();
+
+		int32 DepenedenciesCount = 0;
+		if (!Attributes.ContainAttribute(FBaseNodeStaticData::DependencyCountKey()))
+		{
+			return;
+		}
+		FAttributeStorage::TAttributeHandle<int32> Handle = Attributes.GetAttributeHandle<int32>(FBaseNodeStaticData::DependencyCountKey());
+		if (!Handle.IsValid())
+		{
+			return;
+		}
+		Handle.Get(DepenedenciesCount);
+		for (int32 DepIndex = 0; DepIndex < DepenedenciesCount; ++DepIndex)
+		{
+			FString DepIndexKeyString = FBaseNodeStaticData::DependencyBaseKey().Key.ToString() + FString::FromInt(DepIndex);
+			FAttributeKey DepIndexKey(*DepIndexKeyString);
+			FAttributeStorage::TAttributeHandle<FNodeUniqueID> HandleDep = Attributes.GetAttributeHandle<FNodeUniqueID>(DepIndexKey);
+			if (!HandleDep.IsValid())
+			{
+				continue;
+			}
+			FNodeUniqueID& NodeUniqueID = OutDependencies.AddDefaulted_GetRef();
+			HandleDep.Get(NodeUniqueID);
+		}
+	}
+
+	/**
+	 * Add one dependency to this object.
+	 */
+	virtual bool SetDependencyUID(FNodeUniqueID DependencyUID)
+	{
+		if (!Attributes.ContainAttribute(FBaseNodeStaticData::DependencyCountKey()))
+		{
+			const int32 DependencyCount = 0;
+			EAttributeStorageResult Result = Attributes.RegisterAttribute<int32>(FBaseNodeStaticData::DependencyCountKey(), DependencyCount);
+		}
+		FAttributeStorage::TAttributeHandle<int32> Handle = Attributes.GetAttributeHandle<int32>(FBaseNodeStaticData::DependencyCountKey());
+		if(!ensure(Handle.IsValid()))
+		{
+			return false;
+		}
+		int32 DepIndex = 0;
+		Handle.Get(DepIndex);
+		FString DepIndexKeyString = FBaseNodeStaticData::DependencyBaseKey().Key.ToString() + FString::FromInt(DepIndex);
+		FAttributeKey DepIndexKey(*DepIndexKeyString);
+		//Increment the counter
+		DepIndex++;
+		Handle.Set(DepIndex);
+
+		EAttributeStorageResult DepResult = Attributes.RegisterAttribute<FNodeUniqueID>(DepIndexKey, DependencyUID);
+		return true;
+	}
+
+	/**
+	 * IsEnable true mean that the node will be import/export, if false it will be discarded.
+	 * Return false if this node was disabled. Return true if the attribute is not there or if it was enabled.
+	 */
+	virtual bool IsEnabled() const
+	{
+		FAttributeStorage::TAttributeHandle<bool> Handle = Attributes.GetAttributeHandle<bool>(FBaseNodeStaticData::IsEnabledKey());
+		if (Handle.IsValid())
+		{
+			bool bIsEnabled = false;
+			Handle.Get(bIsEnabled);
+			return bIsEnabled;
+		}
+		return true;
+	}
+
+	/**
+	 * Set the IsEnable attribute to determine if this node should be part of the import/export process
+	 * @param bIsEnabled - The enabled state we want to set this node. True will import/export the node, fals will not.
+	 * @return true if it was able to set the attribute, false otherwise.
+	 */
+	virtual bool SetEnabled(const bool bIsEnabled)
+	{
+		EAttributeStorageResult Result = Attributes.RegisterAttribute(FBaseNodeStaticData::IsEnabledKey(), bIsEnabled);
+		if (IsAttributeStorageResultSuccess(Result))
+		{
+			FAttributeStorage::TAttributeHandle<bool> Handle = Attributes.GetAttributeHandle<bool>(FBaseNodeStaticData::IsEnabledKey());
+			return Handle.IsValid();
+		}
+		return false;
+	}
+
+	/**
+	 * Return a FGuid build from the FSHA1 of all the attribute data contain in the node.
+	 *
+	 * @note the attribute are sorted by key when building the FSHA1 data. The hash will be deterministic for the same data whatever
+	 * the order we add the attributes.
+	 * This function interface is pure virtual
+	 */
+	virtual FGuid GetHash() const
+	{
+		return Attributes.GetStorageHash();
+	}
+
+	/**
+	 * Optional, Any node that can import/export an asset should return the UClass of the asset so we can find factory/writer
+	 */
+	virtual class UClass* GetAssetClass() const
+	{
+		return nullptr;
+	}
+
+	/**
+	 * Normally there should be no other data need to be serialize even in the derived class
+	 * Everything should be in the attribute storage. Payload source data path should be store in the attribute storage.
+	 * We add this virtual function in case an node need to serialize more stuff.
+	 */
+	virtual FArchive& Serialize(FArchive& Ar)
+	{
+		Ar << Attributes;
+		return Ar;
+	}
+	
+	friend FArchive& operator<<(FArchive& Ar, FBaseNode& Node)
+	{
+		//Call the  virtual version of the serialize.
+		return Node.Serialize(Ar);
+	}
+
+	/** Return the invalid unique ID */
+	static FNodeUniqueID InvalidNodeUID()
+	{
+		return NAME_None;
+	}
+
+	void ApplyAllCustomAttributeToAsset(UObject* Object) const
+	{
+		UClass* ObjectClass = Object->GetClass();
+		for (const TPair<UClass*, TArray<FApplyAttributeToAsset>>& ClassDelegatePair : ApplyCustomAttributeDelegates)
+		{
+			if(ObjectClass->IsChildOf(ClassDelegatePair.Key))
+			{
+				for (const FApplyAttributeToAsset& Delegate : ClassDelegatePair.Value)
+				{
+					if(Delegate.IsBound())
+					{
+						Delegate.Execute(Object);
+					}
+				}
+			}
+		}
+	}
+
+protected:
+	/** The storage use to store the Key value attribute for this node. */
+	FAttributeStorage Attributes;
+
+	/* This array hold the delegate to apply the attribute that has to be set on an UObject */
+	TMap<UClass*, TArray<FApplyAttributeToAsset>> ApplyCustomAttributeDelegates;
+};
+
+}
diff --git a/Engine/Source/Runtime/Interchange/Core/Public/Nodes/BaseNodeContainer.h b/Engine/Source/Runtime/Interchange/Core/Public/Nodes/BaseNodeContainer.h
new file mode 100644
--- /dev/null
+++ b/Engine/Source/Runtime/Interchange/Core/Public/Nodes/BaseNodeContainer.h
@@ -0,0 +1,227 @@
+// Copyright Epic Games, Inc. All Rights Reserved.
+
+#pragma once
+
+#include "CoreMinimal.h"
+#include "Templates/UniquePtr.h"
+#include "Nodes/BaseNode.h"
+
+//Interchange namespace
+namespace Interchange
+{
+
+/**
+ * Interchange FBaseNode graph is a format used to feed asset/scene import/reimport/export factories/writer.
+ * This container hold a flat list of all nodes that have been translated from the source data.
+ *
+ * Translators are filling this container and the Import/Export managers are reading it to execute the import/export process
+ */
+class INTERCHANGECORE_API FBaseNodeContainer
+{
+public:
+	//To avoid issue with the TMap<FNodeUniqueID, TUniquePtr<FBaseNode> > Nodes
+	//We have to delete the copy constructor and the = operator to avoid CopyToDelete to be instantiate which will implement a copy of the TMap.
+	//The CopyToEmpty is instantiate when we add the INTERCHANGECORE_API to the class
+	FBaseNodeContainer(const FBaseNodeContainer&) = delete;
+	FBaseNodeContainer& operator=(const FBaseNodeContainer&) = delete;
+
+	FBaseNodeContainer() = default;
+
+	/**
+	 * Add a node in the container, the node will be set into a TUniquePtr, which mean the container will have the memory ownership of the pass pointer.
+	 * You can get a ref of the added node by using GetNode function.
+	 *
+	 * @param Node - a pointer on the node you want to add
+	 * @return: return the node unique ID of the added item, return InvalidNodeUID if the node is already there or cannot be added.
+	 *
+	 * @note The node memory will be manage by the container after the add, the destructor of the container will free correctly all nodes
+	 *
+	 */
+	FNodeUniqueID AddNode(FBaseNode* Node)
+	{
+		if (!Node)
+		{
+			return FBaseNode::InvalidNodeUID();
+		}
+
+		FNodeUniqueID NodeUniqueID = Node->GetUniqueID();
+		if (NodeUniqueID == FBaseNode::InvalidNodeUID())
+		{
+			return FBaseNode::InvalidNodeUID();
+		}
+		
+		//Cannot add an node with the same IDs
+		if (Nodes.Contains(NodeUniqueID))
+		{
+			return FBaseNode::InvalidNodeUID();
+		}
+		//Create a unique node pointer
+		Nodes.Add(NodeUniqueID, TUniquePtr<FBaseNode>(Node));
+		return NodeUniqueID;
+	}
+
+	/** Return true if the node unique ID exist in the container */
+	bool IsNodeUIDValid(FNodeUniqueID NodeUniqueID) const
+	{
+		if (NodeUniqueID == FBaseNode::InvalidNodeUID())
+		{
+			return false;
+		}
+		return Nodes.Contains(NodeUniqueID);
+	}
+
+	/** Unordered iteration of the all nodes */
+	void IterateNodes(TFunctionRef<void(const FNodeUniqueID, const FBaseNode*)> IterationLambda)
+	{
+		for (TPair<FNodeUniqueID, TUniquePtr<FBaseNode>>& NodeKeyValue : Nodes)
+		{
+			IterationLambda(NodeKeyValue.Key, NodeKeyValue.Value.Get());
+		}
+	}
+
+	/** Return all nodes that do not have any parent */
+	void GetRoots(TArray<FNodeUniqueID>& RootNodes)
+	{
+		for (TPair<FNodeUniqueID, TUniquePtr<FBaseNode>>& NodeKeyValue : Nodes)
+		{
+			if (NodeKeyValue.Value->GetParentUID() == FBaseNode::InvalidNodeUID())
+			{
+				RootNodes.Add(NodeKeyValue.Key);
+			}
+		}
+	}
+
+	/** Get an node pointer */
+	FBaseNode* GetNode(FNodeUniqueID NodeUniqueID)
+	{
+		if (NodeUniqueID == FBaseNode::InvalidNodeUID())
+		{
+			return nullptr;
+		}
+		if (!Nodes.Contains(NodeUniqueID))
+		{
+			return nullptr;
+		}
+		if(TUniquePtr<FBaseNode>& Node = Nodes.FindChecked(NodeUniqueID))
+		{
+			return Node.Get();
+		}
+		return nullptr;
+	}
+
+	/** Get an node pointer */
+	const FBaseNode* GetNode(FNodeUniqueID NodeUniqueID) const
+	{
+		return const_cast<FBaseNodeContainer*>(this)->GetNode(NodeUniqueID);
+	}
+
+	/** Get an node reference. Assert if there is an error */
+	FBaseNode& GetNodeChecked(FNodeUniqueID NodeUniqueID)
+	{
+		return GetNodeCheckedInternal(NodeUniqueID);
+	}
+
+	/** Get an node reference. Assert if there is an error */
+	const FBaseNode& GetNodeCheck(FNodeUniqueID NodeUniqueID) const
+	{
+		return const_cast<FBaseNodeContainer*>(this)->GetNodeCheckedInternal(NodeUniqueID);
+	}
+
+	/** Set node ParentUID */
+	bool SetNodeParentUID(FNodeUniqueID NodeUniqueID, FNodeUniqueID NewParentNodeUID)
+	{
+		if (!Nodes.Contains(NodeUniqueID))
+		{
+			return false;
+		}
+		if (!Nodes.Contains(NewParentNodeUID))
+		{
+			return false;
+		}
+		TUniquePtr<FBaseNode>& Node = Nodes.FindChecked(NodeUniqueID);
+		check(Node.IsValid());
+		Node->SetParentUID(NewParentNodeUID);
+		return true;
+	}
+
+	/** Get the node children count */
+	int32 GetNodeChildrenCount(FNodeUniqueID NodeUniqueID) const
+	{
+		TArray<FNodeUniqueID> ChildrenUIDs = GetNodeChildrenUIDs(NodeUniqueID);
+		return ChildrenUIDs.Num();
+	}
+
+	/** Get all children UID */
+	TArray<FNodeUniqueID> GetNodeChildrenUIDs(FNodeUniqueID NodeUniqueID) const
+	{
+		TArray<FNodeUniqueID> OutChildrenUIDs;
+		for (const TPair<FNodeUniqueID, TUniquePtr<FBaseNode>>& NodeKeyValue : Nodes)
+		{
+			if (NodeKeyValue.Value->GetParentUID() == NodeUniqueID)
+			{
+				OutChildrenUIDs.Add(NodeKeyValue.Key);
+			}
+		}
+		return OutChildrenUIDs;
+	}
+
+	/** Get the node nth const children */
+	FBaseNode* GetNodeChildren(FNodeUniqueID NodeUniqueID, int32 ChildIndex)
+	{
+		return GetNodeChildrenInternal(NodeUniqueID, ChildIndex);
+	}
+
+	/** Get the node nth const children. Const version */
+	const FBaseNode* GetNodeChildren(FNodeUniqueID NodeUniqueID, int32 ChildIndex) const
+	{
+		return const_cast<FBaseNodeContainer*>(this)->GetNodeChildrenInternal(NodeUniqueID, ChildIndex);
+	}
+
+	/** Get the node nth const children. Assert if there is an error */
+	FBaseNode& GetNodeChildrenChecked(FNodeUniqueID NodeUniqueID, int32 ChildIndex)
+	{
+		TArray<FNodeUniqueID> ChildrenUIDs = GetNodeChildrenUIDs(NodeUniqueID);
+		check (!ChildrenUIDs.IsValidIndex(ChildIndex));
+		TUniquePtr<FBaseNode>& Node = Nodes.FindChecked(ChildrenUIDs[ChildIndex]);
+		check(Node.IsValid())
+		return *Node.Get();
+	}
+
+	/** Get the node nth const children. Const version. Assert if there is an error */
+	const FBaseNode& GetNodeChildrenChecked(FNodeUniqueID NodeUniqueID, int32 ChildIndex) const
+	{
+		return const_cast<FBaseNodeContainer*>(this)->GetNodeChildrenChecked(NodeUniqueID, ChildIndex);
+	}
+
+private:
+
+	FBaseNode& GetNodeCheckedInternal(FNodeUniqueID NodeUniqueID)
+	{
+		check(NodeUniqueID != FBaseNode::InvalidNodeUID());
+		TUniquePtr<FBaseNode>& Node = Nodes.FindChecked(NodeUniqueID);
+		check(Node.IsValid());
+		return *Node.Get();
+	}
+
+	FBaseNode* GetNodeChildrenInternal(FNodeUniqueID NodeUniqueID, int32 ChildIndex)
+	{
+		TArray<FNodeUniqueID> ChildrenUIDs = GetNodeChildrenUIDs(NodeUniqueID);
+		if (!ChildrenUIDs.IsValidIndex(ChildIndex))
+		{
+			return nullptr;
+		}
+
+		if (Nodes.Contains(ChildrenUIDs[ChildIndex]))
+		{
+			TUniquePtr<FBaseNode>& Node = Nodes.FindChecked(ChildrenUIDs[ChildIndex]);
+			return Node.Get();
+		}
+
+		return nullptr;
+	}
+
+	/** Flat List of the nodes. Since the nodes are variable size, we store a pointer. */
+	TMap<FNodeUniqueID, TUniquePtr<FBaseNode> > Nodes;
+};
+
+}
diff --git a/Engine/Source/Runtime/Interchange/Core/Public/Nodes/InterchangeBaseNodeContainerAdapter.h b/Engine/Source/Runtime/Interchange/Core/Public/Nodes/InterchangeBaseNodeContainerAdapter.h
new file mode 100644
--- /dev/null
+++ b/Engine/Source/Runtime/Interchange/Core/Public/Nodes/InterchangeBaseNodeContainerAdapter.h
@@ -0,0 +1,378 @@
+// Copyright Epic Games, Inc. All Rights Reserved.
+
+#pragma once
+
+#include "CoreMinimal.h"
+#include "Templates/SharedPointer.h"
+#include "UObject/Object.h"
+#include "UObject/ObjectMacros.h"
+#include "Nodes/BaseNode.h"
+#include "Nodes/BaseNodeContainer.h"
+
+#include "InterchangeBaseNodeContainerAdapter.generated.h"
+
+/**
+ * Interchange node container adapter is a helper class to do some operations on a FBaseNodeContainer. This adapter is available in python and blueprint.
+ *
+ */
+UCLASS(BlueprintType)
+class INTERCHANGECORE_API UInterchangeBaseNodeContainerAdapter : public UObject
+{
+	GENERATED_BODY()
+public:
+	
+	/**
+	 * Set the Container the adapter will work on.
+	 *
+	 * @param InBaseNodeContainer - Pointer on a FBaseNodeContainer, the adapter will change the container directly.
+	 */
+	void SetBaseNodeContainer(Interchange::FBaseNodeContainer* InBaseNodeContainer)
+	{
+		BaseNodeContainer = InBaseNodeContainer;
+	}
+	
+	/**
+	 * Return all root node(the one that do not have nay parent) contain in the BaseNodeContainer control by this adapter
+	 *
+	 * @param Roots - The array where we store the result.
+	 */
+	UFUNCTION(BlueprintCallable, Category = "Interchange | Node Container")
+	bool GetRoots(TArray<FName>& Roots) const
+	{
+		if (!BaseNodeContainer)
+		{
+			return false;
+		}
+		BaseNodeContainer->GetRoots(Roots);
+		return true;
+	}
+
+	/** Return the first root node(the one that do not have nay parent) we find in the specified BaseNodeContainer */
+	UFUNCTION(BlueprintCallable, Category = "Interchange | Node Container")
+	FName GetFirstRoot() const
+	{
+		if (!BaseNodeContainer)
+		{
+			return NAME_None;
+		}
+		TArray<Interchange::FNodeUniqueID> Roots;
+		BaseNodeContainer->GetRoots(Roots);
+		if (Roots.Num() > 0)
+		{
+			return Roots[0];
+		}
+		return NAME_None;
+	}
+
+	/**
+	 * Return the parent of the specified node
+	 *
+	 * @param NodeIdentity - The specified node we query the parent
+	 * @return the parent UID, if there is no parent it return NAME_None
+	 */
+	UFUNCTION(BlueprintCallable, Category = "Interchange | Node Container")
+	FName GetNodeParent(const FName NodeIdentity) const
+	{
+		if (!BaseNodeContainer)
+		{
+			return NAME_None;
+		}
+		Interchange::FBaseNode* BaseNode = BaseNodeContainer->GetNode(NodeIdentity);
+		if (!BaseNode)
+		{
+			return NAME_None;
+		}
+		return BaseNode->GetParentUID();
+	}
+
+	/**
+	 * Find all the children of the specified node
+	 *
+	 * @param NodeIdentity - The specified node we query the child
+	 * @param OutNodeChilds - The array of children where we store the result
+	 */
+	UFUNCTION(BlueprintCallable, Category = "Interchange | Node Container")
+	void GetNodeChild(const FName NodeIdentity, TArray<FName>& OutNodeChilds) const
+	{
+		if (!BaseNodeContainer)
+		{
+			return;
+		}
+		OutNodeChilds = BaseNodeContainer->GetNodeChildrenUIDs(NodeIdentity);
+		return;
+	}
+
+	/**
+	 * Set the node parent.
+	 *
+	 * @param NodeIdentity - The node we want to set the parent
+	 * @param ParentNodeIdentity - The parent node
+	 * @return true if the parent was properly set, false otherwise
+	 */
+	UFUNCTION(BlueprintCallable, Category = "Interchange | Node Container")
+	bool SetNodeParent(const FName NodeIdentity, const FName ParentNodeIdentity)
+	{
+		if (!BaseNodeContainer)
+		{
+			return false;
+		}
+		
+		return BaseNodeContainer->SetNodeParentUID(NodeIdentity, ParentNodeIdentity);
+	}
+
+	/**
+	 * Return true if the node should be imported/exported.
+	 *
+	 * @param NodeIdentity - The node we want to query the enabled state
+	 * @return true if the node should be imported/exported.
+	 */
+	UFUNCTION(BlueprintCallable, Category = "Interchange | Node Container")
+	bool IsNodeEnabled(const FName NodeIdentity) const
+	{
+		if (!BaseNodeContainer)
+		{
+			return false;
+		}
+		Interchange::FBaseNode* BaseNode = BaseNodeContainer->GetNode(NodeIdentity);
+		if (!BaseNode)
+		{
+			return false;
+		}
+		return BaseNode->IsEnabled();
+	}
+
+	/**
+	 * Set the node imported/exported state.
+	 *
+	 * @param NodeIdentity - The node we want to set the enabled state
+	 * @param bIsEnabled - The enabled state to set
+	 * @return true if the node enabled state was set properly. false otherwise
+	 */
+	UFUNCTION(BlueprintCallable, Category = "Interchange | Node Container")
+	bool SetNodeEnabled(const FName NodeIdentity, const bool bIsEnabled)
+	{
+		if (!BaseNodeContainer)
+		{
+			return false;
+		}
+		Interchange::FBaseNode* BaseNode = BaseNodeContainer->GetNode(NodeIdentity);
+		if (!BaseNode)
+		{
+			return false;
+		}
+		return BaseNode->SetEnabled(bIsEnabled);
+	}
+
+	/**
+	 * Fill the OutValue parameter with the content of the attribute point by AttributeName
+	 *
+	 * @param NodeIdentity - The node we want to query the attribute
+	 * @param AttributeName - The attribute name use to retrieve the actual attribute value
+	 * @param OutValue - We store the attribute value into this parameter
+	 * @return true if we successfully retrieve the attribute value, false otherwise.
+	 */
+	template<typename T>
+	bool GetNodeAttributeValue(const FName NodeIdentity, const FName AttributeName, T& OutValue)
+	{
+		if (!BaseNodeContainer)
+		{
+			return false;
+		}
+		Interchange::FBaseNode* BaseNode = BaseNodeContainer->GetNode(NodeIdentity);
+		if (!BaseNode)
+		{
+			return false;
+		}
+		Interchange::FAttributeKey AttributeKey(AttributeName);
+		Interchange::FAttributeStorage::TAttributeHandle<T> AttributeHandle = BaseNode->GetAttributeHandle<T>(AttributeKey);
+		AttributeHandle.Get(OutValue);
+		return true;
+	}
+
+	/**
+	 * Set the Value parameter in the attribute point by AttributeName
+	 *
+	 * @param NodeIdentity - The node we want to query the attribute
+	 * @param AttributeName - The attribute name use to retrieve the actual attribute we want to set the value
+	 * @param Value - The value we want to store
+	 * @return true if we successfully set the attribute value, false otherwise.
+	 */
+	template<typename T>
+	bool SetNodeAttributeValue(const FName NodeIdentity, const FName AttributeName, const T& Value)
+	{
+		if (!BaseNodeContainer)
+		{
+			return false;
+		}
+		Interchange::FBaseNode* BaseNode = BaseNodeContainer->GetNode(NodeIdentity);
+		if (!BaseNode)
+		{
+			return false;
+		}
+		Interchange::FAttributeKey AttributeKey(AttributeName);
+		Interchange::FAttributeStorage::TAttributeHandle<T> AttributeHandle = BaseNode->GetAttributeHandle<T>(AttributeKey);
+		if(AttributeHandle.IsValid())
+		{
+			AttributeHandle.Set(Value);
+		}
+		else
+		{
+			BaseNode->RegisterAttribute<T>(AttributeKey, Value);
+		}
+
+		return true;
+	}
+
+	//////////////////////////////////////////////////////////////////////////
+	//Attribute getter interface, blueprint callable for all supported type
+	
+	UFUNCTION(BlueprintCallable, Category = "Interchange | Node Attribute")
+	bool GetNodeAttributeValueBool(const FName NodeIdentity, const FName AttributeName, bool& OutValue);
+	UFUNCTION(BlueprintCallable, Category = "Interchange | Node Attribute")
+	bool GetNodeAttributeValueFloat(const FName NodeIdentity, const FName AttributeName, float& OutValue);
+	//type double not supported in blueprint
+ 	bool GetNodeAttributeValueDouble(const FName NodeIdentity, const FName AttributeName, double& OutValue);
+	//type int8 not supported in blueprint
+ 	bool GetNodeAttributeValueInt8(const FName NodeIdentity, const FName AttributeName, int8& OutValue);
+	//type int16 not supported in blueprint
+ 	bool GetNodeAttributeValueInt16(const FName NodeIdentity, const FName AttributeName, int16& OutValue);
+	UFUNCTION(BlueprintCallable, Category = "Interchange | Node Attribute")
+	bool GetNodeAttributeValueInt32(const FName NodeIdentity, const FName AttributeName, int32& OutValue);
+	UFUNCTION(BlueprintCallable, Category = "Interchange | Node Attribute")
+	bool GetNodeAttributeValueInt64(const FName NodeIdentity, const FName AttributeName, int64& OutValue);
+	UFUNCTION(BlueprintCallable, Category = "Interchange | Node Attribute")
+	bool GetNodeAttributeValueUint8(const FName NodeIdentity, const FName AttributeName, uint8& OutValue);
+	//type uint16 not supported in blueprint
+ 	bool GetNodeAttributeValueUint16(const FName NodeIdentity, const FName AttributeName, uint16& OutValue);
+	//type uint32 not supported in blueprint
+ 	bool GetNodeAttributeValueUint32(const FName NodeIdentity, const FName AttributeName, uint32& OutValue);
+	//type uint64 not supported in blueprint
+ 	bool GetNodeAttributeValueUint64(const FName NodeIdentity, const FName AttributeName, uint64& OutValue);
+	UFUNCTION(BlueprintCallable, Category = "Interchange | Node Attribute")
+	bool GetNodeAttributeValueVector(const FName NodeIdentity, const FName AttributeName, FVector& OutValue);
+	UFUNCTION(BlueprintCallable, Category = "Interchange | Node Attribute")
+	bool GetNodeAttributeValueVector2D(const FName NodeIdentity, const FName AttributeName, FVector2D& OutValue);
+	UFUNCTION(BlueprintCallable, Category = "Interchange | Node Attribute")
+	bool GetNodeAttributeValueMatrix(const FName NodeIdentity, const FName AttributeName, FMatrix& OutValue);
+	UFUNCTION(BlueprintCallable, Category = "Interchange | Node Attribute")
+	bool GetNodeAttributeValueBox(const FName NodeIdentity, const FName AttributeName, FBox& OutValue);
+	UFUNCTION(BlueprintCallable, Category = "Interchange | Node Attribute")
+	bool GetNodeAttributeValueBoxSphereBound(const FName NodeIdentity, const FName AttributeName, FBoxSphereBounds& OutValue);
+	UFUNCTION(BlueprintCallable, Category = "Interchange | Node Attribute")
+	bool GetNodeAttributeValueArray(const FName NodeIdentity, const FName AttributeName, TArray<uint8>& OutValue);
+	//type TArray64<uint8> not supported in blueprint
+ 	bool GetNodeAttributeValueArray64(const FName NodeIdentity, const FName AttributeName, TArray64<uint8>& OutValue);
+	UFUNCTION(BlueprintCallable, Category = "Interchange | Node Attribute")
+	bool GetNodeAttributeValueColor(const FName NodeIdentity, const FName AttributeName, FColor& OutValue);
+	UFUNCTION(BlueprintCallable, Category = "Interchange | Node Attribute")
+	bool GetNodeAttributeValueDateTime(const FName NodeIdentity, const FName AttributeName, FDateTime& OutValue);
+	//We use a uint8 to get the enum value, blueprint only support enum of uint8 type
+	UFUNCTION(BlueprintCallable, Category = "Interchange | Node Attribute")
+	bool GetNodeAttributeValueEnum(const FName NodeIdentity, const FName AttributeName, uint8& OutValue);
+	UFUNCTION(BlueprintCallable, Category = "Interchange | Node Attribute")
+	bool GetNodeAttributeValueGuid(const FName NodeIdentity, const FName AttributeName, FGuid& OutValue);
+	UFUNCTION(BlueprintCallable, Category = "Interchange | Node Attribute")
+	bool GetNodeAttributeValueIntPoint(const FName NodeIdentity, const FName AttributeName, FIntPoint& OutValue);
+	UFUNCTION(BlueprintCallable, Category = "Interchange | Node Attribute")
+	bool GetNodeAttributeValueIntVecor(const FName NodeIdentity, const FName AttributeName, FIntVector& OutValue);
+	UFUNCTION(BlueprintCallable, Category = "Interchange | Node Attribute")
+	bool GetNodeAttributeValueLinearColor(const FName NodeIdentity, const FName AttributeName, FLinearColor& OutValue);
+	UFUNCTION(BlueprintCallable, Category = "Interchange | Node Attribute")
+	bool GetNodeAttributeValuePlane(const FName NodeIdentity, const FName AttributeName, FPlane& OutValue);
+	UFUNCTION(BlueprintCallable, Category = "Interchange | Node Attribute")
+	bool GetNodeAttributeValueQuat(const FName NodeIdentity, const FName AttributeName, FQuat& OutValue);
+	UFUNCTION(BlueprintCallable, Category = "Interchange | Node Attribute")
+	bool GetNodeAttributeValueName(const FName NodeIdentity, const FName AttributeName, FName& OutValue);
+	UFUNCTION(BlueprintCallable, Category = "Interchange | Node Attribute")
+	bool GetNodeAttributeValueRandomStream(const FName NodeIdentity, const FName AttributeName, FRandomStream& OutValue);
+	UFUNCTION(BlueprintCallable, Category = "Interchange | Node Attribute")
+	bool GetNodeAttributeValueRotator(const FName NodeIdentity, const FName AttributeName, FRotator& OutValue);
+	UFUNCTION(BlueprintCallable, Category = "Interchange | Node Attribute")
+	bool GetNodeAttributeValueString(const FName NodeIdentity, const FName AttributeName, FString& OutValue);
+	UFUNCTION(BlueprintCallable, Category = "Interchange | Node Attribute")
+	bool GetNodeAttributeValueTimespan(const FName NodeIdentity, const FName AttributeName, FTimespan& OutValue);
+	UFUNCTION(BlueprintCallable, Category = "Interchange | Node Attribute")
+	bool GetNodeAttributeValueTransform(const FName NodeIdentity, const FName AttributeName, FTransform& OutValue);
+	UFUNCTION(BlueprintCallable, Category = "Interchange | Node Attribute")
+	bool GetNodeAttributeValueTwoVectors(const FName NodeIdentity, const FName AttributeName, FTwoVectors& OutValue);
+	UFUNCTION(BlueprintCallable, Category = "Interchange | Node Attribute")
+	bool GetNodeAttributeValueVector4(const FName NodeIdentity, const FName AttributeName, FVector4& OutValue);
+
+	//////////////////////////////////////////////////////////////////////////
+	//Attribute setter interface, blueprint callable for all supported type
+	
+	UFUNCTION(BlueprintCallable, Category = "Interchange | Node Attribute")
+	bool SetNodeAttributeValueBool(const FName NodeIdentity, const FName AttributeName, const bool Value);
+	UFUNCTION(BlueprintCallable, Category = "Interchange | Node Attribute")
+	bool SetNodeAttributeValueFloat(const FName NodeIdentity, const FName AttributeName, const float Value);
+	//type double not supported in blueprint
+ 	bool SetNodeAttributeValueDouble(const FName NodeIdentity, const FName AttributeName, const double Value);
+	//type int8 not supported in blueprint
+ 	bool SetNodeAttributeValueInt8(const FName NodeIdentity, const FName AttributeName, const int8 Value);
+	//type int16 not supported in blueprint
+ 	bool SetNodeAttributeValueInt16(const FName NodeIdentity, const FName AttributeName, const int16 Value);
+	UFUNCTION(BlueprintCallable, Category = "Interchange | Node Attribute")
+	bool SetNodeAttributeValueInt32(const FName NodeIdentity, const FName AttributeName, const int32 Value);
+	UFUNCTION(BlueprintCallable, Category = "Interchange | Node Attribute")
+	bool SetNodeAttributeValueInt64(const FName NodeIdentity, const FName AttributeName, const int64 Value);
+	UFUNCTION(BlueprintCallable, Category = "Interchange | Node Attribute")
+	bool SetNodeAttributeValueUint8(const FName NodeIdentity, const FName AttributeName, const uint8 Value);
+	//type uint16 not supported in blueprint
+ 	bool SetNodeAttributeValueUint16(const FName NodeIdentity, const FName AttributeName, const uint16 Value);
+	//type uint32 not supported in blueprint
+ 	bool SetNodeAttributeValueUint32(const FName NodeIdentity, const FName AttributeName, const uint32 Value);
+	//type uint64 not supported in blueprint
+ 	bool SetNodeAttributeValueUint64(const FName NodeIdentity, const FName AttributeName, const uint64 Value);
+	UFUNCTION(BlueprintCallable, Category = "Interchange | Node Attribute")
+	bool SetNodeAttributeValueVector(const FName NodeIdentity, const FName AttributeName, const FVector& Value);
+	UFUNCTION(BlueprintCallable, Category = "Interchange | Node Attribute")
+	bool SetNodeAttributeValueVector2D(const FName NodeIdentity, const FName AttributeName, const FVector2D& Value);
+	UFUNCTION(BlueprintCallable, Category = "Interchange | Node Attribute")
+	bool SetNodeAttributeValueMatrix(const FName NodeIdentity, const FName AttributeName, const FMatrix& Value);
+	UFUNCTION(BlueprintCallable, Category = "Interchange | Node Attribute")
+	bool SetNodeAttributeValueBox(const FName NodeIdentity, const FName AttributeName, const FBox& Value);
+	UFUNCTION(BlueprintCallable, Category = "Interchange | Node Attribute")
+	bool SetNodeAttributeValueBoxSphereBound(const FName NodeIdentity, const FName AttributeName, const FBoxSphereBounds& Value);
+	UFUNCTION(BlueprintCallable, Category = "Interchange | Node Attribute")
+	bool SetNodeAttributeValueArray(const FName NodeIdentity, const FName AttributeName, const TArray<uint8>& Value);
+	//type TArray64<uint8> not supported in blueprint
+ 	bool SetNodeAttributeValueArray64(const FName NodeIdentity, const FName AttributeName, const TArray64<uint8>& Value);
+	UFUNCTION(BlueprintCallable, Category = "Interchange | Node Attribute")
+	bool SetNodeAttributeValueColor(const FName NodeIdentity, const FName AttributeName, const FColor& Value);
+	UFUNCTION(BlueprintCallable, Category = "Interchange | Node Attribute")
+	bool SetNodeAttributeValueDateTime(const FName NodeIdentity, const FName AttributeName, const FDateTime& Value);
+	//We use a uint8 to get the enum value, blueprint only support enum of uint8 type
+	UFUNCTION(BlueprintCallable, Category = "Interchange | Node Attribute")
+	bool SetNodeAttributeValueEnum(const FName NodeIdentity, const FName AttributeName, const uint8 Value);
+	UFUNCTION(BlueprintCallable, Category = "Interchange | Node Attribute")
+	bool SetNodeAttributeValueGuid(const FName NodeIdentity, const FName AttributeName, const FGuid& Value);
+	UFUNCTION(BlueprintCallable, Category = "Interchange | Node Attribute")
+	bool SetNodeAttributeValueIntPoint(const FName NodeIdentity, const FName AttributeName, const FIntPoint& Value);
+	UFUNCTION(BlueprintCallable, Category = "Interchange | Node Attribute")
+	bool SetNodeAttributeValueIntVecor(const FName NodeIdentity, const FName AttributeName, const FIntVector& Value);
+	UFUNCTION(BlueprintCallable, Category = "Interchange | Node Attribute")
+	bool SetNodeAttributeValueLinearColor(const FName NodeIdentity, const FName AttributeName, const FLinearColor& Value);
+	UFUNCTION(BlueprintCallable, Category = "Interchange | Node Attribute")
+	bool SetNodeAttributeValuePlane(const FName NodeIdentity, const FName AttributeName, const FPlane& Value);
+	UFUNCTION(BlueprintCallable, Category = "Interchange | Node Attribute")
+	bool SetNodeAttributeValueQuat(const FName NodeIdentity, const FName AttributeName, const FQuat& Value);
+	UFUNCTION(BlueprintCallable, Category = "Interchange | Node Attribute")
+	bool SetNodeAttributeValueName(const FName NodeIdentity, const FName AttributeName, const FName Value);
+	UFUNCTION(BlueprintCallable, Category = "Interchange | Node Attribute")
+	bool SetNodeAttributeValueRandomStream(const FName NodeIdentity, const FName AttributeName, const FRandomStream& Value);
+	UFUNCTION(BlueprintCallable, Category = "Interchange | Node Attribute")
+	bool SetNodeAttributeValueRotator(const FName NodeIdentity, const FName AttributeName, const FRotator& Value);
+	UFUNCTION(BlueprintCallable, Category = "Interchange | Node Attribute")
+	bool SetNodeAttributeValueString(const FName NodeIdentity, const FName AttributeName, const FString& Value);
+	UFUNCTION(BlueprintCallable, Category = "Interchange | Node Attribute")
+	bool SetNodeAttributeValueTimespan(const FName NodeIdentity, const FName AttributeName, const FTimespan& Value);
+	UFUNCTION(BlueprintCallable, Category = "Interchange | Node Attribute")
+	bool SetNodeAttributeValueTransform(const FName NodeIdentity, const FName AttributeName, const FTransform& Value);
+	UFUNCTION(BlueprintCallable, Category = "Interchange | Node Attribute")
+	bool SetNodeAttributeValueTwoVectors(const FName NodeIdentity, const FName AttributeName, const FTwoVectors& Value);
+	UFUNCTION(BlueprintCallable, Category = "Interchange | Node Attribute")
+	bool SetNodeAttributeValueVector4(const FName NodeIdentity, const FName AttributeName, const FVector4& Value);
+
+private:
+	Interchange::FBaseNodeContainer* BaseNodeContainer = nullptr;
+};
diff --git a/Engine/Source/Runtime/Interchange/Core/Public/PackageUtils/PackageUtils.h b/Engine/Source/Runtime/Interchange/Core/Public/PackageUtils/PackageUtils.h
new file mode 100644
--- /dev/null
+++ b/Engine/Source/Runtime/Interchange/Core/Public/PackageUtils/PackageUtils.h
@@ -0,0 +1,54 @@
+// Copyright Epic Games, Inc. All Rights Reserved.
+
+#pragma once
+
+#include "CoreMinimal.h"
+#include "Misc/PackageName.h"
+#include "Misc/Paths.h"
+
+//Interchange namespace
+namespace Interchange
+{
+	class FPackageUtils
+	{
+	public:
+		static bool IsMapPackageAsset(const FString& ObjectPath)
+		{
+			FString MapFilePath;
+			return FPackageUtils::IsMapPackageAsset(ObjectPath, MapFilePath);
+		}
+
+		static bool IsMapPackageAsset(const FString& ObjectPath, FString& MapFilePath)
+		{
+			const FString PackageName = FPackageUtils::ExtractPackageName(ObjectPath);
+			if ( PackageName.Len() > 0 )
+			{
+				FString PackagePath;
+				if ( FPackageName::DoesPackageExist(PackageName, NULL, &PackagePath) )
+				{
+					const FString FileExtension = FPaths::GetExtension(PackagePath, true);
+					if ( FileExtension == FPackageName::GetMapPackageExtension() )
+					{
+						MapFilePath = PackagePath;
+						return true;
+					}
+				}
+			}
+
+			return false;
+		}
+
+		static FString ExtractPackageName(const FString& ObjectPath)
+		{
+			// To find the package name in an object path we need to find the path left of the FIRST delimiter.
+			// Assets like BSPs, lightmaps etc. can have multiple '.' delimiters.
+			const int32 PackageDelimiterPos = ObjectPath.Find(TEXT("."), ESearchCase::CaseSensitive, ESearchDir::FromStart);
+			if ( PackageDelimiterPos != INDEX_NONE )
+			{
+				return ObjectPath.Left(PackageDelimiterPos);
+			}
+
+			return ObjectPath;
+		}
+	};
+} //End Interchange namespace
\ No newline at end of file
diff --git a/Engine/Source/Runtime/Interchange/Core/Public/Types/AttributeStorage.h b/Engine/Source/Runtime/Interchange/Core/Public/Types/AttributeStorage.h
new file mode 100644
--- /dev/null
+++ b/Engine/Source/Runtime/Interchange/Core/Public/Types/AttributeStorage.h
@@ -0,0 +1,1318 @@
+// Copyright Epic Games, Inc. All Rights Reserved.
+
+#pragma once
+
+#include "Containers/SortedMap.h"
+#include "Containers/UnrealString.h"
+#include "CoreMinimal.h"
+#include "Math/DualQuat.h"
+#include "Math/RandomStream.h"
+#include "Misc/FrameRate.h"
+#include "Misc/NetworkGuid.h"
+
+
+
+//Interchange namespace
+namespace Interchange
+{
+
+	struct FAttributeKey
+	{
+		FName Key = NAME_None;
+
+		FAttributeKey() = default;
+		
+		FAttributeKey(const FAttributeKey& Other)
+		{
+			Key = Other.Key;
+		}
+
+		explicit FAttributeKey(const FName& Other)
+		{
+			Key = Other;
+		}
+
+		explicit FAttributeKey(const FString& Other)
+		{
+			Key = Other.IsEmpty() ? NAME_None : FName(Other);
+		}
+
+		explicit FAttributeKey(const FText& Other)
+		{
+			Key = Other.IsEmpty() ? NAME_None : FName(Other.ToString());
+		}
+
+		explicit FAttributeKey(const TCHAR* Other)
+		{
+			const FString Helper(Other);
+			Key = Helper.IsEmpty() ? NAME_None : FName(Helper);
+		}
+
+		FORCEINLINE const TCHAR* operator*() const
+		{
+			return *(Key.ToString());
+		}
+
+		FORCEINLINE FAttributeKey& operator=(const FAttributeKey& Other)
+		{
+			Key = Other.Key;
+			return *this;
+		}
+
+		FORCEINLINE FAttributeKey& operator=(const FName& Other)
+		{
+			Key = Other;
+			return *this;
+		}
+
+		FORCEINLINE FAttributeKey& operator=(const FString& Other)
+		{
+			Key = Other.IsEmpty() ? NAME_None : FName(Other);
+			return *this;
+		}
+
+		FORCEINLINE FAttributeKey& operator=(const FText& Other)
+		{
+			Key = Other.IsEmpty() ? NAME_None : FName(Other.ToString());
+			return *this;
+		}
+
+		FORCEINLINE FAttributeKey& operator=(const TCHAR* Other)
+		{
+			const FString Helper(Other);
+			Key = Helper.IsEmpty() ? NAME_None : FName(Helper);
+			return *this;
+		}
+		
+		FORCEINLINE bool operator==(const FAttributeKey& Other) const
+		{
+			return Key == Other.Key;
+		}
+
+		FORCEINLINE bool operator!=(const FAttributeKey& Other) const
+		{
+			return Key != Other.Key;
+		}
+
+		FORCEINLINE bool operator<(const FAttributeKey& Other) const
+		{
+			return Key.LexicalLess(Other.Key);
+		}
+
+		FORCEINLINE bool operator<=(const FAttributeKey& Other) const
+		{
+			return Key.Compare(Other.Key) <= 0;
+		}
+
+		FORCEINLINE bool operator>(const FAttributeKey& Other) const
+		{
+			return Other.Key.LexicalLess(Key);
+		}
+
+		FORCEINLINE bool operator>=(const FAttributeKey& Other) const
+		{
+			return Key.Compare(Other.Key) >= 0;
+		}
+
+		friend FArchive& operator<<(FArchive& Ar, FAttributeKey& AttributeKey)
+		{
+			FString KeyString = AttributeKey.Key.ToString();
+			Ar << KeyString;
+			if(Ar.IsLoading())
+			{
+				AttributeKey.Key = FName(KeyString);
+			}
+			return Ar;
+		}
+	};
+
+	FORCEINLINE uint32 GetTypeHash(const FAttributeKey& AttributeKey)
+	{
+		return GetTypeHash(AttributeKey.Key);
+	}
+	
+	/**
+	 * Enumerates the built-in types that can be stored in instances of FAttributeStorage.
+	 * We cannot change the value of a type to make sure the serialization of old asset is always working
+	 */
+	enum class EAttributeTypes : int32
+	{
+		None				= 0,
+		Bool				= 1,
+		Box					= 2,
+		BoxSphereBounds		= 3,
+		ByteArray			= 4,
+		ByteArray64			= 5,
+		Color				= 6,
+		DateTime			= 7,
+		Double				= 8,
+		Enum				= 9,
+		Float				= 10,
+		Guid				= 11,
+		Int8				= 12,
+		Int16				= 13,
+		Int32				= 14,
+		Int64				= 15,
+		IntRect				= 16,
+		LinearColor			= 17,
+		Matrix				= 18,
+		Name				= 19,
+		Plane				= 20,
+		Quat				= 21,
+		RandomStream		= 22,
+		Rotator				= 23,
+		String				= 24,
+		Timespan			= 25,
+		Transform			= 26,
+		TwoVectors			= 27,
+		UInt8				= 28,
+		UInt16				= 29,
+		UInt32				= 30,
+		UInt64				= 31,
+		Vector				= 32,
+		Vector2d			= 33,
+		Vector4				= 34,
+		IntPoint			= 35,
+		IntVector			= 36,
+		NetworkGUID			= 37,
+		Vector2DHalf		= 38,
+		Float16				= 39,
+		CapsuleShape		= 40,
+		DualQuat			= 41,
+		OrientedBox			= 42,
+		Ray					= 43,
+		Sphere				= 44,
+		FrameNumber			= 45,
+		FrameRate			= 46,
+		FrameTime			= 47,
+	};
+
+
+	/**
+	 * Stub for attribute type traits.
+	 *
+	 * Actual type traits need to be declared through template specialization for custom
+	 * data types that are to be used internally by FAttributeStorage. Traits for the most commonly used built-in
+	 * types are declared below.
+	 *
+	 * Complex types, such as structures and classes can be serialized into a byte array
+	 * and then assigned to an attribute. Note that you will be responsible for ensuring
+	 * correct byte ordering when serializing those types.
+	 *
+	 * @param T The type to be used in FAttributeStorage.
+	 */
+	template<typename T> struct TAttributeTypeTraits
+	{
+		static CONSTEXPR EAttributeTypes GetType()
+		{
+			static_assert(!sizeof(T), "Attribute type trait must be specialized for this type.");
+			return EAttributeTypes::None;
+		}
+	};
+
+	/**
+	 * Enum to return complete status of a storage operation
+	 * It support success with additional information
+	 * It support multiple errors
+	 */
+	enum class EAttributeStorageResult : uint64
+	{
+		None											= 0x0,
+		//success result
+		Operation_Success								= 0x1,
+		
+
+		//Operation error result from here
+
+		//The type of the value was not matching the existing type. We cannot change the type of an existing attribute
+		Operation_Error_WrongType						= (0x1 << 20),
+		//The size of the value is different from the existing size. We cannot have a different size
+		Operation_Error_WrongSize						= (0x1 << 21),
+		//The AttributeAllocationTable has an attribute that the offset is no valid in the storage
+		Operation_Error_AttributeAllocationCorrupted	= (0x1 << 22),
+		//We cannot find the specified key
+		Operation_Error_CannotFoundKey					= (0x1 << 23),
+		//There was an error when removing an attribute from the AttributeAllocationTable. The TArray remove has failed
+		Operation_Error_CannotRemoveAttribute			= (0x1 << 24),
+		//We try to override an attribute but the specified options do not allow override.
+		Operation_Error_CannotOverrideAttribute			= (0x1 << 25),
+		//The type pass in parameter is not supported, add new type to support in EAttributeTypes
+        Operation_Error_UnsupportedType					= (0x1 << 26),
+		//The storage is invalid (nullptr).
+		Operation_Error_InvalidStorage					= (0x1 << 27),
+		//Cannot get a valid value data pointer.
+		Operation_Error_InvalidMultiSizeValueData		= (0x1 << 28),
+	};
+
+	ENUM_CLASS_FLAGS(EAttributeStorageResult)
+
+	/**
+	 * Helper function to interpret storage result
+	 * @return true if the Result contain at least one of the RefResult flags
+	 */
+	FORCEINLINE bool HasAttributeStorageResult(const EAttributeStorageResult Result, const EAttributeStorageResult RefResult)
+	{
+		return ((Result & RefResult) != EAttributeStorageResult::None);
+	}
+	
+	/**
+	 * Helper function to determine if the storage result is success
+	 * @return true if result contain Operation_Success
+	 */
+	FORCEINLINE bool IsAttributeStorageResultSuccess(const EAttributeStorageResult Result)
+	{
+		return HasAttributeStorageResult(Result, EAttributeStorageResult::Operation_Success);
+	}
+
+	/**
+	 * Helper function to transform an operation result into a LOG.
+	 * @param Result - The result we want to output a log for
+	 * @param OperationName - The operation name that end up with the speicified Result.
+	 * @param AttributeKey - The attribute we did apply the operation on.
+	 */
+	INTERCHANGECORE_API void LogAttributeStorageErrors(const EAttributeStorageResult Result, const FString OperationName, const FAttributeKey AttributeKey );
+
+
+	/**
+	 * Enum to pass options when we add an attribute
+	 */
+	enum class EAttributeStorageAddOptions : uint32
+	{
+		None					= 0x0,
+		Option_Override			= 0x1,
+		//allow the AddAttribute to override the value if it exist.
+	};
+
+	ENUM_CLASS_FLAGS(EAttributeStorageAddOptions)
+
+	/**
+	 * Helper function to interpret storage add options
+	 */
+	FORCEINLINE bool HasAttributeStorageAddOption(const EAttributeStorageAddOptions Options, const EAttributeStorageAddOptions RefOptions)
+	{
+		return ((Options & RefOptions) != EAttributeStorageAddOptions::None);
+	}
+
+	/**
+	 * Enumerates the attribute properties. Those properties affect how the attribute are stored or what they are use for
+	 */
+	enum class EAttributeProperty : uint32
+	{
+		None						= 0x0,
+		NoHash						= 0x1, /*No hash attribute will not be part of the hash result when calling GetStorageHash*/
+	};
+
+	ENUM_CLASS_FLAGS(EAttributeProperty)
+
+	/**
+	 * Helper function to interpret attribute property
+	 */
+	FORCEINLINE bool HasAttributeProperty(const EAttributeProperty PropertyA, const EAttributeProperty PropertyB)
+	{
+		return ((PropertyA & PropertyB) != EAttributeProperty::None);
+	}
+
+	/**
+	 * This is a helper class to specialize template function inside a class, it allow specialize template to not be static
+	 */
+	template<typename T>
+	struct TSpecializeType
+	{
+		typedef T Type;
+	};
+
+	
+	/**
+	 * This class is a Key/Value storage inside a TArray64<uint8>
+	 * The keys are of type FAttributeKey which is a FString, each key is unique and have only one value.
+	 * The value can be of any type contain in EAttributeTypes.
+	 * 
+	 * @notes
+	 * The storage is multi thread safe, it use a mutex to lock the storage for every read/write operation.
+	 * The hash of the storage is deterministic because it sort the attributes before calculating the hash.
+	 */
+	class INTERCHANGECORE_API FAttributeStorage
+	{
+	public:
+		/**
+		 * Class to Get/Set an attribute of the storage 
+		 */
+		template<typename T>
+		class TAttributeHandle
+		{
+		public:
+			TAttributeHandle()
+			: AttributeStorage(nullptr)
+			, Key(NAME_None)
+			{}
+
+			/**
+			* Return true if the storage contain a valid attribute key, false otherwise. 
+			*/
+			FORCEINLINE bool IsValid() const
+			{
+				const EAttributeTypes ValueType = TAttributeTypeTraits<T>::GetType();
+				return (AttributeStorage && AttributeStorage->ContainAttribute(Key) && AttributeStorage->GetAttributeType(Key) == ValueType);
+			}
+			
+			EAttributeStorageResult Get(T& Value) const
+			{
+				if (AttributeStorage)
+				{
+					return AttributeStorage->GetAttribute<T>(Key, Value);
+				}
+				return EAttributeStorageResult::Operation_Error_InvalidStorage;
+			}
+
+			EAttributeStorageResult Set(const T& Value)
+			{
+				if (AttributeStorage)
+				{
+					return AttributeStorage->SetAttribute<T>(Key, Value);
+				}
+				return EAttributeStorageResult::Operation_Error_InvalidStorage;
+			}
+		protected:
+			class FAttributeStorage* AttributeStorage;
+			FAttributeKey Key;
+			
+			TAttributeHandle(const FAttributeKey& InKey, const FAttributeStorage* InAttributeStorage)
+			{
+				AttributeStorage = const_cast<FAttributeStorage*>(InAttributeStorage);
+				Key = InKey;
+				//Log some error in case something is wrong
+				const EAttributeTypes ValueType = TAttributeTypeTraits<T>::GetType();
+				if (ValueType == EAttributeTypes::None)
+				{
+					//Value Type is unsupported
+					LogAttributeStorageErrors(EAttributeStorageResult::Operation_Error_UnsupportedType, TEXT("GetAttributeHandle"), Key);
+				}
+				//Look for storage validity
+				if (AttributeStorage == nullptr)
+				{
+					//Storage is null
+					LogAttributeStorageErrors(EAttributeStorageResult::Operation_Error_InvalidStorage, TEXT("GetAttributeHandle"), Key);
+				}
+				else
+				{
+
+					if (!AttributeStorage->ContainAttribute(Key))
+					{
+						//Storage do not contain the key
+						LogAttributeStorageErrors(EAttributeStorageResult::Operation_Error_CannotFoundKey, TEXT("GetAttributeHandle"), Key);
+					}
+					if (AttributeStorage->GetAttributeType(Key) != ValueType)
+					{
+						//Value Type is different from the existing key
+						LogAttributeStorageErrors(EAttributeStorageResult::Operation_Error_WrongType, TEXT("GetAttributeHandle"), Key);
+					}
+				}
+			}
+
+			friend FAttributeStorage;
+		};
+
+		FAttributeStorage& operator=(const FAttributeStorage& Other);
+
+		/**
+		 * Register an attribute in the storage. Return success if the attribute was properly added, or there is an existing
+		 * attribute of the same type, return error otherwise.
+		 *
+		 * @Param ElementAttributeKey - is the storage key (the path) of the attribute
+		 * @Param DefaultValue - is the default value for the registered attribute
+		 *
+		 * @note Possible errors
+		 * - An attribute of a unsupported type see IsSupportedType
+		 * - Key exist with a different type
+		 * - Storage is corrupted
+		 */
+		template<typename T>
+		EAttributeStorageResult RegisterAttribute(const FAttributeKey& ElementAttributeKey, const T& DefaultValue, EAttributeProperty AttributeProperty = EAttributeProperty::None)
+		{
+			//Lock the storage
+			FScopeLock ScopeLock(&StorageMutex);
+
+			if (!IsSupportedType(DefaultValue))
+			{
+				return EAttributeStorageResult::Operation_Error_UnsupportedType;
+			}
+
+			const EAttributeTypes ValueType = TAttributeTypeTraits<T>::GetType();
+			const uint64 ValueSize = GetValueSize(DefaultValue);
+
+			FAttributeAllocationInfo* AttributeAllocationInfo = AttributeAllocationTable.Find(ElementAttributeKey);
+			if (AttributeAllocationInfo)
+			{
+				if (AttributeAllocationInfo->Type != ValueType)
+				{
+					return EAttributeStorageResult::Operation_Error_WrongType;
+				}
+				if (!AttributeStorage.IsValidIndex(AttributeAllocationInfo->Offset))
+				{
+					return EAttributeStorageResult::Operation_Error_AttributeAllocationCorrupted;
+				}
+			}
+			else
+			{
+				AttributeAllocationInfo = &AttributeAllocationTable.Add(ElementAttributeKey);
+				AttributeAllocationInfo->Type = ValueType;
+				AttributeAllocationInfo->Size = ValueSize;
+				AttributeAllocationInfo->Offset = AttributeStorage.AddZeroed(ValueSize);
+			}
+
+			//Force the specified attribute property
+			AttributeAllocationInfo->Property = AttributeProperty;
+
+			//Use template SetAttribute which support all the types
+			const EAttributeStorageResult Result = SetAttribute(AttributeAllocationInfo, DefaultValue);
+			if (!IsAttributeStorageResultSuccess(Result))
+			{
+				//An error occured, unregister the key from the storage
+				UnregisterAttribute(ElementAttributeKey);
+			}
+			return Result;
+		}
+
+		/**
+		 * Remove an attribute from the storage.
+		 *
+		 * @param ElementAttributeKey - is the storage key (the path) of the attribute to remove.
+		 *
+		 * @note Possible errors
+		 * - Key do not exist
+		 * - Internal storage structure remove error
+		 */
+		EAttributeStorageResult UnregisterAttribute(const FAttributeKey& ElementAttributeKey);
+
+		/**
+		 * Return an attribute handle for the specified attribute. This handle is compile type check and is use to Get and Set the attribute value type.
+		 * The function will assert if the key is missing or the type doesnt match the specified template type.
+		 *
+		 * @param ElementAttributeKey is the storage key (the path) of the attribute.
+		 */
+		template<typename T>
+		TAttributeHandle<T> GetAttributeHandle(const FAttributeKey& ElementAttributeKey) const
+		{
+			TAttributeHandle<T> AttributeHandle(ElementAttributeKey, this);
+			return AttributeHandle;
+		}
+
+		
+		/**
+		 * Return true if the pass value type is support has an attribute
+		 * See EAttributeTypes to get all the supported type.
+		 *
+		 * @param Value is the value we check the type
+		 */
+		template<typename T>
+		static bool IsSupportedType(const T& Value)
+		{
+			const EAttributeTypes ValueType = TAttributeTypeTraits<T>::GetType();
+			if (ValueType == EAttributeTypes::None)
+			{
+				return false;
+			}
+			return true;
+		}
+
+		
+		/**
+		 * Return the attribute type if the key exist, return None if the key is missing
+		 *
+		 * @param ElementAttributeKey - is the storage key (the path) of the attribute.
+		 *
+		 */
+		EAttributeTypes GetAttributeType(const FAttributeKey& ElementAttributeKey);
+
+		/**
+		 * Return true if the attribute key point on an existing attribute in the storage. Return false otherwise
+		 *
+		 * @param ElementAttributeKey - is the storage key (the path) of the attribute.
+		 *
+		 */
+		bool ContainAttribute(const FAttributeKey& ElementAttributeKey) const;
+
+		/**
+		 * This function allow the user to retrieve the array of keys to iterate and do reflection on the storage content
+		 *
+		 */
+		void GetAttributeKeys(TArray<FAttributeKey>& AttributeKeys) const;
+	
+		/**
+		 * Return a FGuid build from the FSHA1 of the specified attribute data. If the Attribute do not exist it return an empty FGUID
+		 *
+		 * @param ElementAttributeKey is the storage key (the path) of the attribute.
+		 *
+		 */
+		FGuid GetAttributeHash(const FAttributeKey& ElementAttributeKey) const;
+	
+		/**
+		 * This function fill the OutGuid with the hash of the specified attribute. Return true if the attribute exist and the OutGuid was assign, return false otherwise without touching the OutGuid.
+		 *
+		 * @param ElementAttributeKey - is the storage key (the path) of the attribute.
+		 * @param OutGuid - is where we put the attribute hash.
+		 *
+		 */
+		bool GetAttributeHash(const FAttributeKey& ElementAttributeKey, FGuid& OutGuid) const;
+
+		/**
+		 * Return a FGuid build from the FSHA1 of all the attribute data contain in the node.
+		 * The data include the UniqueID and the DisplayLabel.
+		 *
+		 * @note the attribute are sorted by key when building the FSHA1 data. The hash will be deterministic for the same data whatever
+		 * the order we add the attributes.
+		 */
+		FGuid GetStorageHash() const;
+
+		/**
+		 * Compare two storage to know which properties was modified/added/removed.
+		 *
+		 * @param BaseStorage - The storage reference
+		 * @param VersionStorage - The storage with the changes
+		 * @param RemovedAttribute - All attributes that are in base storage but not in version storage, contain key only valid for base storage
+		 * @param AddedAttributes - All attributes that are in version storage but not in base storage, contain key only valid for version storage
+		 * @param ModifiedAttributes - All attributes that are in both storage base storage but have a different hash (different value)
+		 *
+		 */
+		static void CompareStorage(const FAttributeStorage& BaseStorage, const FAttributeStorage& VersionStorage, TArray<FAttributeKey>& RemovedAttributes, TArray<FAttributeKey>& AddedAttributes, TArray<FAttributeKey>& ModifiedAttributes);
+
+		/**
+		 * Return the defrag ratio. This ratio is use to know when we need to defrag the storage
+		 * @example - a ratio of 0.1f will defrag the storage if the memory lost is bigger then 10% of the storage allocation.
+		 * Defrag is call when we remove an attribute or when we set the defrag ratio
+		 */
+		FORCEINLINE float GetDefragRatio() const
+		{
+			return DefragRatio;
+		}
+	
+		/** Set the defrag ratio. See the GetDefragRatio for the defrag documentation. */
+		void SetDefragRatio(const float InDefragRatio);
+	
+		friend FArchive& operator<<(FArchive& Ar, FAttributeStorage& Storage)
+		{
+			Ar << Storage.FragmentedMemoryCost;
+			Ar << Storage.DefragRatio;
+			Ar << Storage.AttributeAllocationTable;
+			Ar << Storage.AttributeStorage;
+			return Ar;
+		}
+
+		/**
+		 * Reserve the allocation table and the storage data.
+		 * 
+		 * @param NewAttributeCount: The number of attribute we want to reserve. Zero value do not reserve attribute count
+		 * @param NewStorageSize: The size on the storage all those new attribute will need. Passing zero value do not reserve storage size
+		 */
+		void Reserve(int64 NewAttributeCount, int64 NewStorageSize);
+
+	protected:
+		/** Structure use to hold the attribute information store in the attribute allocation table. */
+        struct FAttributeAllocationInfo
+        {
+        	//The offset in the storage
+        	uint64 Offset = 0;
+        	//The size of the data in the storage
+        	uint64 Size = 0;
+        	//The real type of the attribute
+        	EAttributeTypes Type = EAttributeTypes::None;
+        	//The attribute properties
+        	EAttributeProperty Property = EAttributeProperty::None;
+        	//128 bit Attribute hash
+        	FGuid Hash = FGuid();
+
+        	bool operator==(const FAttributeAllocationInfo& Other) const
+        	{
+        		//Offset is a unique key
+        		return Offset == Other.Offset
+        			&& Size == Other.Size
+        			&& Type == Other.Type
+        			&& Property == Other.Property
+        			&& Hash == Other.Hash;
+        	}
+        	
+        	//Serialization
+        	friend FArchive& operator<<(FArchive& Ar, FAttributeAllocationInfo& AttributeAllocationInfo)
+        	{
+        		Ar << AttributeAllocationInfo.Offset;
+        		Ar << AttributeAllocationInfo.Size;
+        		Ar << AttributeAllocationInfo.Type;
+        		Ar << AttributeAllocationInfo.Property;
+        		Ar << AttributeAllocationInfo.Hash;
+        		return Ar;
+        	}
+        };
+		
+		/**
+	     * Set an attribute value into the storage. Return success if the attribute was properly set.
+	     *
+	     * @param ElementAttributeKey is the storage key (the path) of the attribute
+	     * @param Value is the value we want to add to the storage
+	     *
+	     * @note Possible errors
+	     * - An attribute of a unsupported type see IsSupportedType
+	     * - Key exist with a different type
+	     * - Key exist with a wrong size
+	     */
+		template<typename T>
+        EAttributeStorageResult SetAttribute(const FAttributeKey& ElementAttributeKey, const T& Value)
+		{
+			FAttributeAllocationInfo* AttributeAllocationInfo = AttributeAllocationTable.Find(ElementAttributeKey);
+			return SetAttribute(AttributeAllocationInfo, Value);
+		}
+
+		template<typename T>
+        EAttributeStorageResult SetAttribute(FAttributeAllocationInfo* AttributeAllocationInfo, const T& Value)
+		{
+			return SetAttribute(AttributeAllocationInfo, Value, TSpecializeType<T>());
+		}
+
+		/**
+		 * Retrieve a copy of an attribute value. Return success if the attribute is added.
+		 *
+		 * @param ElementAttributeKey is the storage key (the path) of the attribute.
+		 * @param OutValue the reference value where we copy of the attribute value.
+		 *
+		 * @note Possible errors
+		 * - Key do not exist
+		 * - Key exist with a different type
+		 * - Key exist with a wrong size
+		 */
+		template<typename T>
+		EAttributeStorageResult GetAttribute(const FAttributeKey& ElementAttributeKey, T& OutValue) const
+		{
+			return GetAttribute(ElementAttributeKey, OutValue,TSpecializeType<T>());
+		}
+	
+		/** Defrag the storage by using memmove on the attribute store after a hole in the storage. */
+		void DefragInternal();
+
+		static FGuid GetValueHash(const uint8* Value, uint64 ValueSize);
+
+		/** Return the storage size of the value. */
+		template<typename T>
+		static uint64 GetValueSize(const T& Value)
+		{
+			return GetValueSize(Value, TSpecializeType<T>());
+		}
+	
+		//Begin specialize templates
+	
+		template<typename T>
+		EAttributeStorageResult SetAttribute(FAttributeAllocationInfo* AttributeAllocationInfo, const T& Value, TSpecializeType<T>)
+		{
+			//Lock the storage
+			FScopeLock ScopeLock(&StorageMutex);
+
+			if (!IsSupportedType(Value))
+			{
+				return EAttributeStorageResult::Operation_Error_UnsupportedType;
+			}
+
+			if (!AttributeAllocationInfo)
+			{
+				return EAttributeStorageResult::Operation_Error_CannotFoundKey;
+			}
+			const EAttributeTypes ValueType = TAttributeTypeTraits<T>::GetType();
+			const uint64 ValueSize = GetValueSize(Value);
+
+			if (AttributeAllocationInfo->Type != ValueType)
+			{
+				return EAttributeStorageResult::Operation_Error_WrongType;
+			}
+			if (AttributeAllocationInfo->Size != ValueSize)
+			{
+				return EAttributeStorageResult::Operation_Error_WrongSize;
+			}
+			if (!AttributeStorage.IsValidIndex(AttributeAllocationInfo->Offset))
+			{
+				return EAttributeStorageResult::Operation_Error_AttributeAllocationCorrupted;
+			}
+
+			//Set the hash from the value
+			AttributeAllocationInfo->Hash = GetValueHash(reinterpret_cast<const uint8*>(&Value), ValueSize);
+
+			uint8* StorageData = AttributeStorage.GetData();
+			FMemory::Memcpy(&StorageData[AttributeAllocationInfo->Offset], &Value, ValueSize);
+			return EAttributeStorageResult::Operation_Success;
+		}
+
+		EAttributeStorageResult SetAttribute(FAttributeAllocationInfo* AttributeAllocationInfo, const TArray<uint8>& Value, TSpecializeType<TArray<uint8> >)
+		{
+			return MultiSizeSetAttribute(AttributeAllocationInfo, Value, [&Value]()->uint8*
+			{
+				return const_cast<uint8*>(Value.GetData());
+			});
+		}
+
+		EAttributeStorageResult SetAttribute(FAttributeAllocationInfo* AttributeAllocationInfo, const TArray64<uint8>& Value, TSpecializeType<TArray64<uint8> >)
+		{
+			return MultiSizeSetAttribute(AttributeAllocationInfo, Value, [&Value]()->uint8*
+            {
+                return const_cast<uint8*>(Value.GetData());
+            });
+		}
+
+		EAttributeStorageResult SetAttribute(FAttributeAllocationInfo* AttributeAllocationInfo, const FString& Value, TSpecializeType<FString >)
+		{
+			return MultiSizeSetAttribute(AttributeAllocationInfo, Value, [&Value]()->uint8*
+            {
+                return (uint8*)(*Value);
+            });
+		}
+
+		EAttributeStorageResult SetAttribute(FAttributeAllocationInfo* AttributeAllocationInfo, const FName& Value, TSpecializeType<FName >)
+		{
+			//FName must be store has a FString for persistence
+			FString ValueStr = Value.ToString();
+			return MultiSizeSetAttribute(AttributeAllocationInfo, Value, [&ValueStr]()->uint8*
+			{
+				return (uint8*)(*ValueStr);
+			});
+		}
+
+		template<typename T>
+		EAttributeStorageResult GetAttribute(const FAttributeKey& ElementAttributeKey, T& OutValue, TSpecializeType<T>) const
+		{
+			//Lock the storage
+			FScopeLock ScopeLock(&StorageMutex);
+
+			if (!IsSupportedType(OutValue))
+			{
+				return EAttributeStorageResult::Operation_Error_UnsupportedType;
+			}
+
+			const EAttributeTypes ValueType = TAttributeTypeTraits<T>::GetType();
+			const uint64 ValueSize = GetValueSize(OutValue);
+
+			const FAttributeAllocationInfo* AttributeAllocationInfo = AttributeAllocationTable.Find(ElementAttributeKey);
+			if (AttributeAllocationInfo)
+			{
+				if (AttributeAllocationInfo->Type != ValueType && !((ValueType == EAttributeTypes::UInt8) && (ValueType == EAttributeTypes::Enum)))
+				{
+					return EAttributeStorageResult::Operation_Error_WrongType;
+				}
+				if (AttributeAllocationInfo->Size != ValueSize)
+				{
+					return EAttributeStorageResult::Operation_Error_WrongSize;
+				}
+				if (!AttributeStorage.IsValidIndex(AttributeAllocationInfo->Offset))
+				{
+					return EAttributeStorageResult::Operation_Error_AttributeAllocationCorrupted;
+				}
+			}
+			else
+			{
+				//The key do not exist
+				return EAttributeStorageResult::Operation_Error_CannotFoundKey;
+			}
+
+			const uint8* StorageData = AttributeStorage.GetData();
+			FMemory::Memcpy(&OutValue, &StorageData[AttributeAllocationInfo->Offset], AttributeAllocationInfo->Size);
+			return EAttributeStorageResult::Operation_Success;
+		}
+	
+		EAttributeStorageResult GetAttribute(const FAttributeKey& ElementAttributeKey, TArray<uint8>& OutValue, TSpecializeType<TArray<uint8> >) const
+		{
+			return GenericArrayGetAttribute(ElementAttributeKey, OutValue);
+		}
+    
+		EAttributeStorageResult GetAttribute(const FAttributeKey& ElementAttributeKey, TArray64<uint8>& OutValue, TSpecializeType<TArray64<uint8> >) const
+		{
+			return GenericArrayGetAttribute(ElementAttributeKey, OutValue);
+		}
+	
+		void ExtractFStringAttributeFromStorage(const uint8* StorageData, const FAttributeAllocationInfo* AttributeAllocationInfo, FString& OutValue) const;
+
+		EAttributeStorageResult GetAttribute(const FAttributeKey& ElementAttributeKey, FString& OutValue, TSpecializeType<FString >) const;
+
+		EAttributeStorageResult GetAttribute(const FAttributeKey& ElementAttributeKey, FName& OutValue, TSpecializeType<FName >) const;
+
+		template<typename T>
+		static uint64 GetValueSize(const T& Value, TSpecializeType<T>)
+		{
+			const EAttributeTypes ValueType = TAttributeTypeTraits<T>::GetType();
+			if (ValueType == EAttributeTypes::None)
+			{
+				return 0;
+			}
+			return sizeof(T);
+		}
+	
+		static uint64 GetValueSize(const FString& Value, TSpecializeType<FString >);
+
+		static uint64 GetValueSize(const FName& Value, TSpecializeType<FName>)
+		{
+			FString ValueStr = Value.ToString();
+			return GetValueSize(ValueStr, TSpecializeType<FString >());
+		}
+
+		static uint64 GetValueSize(const TArray<uint8>& Value, TSpecializeType<TArray<uint8> >);
+
+		static uint64 GetValueSize(const TArray64<uint8>& Value, TSpecializeType<TArray64<uint8> >);
+	
+		// End specialize templates
+
+		/** Set a multisize(TArray<uint8>, FString ...) attribute value into the storage. Return success if the attribute was properly set. */
+		template<typename MultiSizeType>
+		EAttributeStorageResult MultiSizeSetAttribute(FAttributeAllocationInfo* AttributeAllocationInfo, const MultiSizeType& Value, TFunctionRef<uint8* ()> GetValuePointer)
+		{
+			//Lock the storage
+			FScopeLock ScopeLock(&StorageMutex);
+
+			const EAttributeTypes ValueType = TAttributeTypeTraits<MultiSizeType>::GetType();
+
+			if (!AttributeAllocationInfo)
+			{
+				return EAttributeStorageResult::Operation_Error_CannotFoundKey;
+			}
+
+			if (AttributeAllocationInfo->Type != ValueType)
+			{
+				return EAttributeStorageResult::Operation_Error_WrongType;
+			}
+
+			const uint64 ValueSize = GetValueSize(Value);
+
+			uint8* DataPtr = GetValuePointer();
+			//We must have a valid data pointer if we have something to copy into the storage.
+			//This case support empty array or empty string because the ValueSize will be 0.
+			if (!DataPtr && ValueSize > 0)
+			{
+				return EAttributeStorageResult::Operation_Error_InvalidMultiSizeValueData;
+			}
+
+			//If the new size is greater then the old size we have to create a new allocation info entry and delete the old one
+			if (ValueSize > AttributeAllocationInfo->Size)
+			{
+				//Find the key to remove it
+				const FAttributeKey* ElementAttributeKeyPtr = AttributeAllocationTable.FindKey(*AttributeAllocationInfo);
+				if (!ElementAttributeKeyPtr)
+				{
+					return EAttributeStorageResult::Operation_Error_CannotFoundKey;
+				}
+				const EAttributeProperty AttributePropertyBackup = AttributeAllocationInfo->Property;
+				const EAttributeStorageResult RemoveResult = UnregisterAttribute(*ElementAttributeKeyPtr);
+				if (!IsAttributeStorageResultSuccess(RemoveResult))
+				{
+					return RemoveResult;
+				}
+				AttributeAllocationInfo = nullptr;
+				//Reallocate a new one
+				AttributeAllocationInfo = &AttributeAllocationTable.Add(*ElementAttributeKeyPtr);
+				AttributeAllocationInfo->Type = ValueType;
+				AttributeAllocationInfo->Offset = AttributeStorage.AddZeroed(ValueSize);
+				//Force the specified attribute property
+				AttributeAllocationInfo->Property = AttributePropertyBackup;
+			}
+			else
+			{
+				//In case we reuse the allocation table, simply adjust the waste memory counter
+				FragmentedMemoryCost += AttributeAllocationInfo->Size - ValueSize;
+			}
+
+			//We can recycle allocation table if the size is equal or smaller, so we have to always set it
+			AttributeAllocationInfo->Size = ValueSize;
+
+			//Set the hash from the value
+			AttributeAllocationInfo->Hash = GetValueHash(DataPtr, ValueSize);
+
+			if (ValueSize > 0)
+			{
+				uint8* StorageData = AttributeStorage.GetData();
+				FMemory::Memcpy(&StorageData[AttributeAllocationInfo->Offset], DataPtr, ValueSize);
+			}
+
+			return EAttributeStorageResult::Operation_Success;
+		}
+	
+		template<typename ArrayType>
+		EAttributeStorageResult GenericArrayGetAttribute(const FAttributeKey& ElementAttributeKey, ArrayType& OutValue) const
+		{
+			//Lock the storage
+			FScopeLock ScopeLock(&StorageMutex);
+
+			const EAttributeTypes ValueType = TAttributeTypeTraits<ArrayType>::GetType();
+			if (ValueType != EAttributeTypes::ByteArray && ValueType != EAttributeTypes::ByteArray64)
+			{
+				return EAttributeStorageResult::Operation_Error_WrongType;
+			}
+
+			const FAttributeAllocationInfo* AttributeAllocationInfo = AttributeAllocationTable.Find(ElementAttributeKey);
+			if (AttributeAllocationInfo)
+			{
+				if (AttributeAllocationInfo->Type != ValueType)
+				{
+					return EAttributeStorageResult::Operation_Error_WrongType;
+				}
+				if (!AttributeStorage.IsValidIndex(AttributeAllocationInfo->Offset))
+				{
+					return EAttributeStorageResult::Operation_Error_AttributeAllocationCorrupted;
+				}
+			}
+			else
+			{
+				//The key do not exist
+				return EAttributeStorageResult::Operation_Error_CannotFoundKey;
+			}
+
+			if (AttributeAllocationInfo->Size == 0)
+			{
+				OutValue.Empty();
+				return EAttributeStorageResult::Operation_Success;
+			}
+
+			//Allocate the ByteArray
+			OutValue.AddZeroed(AttributeAllocationInfo->Size);
+			//Copy into the buffer
+			const uint8* StorageData = AttributeStorage.GetData();
+			FMemory::Memcpy(OutValue.GetData(), &StorageData[AttributeAllocationInfo->Offset], AttributeAllocationInfo->Size);
+			return EAttributeStorageResult::Operation_Success;
+		}
+
+		/** The attribute allocation table is use to index the attributes into the storage. */
+		TMap<FAttributeKey, FAttributeAllocationInfo> AttributeAllocationTable;
+
+		/** The storage of the data point by the attribute allocation table */
+		TArray64<uint8> AttributeStorage;
+
+		/**
+		 * The total size of the fragmented holes in the AttributeStorage (memory waste).
+		 * A Hole is create each time we remove an attribute.
+		 */
+		uint64 FragmentedMemoryCost = 0;
+	
+		/**
+		 * if FragmentedMemoryCost > AttributeStorage.Num*DefragRatio then defrag.
+		 * This is use whenever we remove attribute or change DefragRatio value.
+		 */
+		float DefragRatio = 0.1;
+
+		/**
+		 * Mutex use when accessing or modifying the storage
+		 */
+		mutable FCriticalSection StorageMutex;
+
+	};
+
+	/************************************************************************/
+	/* Default FVariant traits for built-in types */
+
+	/** Implements variant type traits for the built-in bool type. */
+	template<> struct TAttributeTypeTraits<bool>
+	{
+		static CONSTEXPR EAttributeTypes GetType() { return EAttributeTypes::Bool; }
+	};
+
+
+	/** Implements variant type traits for the built-in FBox type. */
+	template<> struct TAttributeTypeTraits<FBox>
+	{
+		static CONSTEXPR EAttributeTypes GetType() { return EAttributeTypes::Box; }
+	};
+
+
+	/** Implements variant type traits for the built-in FBoxSphereBounds type. */
+	template<> struct TAttributeTypeTraits<FBoxSphereBounds>
+	{
+		static CONSTEXPR EAttributeTypes GetType() { return EAttributeTypes::BoxSphereBounds; }
+	};
+
+
+	/** Implements variant type traits for byte arrays. */
+	template<> struct TAttributeTypeTraits<TArray<uint8> >
+	{
+		static CONSTEXPR EAttributeTypes GetType() { return EAttributeTypes::ByteArray; }
+	};
+
+	/** Implements variant type traits for byte array64s. */
+	template<> struct TAttributeTypeTraits<TArray64<uint8> >
+	{
+		static CONSTEXPR EAttributeTypes GetType() { return EAttributeTypes::ByteArray64; }
+	};
+
+
+	/** Implements variant type traits for the built-in FColor type. */
+	template<> struct TAttributeTypeTraits<FColor>
+	{
+		static CONSTEXPR EAttributeTypes GetType() { return EAttributeTypes::Color; }
+	};
+
+
+	/** Implements variant type traits for the built-in FDateTime type. */
+	template<> struct TAttributeTypeTraits<FDateTime>
+	{
+		static CONSTEXPR EAttributeTypes GetType() { return EAttributeTypes::DateTime; }
+	};
+
+
+	/** Implements variant type traits for the built-in double type. */
+	template<> struct TAttributeTypeTraits<double>
+	{
+		static CONSTEXPR EAttributeTypes GetType() { return EAttributeTypes::Double; }
+	};
+
+
+	/** Implements variant type traits for enumeration types. */
+	template<typename EnumType> struct TAttributeTypeTraits<TEnumAsByte<EnumType> >
+	{
+		static CONSTEXPR EAttributeTypes GetType() { return EAttributeTypes::Enum; }
+	};
+
+
+	/** Implements variant type traits for the built-in float type. */
+	template<> struct TAttributeTypeTraits<float>
+	{
+		static CONSTEXPR EAttributeTypes GetType() { return EAttributeTypes::Float; }
+	};
+
+
+	/** Implements variant type traits for the built-in FGuid type. */
+	template<> struct TAttributeTypeTraits<FGuid>
+	{
+		static CONSTEXPR EAttributeTypes GetType() { return EAttributeTypes::Guid; }
+	};
+
+
+	/** Implements variant type traits for the built-in int8 type. */
+	template<> struct TAttributeTypeTraits<int8>
+	{
+		static CONSTEXPR EAttributeTypes GetType() { return EAttributeTypes::Int8; }
+	};
+
+
+	/** Implements variant type traits for the built-in int16 type. */
+	template<> struct TAttributeTypeTraits<int16>
+	{
+		static CONSTEXPR EAttributeTypes GetType() { return EAttributeTypes::Int16; }
+	};
+
+
+	/** Implements variant type traits for the built-in int32 type. */
+	template<> struct TAttributeTypeTraits<int32>
+	{
+		static CONSTEXPR EAttributeTypes GetType() { return EAttributeTypes::Int32; }
+	};
+
+
+	/** Implements variant type traits for the built-in int64 type. */
+	template<> struct TAttributeTypeTraits<int64>
+	{
+		static CONSTEXPR EAttributeTypes GetType() { return EAttributeTypes::Int64; }
+	};
+
+
+	/** Implements variant type traits for the built-in FIntPoint type. */
+	template<> struct TAttributeTypeTraits<FIntPoint>
+	{
+		static CONSTEXPR EAttributeTypes GetType() { return EAttributeTypes::IntPoint; }
+	};
+
+
+	/** Implements variant type traits for the built-in FIntVector type. */
+	template<> struct TAttributeTypeTraits<FIntVector>
+	{
+		static CONSTEXPR EAttributeTypes GetType() { return EAttributeTypes::IntVector; }
+	};
+
+
+	/** Implements variant type traits for the built-in FIntRect type. */
+	template<> struct TAttributeTypeTraits<FIntRect>
+	{
+		static CONSTEXPR EAttributeTypes GetType() { return EAttributeTypes::IntRect; }
+	};
+
+
+	/** Implements variant type traits for the built-in FLinearColor type. */
+	template<> struct TAttributeTypeTraits<FLinearColor>
+	{
+		static CONSTEXPR EAttributeTypes GetType() { return EAttributeTypes::LinearColor; }
+	};
+
+
+	/** Implements variant type traits for the built-in FMatrix type. */
+	template<> struct TAttributeTypeTraits<FMatrix>
+	{
+		static CONSTEXPR EAttributeTypes GetType() { return EAttributeTypes::Matrix; }
+	};
+
+
+	/** Implements variant type traits for the built-in FPlane type. */
+	template<> struct TAttributeTypeTraits<FPlane>
+	{
+		static CONSTEXPR EAttributeTypes GetType() { return EAttributeTypes::Plane; }
+	};
+
+
+	/** Implements variant type traits for the built-in FQuat type. */
+	template<> struct TAttributeTypeTraits<FQuat>
+	{
+		static CONSTEXPR EAttributeTypes GetType() { return EAttributeTypes::Quat; }
+	};
+
+
+	/** Implements variant type traits for the built-in FName type. */
+	template<> struct TAttributeTypeTraits<FName>
+	{
+		static CONSTEXPR EAttributeTypes GetType() { return EAttributeTypes::Name; }
+	};
+
+
+	/** Implements variant type traits for the built-in FRandomStream type. */
+	template<> struct TAttributeTypeTraits<FRandomStream>
+	{
+		static CONSTEXPR EAttributeTypes GetType() { return EAttributeTypes::RandomStream; }
+	};
+
+
+	/** Implements variant type traits for the built-in FRotator type. */
+	template<> struct TAttributeTypeTraits<FRotator>
+	{
+		static CONSTEXPR EAttributeTypes GetType() { return EAttributeTypes::Rotator; }
+	};
+
+
+	/** Implements variant type traits for the built-in FString type. */
+	template<> struct TAttributeTypeTraits<FString>
+	{
+		static CONSTEXPR EAttributeTypes GetType() { return EAttributeTypes::String; }
+	};
+
+
+	/** Implements variant type traits for the built-in FTimespan type. */
+	template<> struct TAttributeTypeTraits<FTimespan>
+	{
+		static CONSTEXPR EAttributeTypes GetType() { return EAttributeTypes::Timespan; }
+	};
+
+
+	/** Implements variant type traits for the built-in FTransform type. */
+	template<> struct TAttributeTypeTraits<FTransform>
+	{
+		static CONSTEXPR EAttributeTypes GetType() { return EAttributeTypes::Transform; }
+	};
+
+
+	/** Implements variant type traits for the built-in FTwoVectors type. */
+	template<> struct TAttributeTypeTraits<FTwoVectors>
+	{
+		static CONSTEXPR EAttributeTypes GetType() { return EAttributeTypes::TwoVectors; }
+	};
+
+
+	/** Implements variant type traits for the built-in uint8 type. */
+	template<> struct TAttributeTypeTraits<uint8>
+	{
+		static CONSTEXPR EAttributeTypes GetType() { return EAttributeTypes::UInt8; }
+	};
+
+
+	/** Implements variant type traits for the built-in uint16 type. */
+	template<> struct TAttributeTypeTraits<uint16>
+	{
+		static CONSTEXPR EAttributeTypes GetType() { return EAttributeTypes::UInt16; }
+	};
+
+
+	/** Implements variant type traits for the built-in uint32 type. */
+	template<> struct TAttributeTypeTraits<uint32>
+	{
+		static CONSTEXPR EAttributeTypes GetType() { return EAttributeTypes::UInt32; }
+	};
+
+
+	/** Implements variant type traits for the built-in uint64 type. */
+	template<> struct TAttributeTypeTraits<uint64>
+	{
+		static CONSTEXPR EAttributeTypes GetType() { return EAttributeTypes::UInt64; }
+	};
+
+
+	/** Implements variant type traits for the built-in FVector type. */
+	template<> struct TAttributeTypeTraits<FVector>
+	{
+		static CONSTEXPR EAttributeTypes GetType() { return EAttributeTypes::Vector; }
+	};
+
+
+	/** Implements variant type traits for the built-in FVector2D type. */
+	template<> struct TAttributeTypeTraits<FVector2D>
+	{
+		static CONSTEXPR EAttributeTypes GetType() { return EAttributeTypes::Vector2d; }
+	};
+
+
+	/** Implements variant type traits for the built-in FVector4 type. */
+	template<> struct TAttributeTypeTraits<FVector4>
+	{
+		static CONSTEXPR EAttributeTypes GetType() { return EAttributeTypes::Vector4; }
+	};
+
+
+	/** Implements variant type traits for the built-in NetworkGUID type. */
+	template<> struct TAttributeTypeTraits<FNetworkGUID>
+	{
+		static CONSTEXPR EAttributeTypes GetType() { return EAttributeTypes::NetworkGUID; }
+	};
+
+	/** Implements variant type traits for the built-in Vector2DHalf type. */
+	template<> struct TAttributeTypeTraits<FVector2DHalf>
+	{
+		static CONSTEXPR EAttributeTypes GetType() { return EAttributeTypes::Vector2DHalf; }
+	};
+
+	/** Implements variant type traits for the built-in Float16 type. */
+	template<> struct TAttributeTypeTraits<FFloat16>
+	{
+		static CONSTEXPR EAttributeTypes GetType() { return EAttributeTypes::Float16; }
+	};
+
+	/** Implements variant type traits for the built-in CapsuleShape type. */
+	template<> struct TAttributeTypeTraits<FCapsuleShape>
+	{
+		static CONSTEXPR EAttributeTypes GetType() { return EAttributeTypes::CapsuleShape; }
+	};
+
+	/** Implements variant type traits for the built-in DualQuat type. */
+	template<> struct TAttributeTypeTraits<FDualQuat>
+	{
+		static CONSTEXPR EAttributeTypes GetType() { return EAttributeTypes::DualQuat; }
+	};
+
+	/** Implements variant type traits for the built-in OrientedBox type. */
+	template<> struct TAttributeTypeTraits<FOrientedBox>
+	{
+		static CONSTEXPR EAttributeTypes GetType() { return EAttributeTypes::OrientedBox; }
+	};
+
+	/** Implements variant type traits for the built-in Ray type. */
+	template<> struct TAttributeTypeTraits<FRay>
+	{
+		static CONSTEXPR EAttributeTypes GetType() { return EAttributeTypes::Ray; }
+	};
+
+	/** Implements variant type traits for the built-in Sphere type. */
+	template<> struct TAttributeTypeTraits<FSphere>
+	{
+		static CONSTEXPR EAttributeTypes GetType() { return EAttributeTypes::Sphere; }
+	};
+
+	/** Implements variant type traits for the built-in FrameNumber type. */
+	template<> struct TAttributeTypeTraits<FFrameNumber>
+	{
+		static CONSTEXPR EAttributeTypes GetType() { return EAttributeTypes::FrameNumber; }
+	};
+
+	/** Implements variant type traits for the built-in FrameRate type. */
+	template<> struct TAttributeTypeTraits<FFrameRate>
+	{
+		static CONSTEXPR EAttributeTypes GetType() { return EAttributeTypes::FrameRate; }
+	};
+
+	/** Implements variant type traits for the built-in FrameTime type. */
+	template<> struct TAttributeTypeTraits<FFrameTime>
+	{
+		static CONSTEXPR EAttributeTypes GetType() { return EAttributeTypes::FrameTime; }
+	};
+
+} //End Interchange namespace
\ No newline at end of file
