Fix non unity build for interchange initial submit

diff --git a/Engine/Plugins/Runtime/Interchange/Source/Import/Private/LogInterchangeImportPlugin.h b/Engine/Plugins/Runtime/Interchange/Source/Import/Private/LogInterchangeImportPlugin.h
new file mode 100644
--- /dev/null
+++ b/Engine/Plugins/Runtime/Interchange/Source/Import/Private/LogInterchangeImportPlugin.h
@@ -0,0 +1,7 @@
+// Copyright Epic Games, Inc. All Rights Reserved.
+
+#pragma once
+
+#include "CoreMinimal.h"
+
+DECLARE_LOG_CATEGORY_EXTERN(LogInterchangeImportPlugin, Log, All);
diff --git a/Engine/Plugins/Runtime/Interchange/Source/Import/Private/Texture/InterchangeJPGTranslator.cpp b/Engine/Plugins/Runtime/Interchange/Source/Import/Private/Texture/InterchangeJPGTranslator.cpp
--- a/Engine/Plugins/Runtime/Interchange/Source/Import/Private/Texture/InterchangeJPGTranslator.cpp
+++ b/Engine/Plugins/Runtime/Interchange/Source/Import/Private/Texture/InterchangeJPGTranslator.cpp
@@ -1,12 +1,18 @@
 // Copyright Epic Games, Inc. All Rights Reserved. 
 #include "Texture/InterchangeJPGTranslator.h"
 
-#include "Nodes/BaseNodeContainer.h"
-#include "TextureNode.h"
 #include "Engine/Texture.h"
+#include "Engine/Texture2D.h"
 #include "IImageWrapper.h"
 #include "IImageWrapperModule.h"
+#include "LogInterchangeImportPlugin.h"
+#include "Misc/ConfigCacheIni.h"
+#include "Misc/FileHelper.h"
+#include "Misc/Paths.h"
 #include "Misc/ScopedSlowTask.h"
+#include "Modules/ModuleManager.h"
+#include "Nodes/BaseNodeContainer.h"
+#include "TextureNode.h"
 
 bool UInterchangeJPGTranslator::CanImportSourceData(const UInterchangeSourceData* InSourceData) const
 {
@@ -87,7 +93,7 @@ const TOptional<Interchange::FImportImage> UInterchangeJPGTranslator::GetPayload
 		UE_LOG(LogInterchangeImportPlugin, Error, TEXT("Failed to decode JPEG. [%s]"), *Filename);
 		return TOptional<Interchange::FImportImage>();
 	}
-	if (!Interchange::IsImportResolutionValid(JpegImageWrapper->GetWidth(), JpegImageWrapper->GetHeight(), bAllowNonPowerOfTwo))
+	if (!Interchange::FImportImageHelper::IsImportResolutionValid(JpegImageWrapper->GetWidth(), JpegImageWrapper->GetHeight(), bAllowNonPowerOfTwo))
 	{
 		UE_LOG(LogInterchangeImportPlugin, Error, TEXT("Failed to import JPEG, invalid resolution. Resolution[%d, %d], AllowPowerOfTwo[%s], [%s]"), JpegImageWrapper->GetWidth(), JpegImageWrapper->GetHeight(), bAllowNonPowerOfTwo ? TEXT("True") : TEXT("false"), *Filename);
 		return TOptional<Interchange::FImportImage>();
diff --git a/Engine/Plugins/Runtime/Interchange/Source/Import/Private/Texture/InterchangePNGTranslator.cpp b/Engine/Plugins/Runtime/Interchange/Source/Import/Private/Texture/InterchangePNGTranslator.cpp
--- a/Engine/Plugins/Runtime/Interchange/Source/Import/Private/Texture/InterchangePNGTranslator.cpp
+++ b/Engine/Plugins/Runtime/Interchange/Source/Import/Private/Texture/InterchangePNGTranslator.cpp
@@ -1,12 +1,18 @@
 // Copyright Epic Games, Inc. All Rights Reserved. 
 #include "Texture/InterchangePNGTranslator.h"
 
-#include "Nodes/BaseNodeContainer.h"
-#include "TextureNode.h"
 #include "Engine/Texture.h"
+#include "Engine/Texture2D.h"
 #include "IImageWrapper.h"
 #include "IImageWrapperModule.h"
+#include "LogInterchangeImportPlugin.h"
+#include "Misc/ConfigCacheIni.h"
+#include "Misc/FileHelper.h"
+#include "Misc/Paths.h"
 #include "Misc/ScopedSlowTask.h"
+#include "Modules/ModuleManager.h"
+#include "Nodes/BaseNodeContainer.h"
+#include "TextureNode.h"
 
 
 //////////////////////////////////////////////////////////////////////////
@@ -254,7 +260,7 @@ const TOptional<Interchange::FImportImage> UInterchangePNGTranslator::GetPayload
 		UE_LOG(LogInterchangeImportPlugin, Error, TEXT("Failed to decode PNG. [%s]"), *Filename);
 		return TOptional<Interchange::FImportImage>();
 	}
-	if (!Interchange::IsImportResolutionValid(PngImageWrapper->GetWidth(), PngImageWrapper->GetHeight(), bAllowNonPowerOfTwo))
+	if (!Interchange::FImportImageHelper::IsImportResolutionValid(PngImageWrapper->GetWidth(), PngImageWrapper->GetHeight(), bAllowNonPowerOfTwo))
 	{
 		UE_LOG(LogInterchangeImportPlugin, Error, TEXT("Failed to import PNG, invalid resolution. Resolution[%d, %d], AllowPowerOfTwo[%s], [%s]"), PngImageWrapper->GetWidth(), PngImageWrapper->GetHeight(), bAllowNonPowerOfTwo ? TEXT("True") : TEXT("false"), *Filename);
 		return TOptional<Interchange::FImportImage>();
diff --git a/Engine/Plugins/Runtime/Interchange/Source/Import/Private/Texture/InterchangeTextureFactory.cpp b/Engine/Plugins/Runtime/Interchange/Source/Import/Private/Texture/InterchangeTextureFactory.cpp
--- a/Engine/Plugins/Runtime/Interchange/Source/Import/Private/Texture/InterchangeTextureFactory.cpp
+++ b/Engine/Plugins/Runtime/Interchange/Source/Import/Private/Texture/InterchangeTextureFactory.cpp
@@ -1,11 +1,14 @@
 // Copyright Epic Games, Inc. All Rights Reserved. 
 #include "Texture/InterchangeTextureFactory.h"
 
+#include "LogInterchangeImportPlugin.h"
 #include "Nodes/BaseNode.h"
 #include "TextureNode.h"
 #include "TextureCompiler.h"
+#include "Texture/InterchangeTextureTranslator.h"
 #include "TextureNode.h"
 
+
 #if WITH_EDITORONLY_DATA
 
 #include "EditorFramework/AssetImportData.h"
diff --git a/Engine/Plugins/Runtime/Interchange/Source/Import/Private/Texture/InterchangeTexturePayloadData.cpp b/Engine/Plugins/Runtime/Interchange/Source/Import/Private/Texture/InterchangeTexturePayloadData.cpp
--- a/Engine/Plugins/Runtime/Interchange/Source/Import/Private/Texture/InterchangeTexturePayloadData.cpp
+++ b/Engine/Plugins/Runtime/Interchange/Source/Import/Private/Texture/InterchangeTexturePayloadData.cpp
@@ -9,7 +9,7 @@
 #include "Misc/MessageDialog.h"
 #include "RHI.h"
 
-bool Interchange::IsImportResolutionValid(int32 Width, int32 Height, bool bAllowNonPowerOfTwo)
+bool Interchange::FImportImageHelper::IsImportResolutionValid(int32 Width, int32 Height, bool bAllowNonPowerOfTwo)
 {
 	static const auto CVarVirtualTexturesEnabled = IConsoleManager::Get().FindTConsoleVariableDataInt(TEXT("r.VirtualTextures")); check(CVarVirtualTexturesEnabled);
 
diff --git a/Engine/Plugins/Runtime/Interchange/Source/Import/Public/IInterchangeImportPlugin.h b/Engine/Plugins/Runtime/Interchange/Source/Import/Public/IInterchangeImportPlugin.h
--- a/Engine/Plugins/Runtime/Interchange/Source/Import/Public/IInterchangeImportPlugin.h
+++ b/Engine/Plugins/Runtime/Interchange/Source/Import/Public/IInterchangeImportPlugin.h
@@ -3,11 +3,10 @@
 #pragma once
 
 #include "CoreMinimal.h"
+#include "LogInterchangeImportPlugin.h"
 #include "Modules/ModuleInterface.h"
 #include "Modules/ModuleManager.h"
 
-DECLARE_LOG_CATEGORY_EXTERN(LogInterchangeImportPlugin, Log, All);
-
 /**
  * The public interface to this module
  */
diff --git a/Engine/Plugins/Runtime/Interchange/Source/Import/Public/Texture/InterchangeTexturePayloadData.h b/Engine/Plugins/Runtime/Interchange/Source/Import/Public/Texture/InterchangeTexturePayloadData.h
--- a/Engine/Plugins/Runtime/Interchange/Source/Import/Public/Texture/InterchangeTexturePayloadData.h
+++ b/Engine/Plugins/Runtime/Interchange/Source/Import/Public/Texture/InterchangeTexturePayloadData.h
@@ -25,16 +25,19 @@ namespace Interchange
 		void* GetMipData(int32 InMipIndex);
 	};
 
-	/**
-	 * Tests if the given height and width specify a supported texture resolution to import; Can optionally check if the height/width are powers of two
-	 *
-	 * @param Width - The width of an imported texture whose validity should be checked
-	 * @param Height - The height of an imported texture whose validity should be checked
-	 * @param bAllowNonPowerOfTwo - Whether or not non-power-of-two textures are allowed
-	 *
-	 * @return bool true if the given height/width represent a supported texture resolution, false if not
-	 */
-	static bool IsImportResolutionValid(int32 Width, int32 Height, bool bAllowNonPowerOfTwo);
+	struct FImportImageHelper
+	{
+		/**
+		 * Tests if the given height and width specify a supported texture resolution to import; Can optionally check if the height/width are powers of two
+		 *
+		 * @param Width - The width of an imported texture whose validity should be checked
+		 * @param Height - The height of an imported texture whose validity should be checked
+		 * @param bAllowNonPowerOfTwo - Whether or not non-power-of-two textures are allowed
+		 *
+		 * @return bool true if the given height/width represent a supported texture resolution, false if not
+		 */
+		static bool IsImportResolutionValid(int32 Width, int32 Height, bool bAllowNonPowerOfTwo);
+	};
 }
 
 
diff --git a/Engine/Plugins/Runtime/Interchange/Source/Nodes/Public/TextureNode.h b/Engine/Plugins/Runtime/Interchange/Source/Nodes/Public/TextureNode.h
--- a/Engine/Plugins/Runtime/Interchange/Source/Nodes/Public/TextureNode.h
+++ b/Engine/Plugins/Runtime/Interchange/Source/Nodes/Public/TextureNode.h
@@ -5,6 +5,7 @@
 #include "CoreMinimal.h"
 #include "Nodes/BaseNode.h"
 #include "Engine/Texture.h"
+#include "Engine/Texture2D.h"
 
 //Interchange namespace
 namespace Interchange
diff --git a/Engine/Source/Runtime/Interchange/Core/Private/InterchangeManager.cpp b/Engine/Source/Runtime/Interchange/Core/Private/InterchangeManager.cpp
--- a/Engine/Source/Runtime/Interchange/Core/Private/InterchangeManager.cpp
+++ b/Engine/Source/Runtime/Interchange/Core/Private/InterchangeManager.cpp
@@ -1,14 +1,13 @@
 // Copyright Epic Games, Inc. All Rights Reserved.
 
-#pragma once
-
 #include "InterchangeManager.h"
 
 #include "AssetDataTagMap.h"
 #include "AssetRegistryModule.h"
 #include "CoreMinimal.h"
-#include "Nodes/BaseNodeContainer.h"
+#include "Engine/Blueprint.h"
 #include "InterchangeFactoryBase.h"
+#include "InterchangeLogPrivate.h"
 #include "InterchangeSourceData.h"
 #include "InterchangeTranslatorBase.h"
 #include "InterchangeWriterBase.h"
@@ -17,13 +16,16 @@
 #include "Misc/AsyncTaskNotification.h"
 #include "Misc/MessageDialog.h"
 #include "Misc/ScopedSlowTask.h"
+#include "Nodes/BaseNodeContainer.h"
 #include "PackageUtils/PackageUtils.h"
 #include "Tasks/InterchangeTaskParsing.h"
 #include "Tasks/InterchangeTaskPipeline.h"
 #include "Tasks/InterchangeTaskTranslator.h"
+#include "UObject/Class.h"
 #include "UObject/NameTypes.h"
 #include "UObject/Object.h"
 #include "UObject/ObjectMacros.h"
+#include "UObject/UObjectIterator.h"
 #include "UObject/WeakObjectPtrTemplates.h"
 
 namespace InternalInterchangePrivate
diff --git a/Engine/Source/Runtime/Interchange/Core/Private/Nodes/InterchangeBaseNodeContainerAdapter.cpp b/Engine/Source/Runtime/Interchange/Core/Private/Nodes/InterchangeBaseNodeContainerAdapter.cpp
--- a/Engine/Source/Runtime/Interchange/Core/Private/Nodes/InterchangeBaseNodeContainerAdapter.cpp
+++ b/Engine/Source/Runtime/Interchange/Core/Private/Nodes/InterchangeBaseNodeContainerAdapter.cpp
@@ -1,7 +1,5 @@
 // Copyright Epic Games, Inc. All Rights Reserved.
 
-#pragma once
-
 #include "Nodes/InterchangeBaseNodeContainerAdapter.h"
 
 #include "CoreMinimal.h"
diff --git a/Engine/Source/Runtime/Interchange/Core/Private/Tasks/InterchangeTaskCreateAsset.cpp b/Engine/Source/Runtime/Interchange/Core/Private/Tasks/InterchangeTaskCreateAsset.cpp
--- a/Engine/Source/Runtime/Interchange/Core/Private/Tasks/InterchangeTaskCreateAsset.cpp
+++ b/Engine/Source/Runtime/Interchange/Core/Private/Tasks/InterchangeTaskCreateAsset.cpp
@@ -1,5 +1,6 @@
 // Copyright Epic Games, Inc. All Rights Reserved.
 #include "InterchangeTaskCreateAsset.h"
+#include "InterchangeLogPrivate.h"
 
 #include "AssetRegistryModule.h"
 #include "Async/TaskGraphInterfaces.h"
diff --git a/Engine/Source/Runtime/Interchange/Core/Private/Tests/StorageTest.cpp b/Engine/Source/Runtime/Interchange/Core/Private/Tests/StorageTest.cpp
--- a/Engine/Source/Runtime/Interchange/Core/Private/Tests/StorageTest.cpp
+++ b/Engine/Source/Runtime/Interchange/Core/Private/Tests/StorageTest.cpp
@@ -4,12 +4,13 @@
 #include "Async/ParallelFor.h"
 #include "Misc/AutomationTest.h"
 #include "Serialization/MemoryReader.h"
+#include "Serialization/MemoryWriter.h"
 #include "Templates/Atomic.h"
 #include "Types/AttributeStorage.h"
 
 #if WITH_DEV_AUTOMATION_TESTS
 
-IMPLEMENT_SIMPLE_AUTOMATION_TEST(FAttributeStorageTest, "System.Runtime.Interchange.AttributeStorage", EAutomationTestFlags::ApplicationContextMask | EAutomationTestFlags::SmokeFilter)
+IMPLEMENT_SIMPLE_AUTOMATION_TEST(FAttributeStorageTest, "System.Runtime.Interchange.AttributeStorage", EAutomationTestFlags::EditorContext | EAutomationTestFlags::EngineFilter)
 
 bool FAttributeStorageTest::RunTest(const FString& Parameters)
 {
diff --git a/Engine/Source/Runtime/Interchange/Core/Private/Types/AttributeStorage.cpp b/Engine/Source/Runtime/Interchange/Core/Private/Types/AttributeStorage.cpp
--- a/Engine/Source/Runtime/Interchange/Core/Private/Types/AttributeStorage.cpp
+++ b/Engine/Source/Runtime/Interchange/Core/Private/Types/AttributeStorage.cpp
@@ -1,8 +1,10 @@
 // Copyright Epic Games, Inc. All Rights Reserved.
 #include "Types/AttributeStorage.h"
-#include "../InterchangeLogPrivate.h"
+#include "InterchangeLogPrivate.h"
 
 #include "CoreMinimal.h"
+#include "Misc/ScopeLock.h"
+#include "Misc/SecureHash.h"
 
 //Interchange namespace
 namespace Interchange
diff --git a/Engine/Source/Runtime/Interchange/Core/Public/InterchangeManager.h b/Engine/Source/Runtime/Interchange/Core/Public/InterchangeManager.h
--- a/Engine/Source/Runtime/Interchange/Core/Public/InterchangeManager.h
+++ b/Engine/Source/Runtime/Interchange/Core/Public/InterchangeManager.h
@@ -2,6 +2,7 @@
 
 #pragma once
 
+#include "Async/TaskGraphInterfaces.h"
 #include "CoreMinimal.h"
 #include "Delegates/DelegateCombinations.h"
 #include "InterchangeFactoryBase.h"
@@ -9,10 +10,11 @@
 #include "InterchangeSourceData.h"
 #include "InterchangeTranslatorBase.h"
 #include "InterchangeWriterBase.h"
+#include "Nodes/InterchangeBaseNodeContainerAdapter.h"
+#include "UObject/Package.h"
 #include "UObject/Object.h"
 #include "UObject/ObjectMacros.h"
 #include "UObject/StrongObjectPtr.h"
-#include "Nodes/InterchangeBaseNodeContainerAdapter.h"
 
 #include "InterchangeManager.generated.h"
 
diff --git a/Engine/Source/Runtime/Interchange/Core/Public/InterchangeSourceData.h b/Engine/Source/Runtime/Interchange/Core/Public/InterchangeSourceData.h
--- a/Engine/Source/Runtime/Interchange/Core/Public/InterchangeSourceData.h
+++ b/Engine/Source/Runtime/Interchange/Core/Public/InterchangeSourceData.h
@@ -3,6 +3,7 @@
 #pragma once
 
 #include "CoreMinimal.h"
+#include "Misc/Paths.h"
 #include "Misc/SecureHash.h"
 #include "UObject/Object.h"
 #include "UObject/ObjectMacros.h"
diff --git a/Engine/Source/Runtime/Interchange/Core/Public/Nodes/BaseNode.h b/Engine/Source/Runtime/Interchange/Core/Public/Nodes/BaseNode.h
--- a/Engine/Source/Runtime/Interchange/Core/Public/Nodes/BaseNode.h
+++ b/Engine/Source/Runtime/Interchange/Core/Public/Nodes/BaseNode.h
@@ -4,7 +4,8 @@
 
 #include "CoreMinimal.h"
 #include "Types/AttributeStorage.h"
-
+#include "UObject/Class.h"
+#include "UObject/Object.h"
 
 /**
  * Internal Helper to get set custom property for class that derive from FBaseNode. This is use by the macro IMPLEMENT_UOD_ATTRIBUTE.
diff --git a/Engine/Source/Runtime/Interchange/Core/Public/Types/AttributeStorage.h b/Engine/Source/Runtime/Interchange/Core/Public/Types/AttributeStorage.h
--- a/Engine/Source/Runtime/Interchange/Core/Public/Types/AttributeStorage.h
+++ b/Engine/Source/Runtime/Interchange/Core/Public/Types/AttributeStorage.h
@@ -9,7 +9,7 @@
 #include "Math/RandomStream.h"
 #include "Misc/FrameRate.h"
 #include "Misc/NetworkGuid.h"
-
+#include "Misc/ScopeLock.h"
 
 
 //Interchange namespace
@@ -34,18 +34,18 @@ namespace Interchange
 
 		explicit FAttributeKey(const FString& Other)
 		{
-			Key = Other.IsEmpty() ? NAME_None : FName(Other);
+			Key = Other.IsEmpty() ? NAME_None : FName(*Other);
 		}
 
 		explicit FAttributeKey(const FText& Other)
 		{
-			Key = Other.IsEmpty() ? NAME_None : FName(Other.ToString());
+			Key = Other.IsEmpty() ? NAME_None : FName(*(Other.ToString()));
 		}
 
 		explicit FAttributeKey(const TCHAR* Other)
 		{
 			const FString Helper(Other);
-			Key = Helper.IsEmpty() ? NAME_None : FName(Helper);
+			Key = Helper.IsEmpty() ? NAME_None : FName(Other);
 		}
 
 		FORCEINLINE const TCHAR* operator*() const
@@ -67,20 +67,20 @@ namespace Interchange
 
 		FORCEINLINE FAttributeKey& operator=(const FString& Other)
 		{
-			Key = Other.IsEmpty() ? NAME_None : FName(Other);
+			Key = Other.IsEmpty() ? NAME_None : FName(*Other);
 			return *this;
 		}
 
 		FORCEINLINE FAttributeKey& operator=(const FText& Other)
 		{
-			Key = Other.IsEmpty() ? NAME_None : FName(Other.ToString());
+			Key = Other.IsEmpty() ? NAME_None : FName(*(Other.ToString()));
 			return *this;
 		}
 
 		FORCEINLINE FAttributeKey& operator=(const TCHAR* Other)
 		{
 			const FString Helper(Other);
-			Key = Helper.IsEmpty() ? NAME_None : FName(Helper);
+			Key = Helper.IsEmpty() ? NAME_None : FName(Other);
 			return *this;
 		}
 		
@@ -120,7 +120,7 @@ namespace Interchange
 			Ar << KeyString;
 			if(Ar.IsLoading())
 			{
-				AttributeKey.Key = FName(KeyString);
+				AttributeKey.Key = FName(*KeyString);
 			}
 			return Ar;
 		}
