Various bug fixes in interchange framework.

diff --git a/Engine/Source/Runtime/Interchange/Core/Private/InterchangeManager.cpp b/Engine/Source/Runtime/Interchange/Core/Private/InterchangeManager.cpp
--- a/Engine/Source/Runtime/Interchange/Core/Private/InterchangeManager.cpp
+++ b/Engine/Source/Runtime/Interchange/Core/Private/InterchangeManager.cpp
@@ -127,7 +127,7 @@ void Interchange::FImportAsyncHelper::CleanUp()
 
 void Interchange::SanitizeInvalidChar(FString& String)
 {
-	const TCHAR* InvalidChar = *FString(INVALID_OBJECTNAME_CHARACTERS);
+	const TCHAR* InvalidChar = INVALID_OBJECTPATH_CHARACTERS;
 	while (*InvalidChar)
 	{
 		String.ReplaceCharInline(*InvalidChar, TCHAR('_'), ESearchCase::CaseSensitive);
@@ -224,9 +224,9 @@ bool UInterchangeManager::ImportAsset(const FString& ContentPath, const UInterch
 	TaskData.ImportType = Interchange::EImportType::ImportType_Asset;
 	TaskData.ReimportObject = ImportAssetParameters.ReimportAsset;
 	TWeakPtr<Interchange::FImportAsyncHelper> WeakAsyncHelper = CreateAsyncHelper(TaskData);
-	check(WeakAsyncHelper.IsValid());
 	TSharedPtr<Interchange::FImportAsyncHelper> AsyncHelper = WeakAsyncHelper.Pin();
-		
+	check(AsyncHelper.IsValid());
+
 	FText TitleText = NSLOCTEXT("Interchange", "Asynchronous_import_start", "Importing");
 	if(!Notification.IsValid())
 	{
diff --git a/Engine/Source/Runtime/Interchange/Core/Private/Tasks/InterchangeTaskCreateAsset.cpp b/Engine/Source/Runtime/Interchange/Core/Private/Tasks/InterchangeTaskCreateAsset.cpp
--- a/Engine/Source/Runtime/Interchange/Core/Private/Tasks/InterchangeTaskCreateAsset.cpp
+++ b/Engine/Source/Runtime/Interchange/Core/Private/Tasks/InterchangeTaskCreateAsset.cpp
@@ -21,9 +21,9 @@
 
 namespace Interchange_InternalImplementation
 {
-	void InternalGetPackageName(TSharedPtr<Interchange::FImportAsyncHelper> AsyncHelper, const int32 SourceIndex, const FString& PackageBasePath, const Interchange::FBaseNode* Node, FString& OutPackageName, FString& OutAssetName)
+	void InternalGetPackageName(const Interchange::FImportAsyncHelper& AsyncHelper, const int32 SourceIndex, const FString& PackageBasePath, const Interchange::FBaseNode* Node, FString& OutPackageName, FString& OutAssetName)
 	{
-		const UInterchangeSourceData* SourceData = AsyncHelper->SourceDatas[SourceIndex];
+		const UInterchangeSourceData* SourceData = AsyncHelper.SourceDatas[SourceIndex];
 		check(SourceData);
 		FString NodeDisplayName = Node->GetDisplayLabel().ToString();
 		const FString BaseFileName = FPaths::GetBaseFilename(SourceData->GetFilename());
@@ -63,7 +63,7 @@ void Interchange::FTaskCreatePackage::DoTask(ENamedThreads::Type CurrentThread,
 	}
 	else
 	{
-		Interchange_InternalImplementation::InternalGetPackageName(AsyncHelper, SourceIndex, PackageBasePath, Node, PackageName, AssetName);
+		Interchange_InternalImplementation::InternalGetPackageName(*AsyncHelper, SourceIndex, PackageBasePath, Node, PackageName, AssetName);
 		// We can not create assets that share the name of a map file in the same location
 		if (Interchange::FPackageUtils::IsMapPackageAsset(PackageName))
 		{
@@ -96,24 +96,28 @@ void Interchange::FTaskCreatePackage::DoTask(ENamedThreads::Type CurrentThread,
 	// Make sure the destination package is loaded
 	Pkg->FullyLoad();
 
+	FScopeLock Lock(&AsyncHelper->CreatedPackagesLock);
 	AsyncHelper->CreatedPackages.Add(PackageName, Pkg);
 }
 
 void Interchange::FTaskCreateAsset::DoTask(ENamedThreads::Type CurrentThread, const FGraphEventRef& MyCompletionGraphEvent)
 {
-	check(WeakAsyncHelper.IsValid());
 	TSharedPtr<Interchange::FImportAsyncHelper> AsyncHelper = WeakAsyncHelper.Pin();
+	check(WeakAsyncHelper.IsValid());
 
 	FString PackageName;
 	FString AssetName;
-	Interchange_InternalImplementation::InternalGetPackageName(AsyncHelper, SourceIndex, PackageBasePath, Node, PackageName, AssetName);
+	Interchange_InternalImplementation::InternalGetPackageName(*AsyncHelper, SourceIndex, PackageBasePath, Node, PackageName, AssetName);
 	if (AsyncHelper->TaskData.ReimportObject)
 	{
 		UPackage* Pkg = AsyncHelper->TaskData.ReimportObject->GetPackage();
 		PackageName = Pkg->GetPathName();
 	}
 
+	FScopeLock Lock(&AsyncHelper->CreatedPackagesLock);
 	UPackage** PkgPtr = AsyncHelper->CreatedPackages.Find(PackageName);
+	Lock.Unlock();
+
 	if (!PkgPtr || !(*PkgPtr))
 	{
 		const FText Message = FText::Format(NSLOCTEXT("Interchange", "CannotCreateAssetNoPackageErrorMsg", "Cannot create asset named '{1}', package '{0}'was not created properly."), FText::FromString(PackageName), FText::FromString(AssetName));
diff --git a/Engine/Source/Runtime/Interchange/Core/Private/Tasks/InterchangeTaskCreateAsset.h b/Engine/Source/Runtime/Interchange/Core/Private/Tasks/InterchangeTaskCreateAsset.h
--- a/Engine/Source/Runtime/Interchange/Core/Private/Tasks/InterchangeTaskCreateAsset.h
+++ b/Engine/Source/Runtime/Interchange/Core/Private/Tasks/InterchangeTaskCreateAsset.h
@@ -37,7 +37,7 @@ public:
 		check(Factory);
 	}
 
-	FORCEINLINE ENamedThreads::Type GetDesiredThread()
+	ENamedThreads::Type GetDesiredThread()
 	{
 		if (WeakAsyncHelper.IsValid() && WeakAsyncHelper.Pin()->TaskData.ReimportObject)
 		{
@@ -46,12 +46,13 @@ public:
 		}
 		return ENamedThreads::GameThread;
 	}
-	static FORCEINLINE ESubsequentsMode::Type GetSubsequentsMode()
+
+	static ESubsequentsMode::Type GetSubsequentsMode()
 	{
 		return ESubsequentsMode::TrackSubsequents;
 	}
 
-	FORCEINLINE TStatId GetStatId() const
+	TStatId GetStatId() const
 	{
 		RETURN_QUICK_DECLARE_CYCLE_STAT(FTaskCreatePackage, STATGROUP_TaskGraphTasks);
 	}
@@ -67,6 +68,7 @@ private:
 	TWeakPtr<Interchange::FImportAsyncHelper> WeakAsyncHelper;
 	const FBaseNode* Node;
 	UInterchangeFactoryBase* Factory;
+
 public:
 	FTaskCreateAsset(const FString& InPackageBasePath, const int32 InSourceIndex, TWeakPtr<Interchange::FImportAsyncHelper> InAsyncHelper, const Interchange::FBaseNode* InNode, UInterchangeFactoryBase* InFactory)
 		: PackageBasePath(InPackageBasePath)
@@ -79,16 +81,17 @@ public:
 		check(Factory);
 	}
 
-	FORCEINLINE ENamedThreads::Type GetDesiredThread()
+	ENamedThreads::Type GetDesiredThread()
 	{
 		return Factory->CanExecuteOnAnyThread() ? ENamedThreads::AnyBackgroundThreadNormalTask : ENamedThreads::GameThread;
 	}
-	static FORCEINLINE ESubsequentsMode::Type GetSubsequentsMode()
+
+	static ESubsequentsMode::Type GetSubsequentsMode()
 	{
 		return ESubsequentsMode::TrackSubsequents;
 	}
 
-	FORCEINLINE TStatId GetStatId() const
+	TStatId GetStatId() const
 	{
 		RETURN_QUICK_DECLARE_CYCLE_STAT(FTaskCreateAsset, STATGROUP_TaskGraphTasks);
 	}
diff --git a/Engine/Source/Runtime/Interchange/Core/Public/InterchangeManager.h b/Engine/Source/Runtime/Interchange/Core/Public/InterchangeManager.h
--- a/Engine/Source/Runtime/Interchange/Core/Public/InterchangeManager.h
+++ b/Engine/Source/Runtime/Interchange/Core/Public/InterchangeManager.h
@@ -87,6 +87,7 @@ namespace Interchange
 		FGraphEventRef CompletionTask;
 
 		//Create package map, Key is package name. We cannot create package asynchronously so we have to create a game thread task to do this
+		FCriticalSection CreatedPackagesLock;
 		TMap<FString, UPackage*> CreatedPackages;
 
 		struct FImportedAssetInfo
