Interchange: Use threadsafe shared pointers

diff --git a/Engine/Source/Runtime/Interchange/Core/Private/InterchangeManager.cpp b/Engine/Source/Runtime/Interchange/Core/Private/InterchangeManager.cpp
--- a/Engine/Source/Runtime/Interchange/Core/Private/InterchangeManager.cpp
+++ b/Engine/Source/Runtime/Interchange/Core/Private/InterchangeManager.cpp
@@ -264,8 +264,8 @@ Interchange::FAsyncImportResult UInterchangeManager::ImportAssetAsync(const FStr
 	TaskData.bIsAutomated = ImportAssetParameters.bIsAutomated;
 	TaskData.ImportType = Interchange::EImportType::ImportType_Asset;
 	TaskData.ReimportObject = ImportAssetParameters.ReimportAsset;
-	TWeakPtr<Interchange::FImportAsyncHelper> WeakAsyncHelper = CreateAsyncHelper(TaskData);
-	TSharedPtr<Interchange::FImportAsyncHelper> AsyncHelper = WeakAsyncHelper.Pin();
+	TWeakPtr<Interchange::FImportAsyncHelper, ESPMode::ThreadSafe> WeakAsyncHelper = CreateAsyncHelper(TaskData);
+	TSharedPtr<Interchange::FImportAsyncHelper, ESPMode::ThreadSafe> AsyncHelper = WeakAsyncHelper.Pin();
 	check(AsyncHelper.IsValid());
 
 	FText TitleText = NSLOCTEXT("Interchange", "Asynchronous_import_start", "Importing");
@@ -388,9 +388,9 @@ UInterchangeSourceData* UInterchangeManager::CreateSourceData(const FString& InF
 	return SourceDataAsset;
 }
 
-TWeakPtr<Interchange::FImportAsyncHelper> UInterchangeManager::CreateAsyncHelper(const Interchange::FImportAsyncHelperData& Data)
+TWeakPtr<Interchange::FImportAsyncHelper, ESPMode::ThreadSafe> UInterchangeManager::CreateAsyncHelper(const Interchange::FImportAsyncHelperData& Data)
 {
-	TSharedPtr<Interchange::FImportAsyncHelper> AsyncHelper = MakeShared<Interchange::FImportAsyncHelper>();
+	TSharedPtr<Interchange::FImportAsyncHelper, ESPMode::ThreadSafe> AsyncHelper = MakeShared<Interchange::FImportAsyncHelper, ESPMode::ThreadSafe>();
 	//Copy the task data
 	AsyncHelper->TaskData = Data;
 	int32 AsyncHelperIndex = ImportTasks.Add(AsyncHelper);
@@ -405,7 +405,7 @@ TWeakPtr<Interchange::FImportAsyncHelper> UInterchangeManager::CreateAsyncHelper
 	return ImportTasks[AsyncHelperIndex];
 }
 
-void UInterchangeManager::ReleaseAsyncHelper(TWeakPtr<Interchange::FImportAsyncHelper> AsyncHelper)
+void UInterchangeManager::ReleaseAsyncHelper(TWeakPtr<Interchange::FImportAsyncHelper, ESPMode::ThreadSafe> AsyncHelper)
 {
 	check(AsyncHelper.IsValid());
 	ImportTasks.RemoveSingle(AsyncHelper.Pin());
diff --git a/Engine/Source/Runtime/Interchange/Core/Private/Tasks/InterchangeTaskCompletion.cpp b/Engine/Source/Runtime/Interchange/Core/Private/Tasks/InterchangeTaskCompletion.cpp
--- a/Engine/Source/Runtime/Interchange/Core/Private/Tasks/InterchangeTaskCompletion.cpp
+++ b/Engine/Source/Runtime/Interchange/Core/Private/Tasks/InterchangeTaskCompletion.cpp
@@ -17,7 +17,7 @@
 
 void Interchange::FTaskCompletion::DoTask(ENamedThreads::Type CurrentThread, const FGraphEventRef& MyCompletionGraphEvent)
 {
-	TSharedPtr<Interchange::FImportAsyncHelper> AsyncHelper = WeakAsyncHelper.Pin();
+	TSharedPtr<Interchange::FImportAsyncHelper, ESPMode::ThreadSafe> AsyncHelper = WeakAsyncHelper.Pin();
 	check(AsyncHelper.IsValid());
 
 	for(TPair<int32, TArray<Interchange::FImportAsyncHelper::FImportedAssetInfo>>& AssetInfosPerSourceIndexPair : AsyncHelper->ImportedAssetsPerSourceIndex)
diff --git a/Engine/Source/Runtime/Interchange/Core/Private/Tasks/InterchangeTaskCompletion.h b/Engine/Source/Runtime/Interchange/Core/Private/Tasks/InterchangeTaskCompletion.h
--- a/Engine/Source/Runtime/Interchange/Core/Private/Tasks/InterchangeTaskCompletion.h
+++ b/Engine/Source/Runtime/Interchange/Core/Private/Tasks/InterchangeTaskCompletion.h
@@ -15,9 +15,9 @@ class FTaskCompletion
 {
 private:
 	UInterchangeManager* InterchangeManager;
-	TWeakPtr<Interchange::FImportAsyncHelper> WeakAsyncHelper;
+	TWeakPtr<Interchange::FImportAsyncHelper, ESPMode::ThreadSafe> WeakAsyncHelper;
 public:
-	FTaskCompletion(UInterchangeManager* InInterchangeManager, TWeakPtr<Interchange::FImportAsyncHelper> InAsyncHelper)
+	FTaskCompletion(UInterchangeManager* InInterchangeManager, TWeakPtr<Interchange::FImportAsyncHelper, ESPMode::ThreadSafe> InAsyncHelper)
 		: InterchangeManager(InInterchangeManager)
 		, WeakAsyncHelper(InAsyncHelper)
 	{
diff --git a/Engine/Source/Runtime/Interchange/Core/Private/Tasks/InterchangeTaskCreateAsset.cpp b/Engine/Source/Runtime/Interchange/Core/Private/Tasks/InterchangeTaskCreateAsset.cpp
--- a/Engine/Source/Runtime/Interchange/Core/Private/Tasks/InterchangeTaskCreateAsset.cpp
+++ b/Engine/Source/Runtime/Interchange/Core/Private/Tasks/InterchangeTaskCreateAsset.cpp
@@ -45,7 +45,7 @@ namespace Interchange_InternalImplementation
 
 void Interchange::FTaskCreatePackage::DoTask(ENamedThreads::Type CurrentThread, const FGraphEventRef& MyCompletionGraphEvent)
 {
-	TSharedPtr<Interchange::FImportAsyncHelper> AsyncHelper = WeakAsyncHelper.Pin();
+	TSharedPtr<Interchange::FImportAsyncHelper, ESPMode::ThreadSafe> AsyncHelper = WeakAsyncHelper.Pin();
 	check(AsyncHelper.IsValid());
 
 	//The create package thread must always execute on the game thread
@@ -101,7 +101,7 @@ void Interchange::FTaskCreatePackage::DoTask(ENamedThreads::Type CurrentThread,
 
 void Interchange::FTaskCreateAsset::DoTask(ENamedThreads::Type CurrentThread, const FGraphEventRef& MyCompletionGraphEvent)
 {
-	TSharedPtr<Interchange::FImportAsyncHelper> AsyncHelper = WeakAsyncHelper.Pin();
+	TSharedPtr<Interchange::FImportAsyncHelper, ESPMode::ThreadSafe> AsyncHelper = WeakAsyncHelper.Pin();
 	check(WeakAsyncHelper.IsValid());
 
 	FString PackageName;
diff --git a/Engine/Source/Runtime/Interchange/Core/Private/Tasks/InterchangeTaskCreateAsset.h b/Engine/Source/Runtime/Interchange/Core/Private/Tasks/InterchangeTaskCreateAsset.h
--- a/Engine/Source/Runtime/Interchange/Core/Private/Tasks/InterchangeTaskCreateAsset.h
+++ b/Engine/Source/Runtime/Interchange/Core/Private/Tasks/InterchangeTaskCreateAsset.h
@@ -21,12 +21,12 @@ class FTaskCreatePackage
 private:
 	FString PackageBasePath;
 	int32 SourceIndex;
-	TWeakPtr<Interchange::FImportAsyncHelper> WeakAsyncHelper;
+	TWeakPtr<Interchange::FImportAsyncHelper, ESPMode::ThreadSafe> WeakAsyncHelper;
 	const FBaseNode* Node;
 	UInterchangeFactoryBase* Factory;
 
 public:
-	FTaskCreatePackage(const FString& InPackageBasePath, const int32 InSourceIndex, TWeakPtr<Interchange::FImportAsyncHelper> InAsyncHelper, const Interchange::FBaseNode* InNode, UInterchangeFactoryBase* InFactory)
+	FTaskCreatePackage(const FString& InPackageBasePath, const int32 InSourceIndex, TWeakPtr<Interchange::FImportAsyncHelper, ESPMode::ThreadSafe> InAsyncHelper, const Interchange::FBaseNode* InNode, UInterchangeFactoryBase* InFactory)
 		: PackageBasePath(InPackageBasePath)
 		, SourceIndex(InSourceIndex)
 		, WeakAsyncHelper(InAsyncHelper)
@@ -65,12 +65,12 @@ class FTaskCreateAsset
 private:
 	FString PackageBasePath;
 	int32 SourceIndex;
-	TWeakPtr<Interchange::FImportAsyncHelper> WeakAsyncHelper;
+	TWeakPtr<Interchange::FImportAsyncHelper, ESPMode::ThreadSafe> WeakAsyncHelper;
 	const FBaseNode* Node;
 	UInterchangeFactoryBase* Factory;
 
 public:
-	FTaskCreateAsset(const FString& InPackageBasePath, const int32 InSourceIndex, TWeakPtr<Interchange::FImportAsyncHelper> InAsyncHelper, const Interchange::FBaseNode* InNode, UInterchangeFactoryBase* InFactory)
+	FTaskCreateAsset(const FString& InPackageBasePath, const int32 InSourceIndex, TWeakPtr<Interchange::FImportAsyncHelper, ESPMode::ThreadSafe> InAsyncHelper, const Interchange::FBaseNode* InNode, UInterchangeFactoryBase* InFactory)
 		: PackageBasePath(InPackageBasePath)
 		, SourceIndex(InSourceIndex)
 		, WeakAsyncHelper(InAsyncHelper)
diff --git a/Engine/Source/Runtime/Interchange/Core/Private/Tasks/InterchangeTaskParsing.cpp b/Engine/Source/Runtime/Interchange/Core/Private/Tasks/InterchangeTaskParsing.cpp
--- a/Engine/Source/Runtime/Interchange/Core/Private/Tasks/InterchangeTaskParsing.cpp
+++ b/Engine/Source/Runtime/Interchange/Core/Private/Tasks/InterchangeTaskParsing.cpp
@@ -22,7 +22,7 @@
 
 void Interchange::FTaskParsing::DoTask(ENamedThreads::Type CurrentThread, const FGraphEventRef& MyCompletionGraphEvent)
 {
-	TSharedPtr<Interchange::FImportAsyncHelper> AsyncHelper = WeakAsyncHelper.Pin();
+	TSharedPtr<Interchange::FImportAsyncHelper, ESPMode::ThreadSafe> AsyncHelper = WeakAsyncHelper.Pin();
 	check(AsyncHelper.IsValid());
 
 	//Parse each graph and prepare import task data, we will then be able to create all the task with the correct dependencies
diff --git a/Engine/Source/Runtime/Interchange/Core/Private/Tasks/InterchangeTaskParsing.h b/Engine/Source/Runtime/Interchange/Core/Private/Tasks/InterchangeTaskParsing.h
--- a/Engine/Source/Runtime/Interchange/Core/Private/Tasks/InterchangeTaskParsing.h
+++ b/Engine/Source/Runtime/Interchange/Core/Private/Tasks/InterchangeTaskParsing.h
@@ -15,10 +15,10 @@ class FTaskParsing
 {
 private:
 	UInterchangeManager* InterchangeManager;
-	TWeakPtr<Interchange::FImportAsyncHelper> WeakAsyncHelper;
+	TWeakPtr<Interchange::FImportAsyncHelper, ESPMode::ThreadSafe> WeakAsyncHelper;
 	FString PackageBasePath;
 public:
-	FTaskParsing(UInterchangeManager* InInterchangeManager, const FString& InPackageBasePath, TWeakPtr<Interchange::FImportAsyncHelper> InAsyncHelper)
+	FTaskParsing(UInterchangeManager* InInterchangeManager, const FString& InPackageBasePath, TWeakPtr<Interchange::FImportAsyncHelper, ESPMode::ThreadSafe> InAsyncHelper)
 		: InterchangeManager(InInterchangeManager)
 		, WeakAsyncHelper(InAsyncHelper)
 		, PackageBasePath(InPackageBasePath)
diff --git a/Engine/Source/Runtime/Interchange/Core/Private/Tasks/InterchangeTaskPipeline.cpp b/Engine/Source/Runtime/Interchange/Core/Private/Tasks/InterchangeTaskPipeline.cpp
--- a/Engine/Source/Runtime/Interchange/Core/Private/Tasks/InterchangeTaskPipeline.cpp
+++ b/Engine/Source/Runtime/Interchange/Core/Private/Tasks/InterchangeTaskPipeline.cpp
@@ -12,7 +12,7 @@
 
 void Interchange::FTaskPipeline::DoTask(ENamedThreads::Type CurrentThread, const FGraphEventRef& MyCompletionGraphEvent)
 {
-	TSharedPtr<Interchange::FImportAsyncHelper> AsyncHelper = WeakAsyncHelper.Pin();
+	TSharedPtr<Interchange::FImportAsyncHelper, ESPMode::ThreadSafe> AsyncHelper = WeakAsyncHelper.Pin();
 	check(AsyncHelper.IsValid());
 
 	if (UInterchangePipelineBase* Pipeline = PipelineBase.Get())
diff --git a/Engine/Source/Runtime/Interchange/Core/Private/Tasks/InterchangeTaskPipeline.h b/Engine/Source/Runtime/Interchange/Core/Private/Tasks/InterchangeTaskPipeline.h
--- a/Engine/Source/Runtime/Interchange/Core/Private/Tasks/InterchangeTaskPipeline.h
+++ b/Engine/Source/Runtime/Interchange/Core/Private/Tasks/InterchangeTaskPipeline.h
@@ -16,9 +16,9 @@ class FTaskPipeline
 {
 private:
 	TWeakObjectPtr<UInterchangePipelineBase> PipelineBase;
-	TWeakPtr<Interchange::FImportAsyncHelper> WeakAsyncHelper;
+	TWeakPtr<Interchange::FImportAsyncHelper, ESPMode::ThreadSafe> WeakAsyncHelper;
 public:
-	FTaskPipeline(TWeakObjectPtr<UInterchangePipelineBase> InPipelineBase, TWeakPtr<Interchange::FImportAsyncHelper> InAsyncHelper)
+	FTaskPipeline(TWeakObjectPtr<UInterchangePipelineBase> InPipelineBase, TWeakPtr<Interchange::FImportAsyncHelper, ESPMode::ThreadSafe> InAsyncHelper)
 		: PipelineBase(InPipelineBase)
 		, WeakAsyncHelper(InAsyncHelper)
 	{
diff --git a/Engine/Source/Runtime/Interchange/Core/Private/Tasks/InterchangeTaskTranslator.cpp b/Engine/Source/Runtime/Interchange/Core/Private/Tasks/InterchangeTaskTranslator.cpp
--- a/Engine/Source/Runtime/Interchange/Core/Private/Tasks/InterchangeTaskTranslator.cpp
+++ b/Engine/Source/Runtime/Interchange/Core/Private/Tasks/InterchangeTaskTranslator.cpp
@@ -13,7 +13,7 @@
 
 void Interchange::FTaskTranslator::DoTask(ENamedThreads::Type CurrentThread, const FGraphEventRef& MyCompletionGraphEvent)
 {
-	TSharedPtr<Interchange::FImportAsyncHelper> AsyncHelper = WeakAsyncHelper.Pin();
+	TSharedPtr<Interchange::FImportAsyncHelper, ESPMode::ThreadSafe> AsyncHelper = WeakAsyncHelper.Pin();
 	check(AsyncHelper.IsValid());
 
 	UInterchangeTranslatorBase* Translator = AsyncHelper->Translators.IsValidIndex(SourceIndex) ? AsyncHelper->Translators[SourceIndex] : nullptr;
diff --git a/Engine/Source/Runtime/Interchange/Core/Private/Tasks/InterchangeTaskTranslator.h b/Engine/Source/Runtime/Interchange/Core/Private/Tasks/InterchangeTaskTranslator.h
--- a/Engine/Source/Runtime/Interchange/Core/Private/Tasks/InterchangeTaskTranslator.h
+++ b/Engine/Source/Runtime/Interchange/Core/Private/Tasks/InterchangeTaskTranslator.h
@@ -15,10 +15,10 @@ class FTaskTranslator
 {
 private:
 	int32 SourceIndex = INDEX_NONE;
-	TWeakPtr<Interchange::FImportAsyncHelper> WeakAsyncHelper;
+	TWeakPtr<Interchange::FImportAsyncHelper, ESPMode::ThreadSafe> WeakAsyncHelper;
 
 public:
-	FTaskTranslator(int32 InSourceIndex, TWeakPtr<Interchange::FImportAsyncHelper> InAsyncHelper)
+	FTaskTranslator(int32 InSourceIndex, TWeakPtr<Interchange::FImportAsyncHelper, ESPMode::ThreadSafe> InAsyncHelper)
 		: SourceIndex(InSourceIndex)
 		, WeakAsyncHelper(InAsyncHelper)
 	{
diff --git a/Engine/Source/Runtime/Interchange/Core/Public/InterchangeManager.h b/Engine/Source/Runtime/Interchange/Core/Public/InterchangeManager.h
--- a/Engine/Source/Runtime/Interchange/Core/Public/InterchangeManager.h
+++ b/Engine/Source/Runtime/Interchange/Core/Public/InterchangeManager.h
@@ -297,10 +297,10 @@ public:
 	 * Return an FImportAsynHelper pointer. The pointer is deleted when ReleaseAsyncHelper is call.
 	 * @param Data - The data we want to pass to the different import tasks
 	 */
-	TWeakPtr<Interchange::FImportAsyncHelper> CreateAsyncHelper(const Interchange::FImportAsyncHelperData& Data);
+	TWeakPtr<Interchange::FImportAsyncHelper, ESPMode::ThreadSafe> CreateAsyncHelper(const Interchange::FImportAsyncHelperData& Data);
 
 	/** Delete the specified AsyncHelper and remove it from the array that was holding it. */
-	void ReleaseAsyncHelper(TWeakPtr<Interchange::FImportAsyncHelper> AsyncHelper);
+	void ReleaseAsyncHelper(TWeakPtr<Interchange::FImportAsyncHelper, ESPMode::ThreadSafe> AsyncHelper);
 
 protected:
 
@@ -323,7 +323,7 @@ protected:
 
 private:
 	//By using pointer, there is no issue if the array get resize
-	TArray<TSharedPtr<Interchange::FImportAsyncHelper> > ImportTasks;
+	TArray<TSharedPtr<Interchange::FImportAsyncHelper, ESPMode::ThreadSafe> > ImportTasks;
 
 	TSharedPtr<FAsyncTaskNotification> Notification = nullptr;
 
