Make sure all interchange namespace respect the coding standard

diff --git a/Engine/Source/Runtime/Interchange/Core/Private/Types/AttributeStorage.cpp b/Engine/Source/Runtime/Interchange/Core/Private/Types/AttributeStorage.cpp
--- a/Engine/Source/Runtime/Interchange/Core/Private/Types/AttributeStorage.cpp
+++ b/Engine/Source/Runtime/Interchange/Core/Private/Types/AttributeStorage.cpp
@@ -7,465 +7,468 @@
 #include "Misc/SecureHash.h"
 
 //Interchange namespace
-namespace Interchange
+namespace UE
 {
-	void LogAttributeStorageErrors(const EAttributeStorageResult Result, const FString OperationName, const FAttributeKey AttributeKey)
+	namespace Interchange
 	{
-		//////////////////////////////////////////////////////////////////////////
-		//Errors
-		if (HasAttributeStorageResult(Result, EAttributeStorageResult::Operation_Error_AttributeAllocationCorrupted))
+		void LogAttributeStorageErrors(const EAttributeStorageResult Result, const FString OperationName, const FAttributeKey AttributeKey)
 		{
-			UE_LOG(LogInterchangeCore, Error, TEXT("Attribute storage operation [%s] Key[%s]: Storage is corrupted."), *OperationName, *(AttributeKey.Key.ToString()));
-		}
+			//////////////////////////////////////////////////////////////////////////
+			//Errors
+			if (HasAttributeStorageResult(Result, EAttributeStorageResult::Operation_Error_AttributeAllocationCorrupted))
+			{
+				UE_LOG(LogInterchangeCore, Error, TEXT("Attribute storage operation [%s] Key[%s]: Storage is corrupted."), *OperationName, *(AttributeKey.Key.ToString()));
+			}
 
-		//////////////////////////////////////////////////////////////////////////
-		//Warning
-		if (HasAttributeStorageResult(Result, EAttributeStorageResult::Operation_Error_CannotFoundKey))
-		{
-			UE_LOG(LogInterchangeCore, Warning, TEXT("Attribute storage operation [%s] Key[%s]: Cannot find attribute key."), *OperationName, *(AttributeKey.Key.ToString()));
-		}
+			//////////////////////////////////////////////////////////////////////////
+			//Warning
+			if (HasAttributeStorageResult(Result, EAttributeStorageResult::Operation_Error_CannotFoundKey))
+			{
+				UE_LOG(LogInterchangeCore, Warning, TEXT("Attribute storage operation [%s] Key[%s]: Cannot find attribute key."), *OperationName, *(AttributeKey.Key.ToString()));
+			}
 
-		if (HasAttributeStorageResult(Result, EAttributeStorageResult::Operation_Error_CannotOverrideAttribute))
-		{
-			UE_LOG(LogInterchangeCore, Warning, TEXT("Attribute storage operation [%s] Key[%s]: Cannot override attribute."), *OperationName, *(AttributeKey.Key.ToString()));
-		}
+			if (HasAttributeStorageResult(Result, EAttributeStorageResult::Operation_Error_CannotOverrideAttribute))
+			{
+				UE_LOG(LogInterchangeCore, Warning, TEXT("Attribute storage operation [%s] Key[%s]: Cannot override attribute."), *OperationName, *(AttributeKey.Key.ToString()));
+			}
 
-		if (HasAttributeStorageResult(Result, EAttributeStorageResult::Operation_Error_CannotRemoveAttribute))
-		{
-			UE_LOG(LogInterchangeCore, Warning, TEXT("Attribute storage operation [%s] Key[%s]: Cannot remove an attribute."), *OperationName, *(AttributeKey.Key.ToString()));
-		}
+			if (HasAttributeStorageResult(Result, EAttributeStorageResult::Operation_Error_CannotRemoveAttribute))
+			{
+				UE_LOG(LogInterchangeCore, Warning, TEXT("Attribute storage operation [%s] Key[%s]: Cannot remove an attribute."), *OperationName, *(AttributeKey.Key.ToString()));
+			}
 
-		if (HasAttributeStorageResult(Result, EAttributeStorageResult::Operation_Error_WrongSize))
-		{
-			UE_LOG(LogInterchangeCore, Warning, TEXT("Attribute storage operation [%s] Key[%s]: Stored attribute value size do not match parameter value size."), *OperationName, *(AttributeKey.Key.ToString()));
-		}
+			if (HasAttributeStorageResult(Result, EAttributeStorageResult::Operation_Error_WrongSize))
+			{
+				UE_LOG(LogInterchangeCore, Warning, TEXT("Attribute storage operation [%s] Key[%s]: Stored attribute value size do not match parameter value size."), *OperationName, *(AttributeKey.Key.ToString()));
+			}
 
-		if (HasAttributeStorageResult(Result, EAttributeStorageResult::Operation_Error_WrongType))
-		{
-			UE_LOG(LogInterchangeCore, Warning, TEXT("Attribute storage operation [%s] Key[%s]: Stored attribute value type do not match parameter value type."), *OperationName, *(AttributeKey.Key.ToString()));
+			if (HasAttributeStorageResult(Result, EAttributeStorageResult::Operation_Error_WrongType))
+			{
+				UE_LOG(LogInterchangeCore, Warning, TEXT("Attribute storage operation [%s] Key[%s]: Stored attribute value type do not match parameter value type."), *OperationName, *(AttributeKey.Key.ToString()));
+			}
+
+			if (HasAttributeStorageResult(Result, EAttributeStorageResult::Operation_Error_UnsupportedType))
+			{
+				UE_LOG(LogInterchangeCore, Warning, TEXT("Attribute storage operation [%s] Key[%s]: The storage do not support this type of value, please add unsupported type to EAttributeTypes."), *OperationName, *(AttributeKey.Key.ToString()));
+			}
+
+			if (HasAttributeStorageResult(Result, EAttributeStorageResult::Operation_Error_InvalidStorage))
+			{
+				UE_LOG(LogInterchangeCore, Warning, TEXT("Attribute storage operation [%s] Key[%s]: The storage is invalid (NULL)."), *OperationName, *(AttributeKey.Key.ToString()));
+			}
+
+			if (HasAttributeStorageResult(Result, EAttributeStorageResult::Operation_Error_InvalidMultiSizeValueData))
+			{
+				UE_LOG(LogInterchangeCore, Warning, TEXT("Attribute storage operation [%s] Key[%s]: Cannot retrieve the multisize value data pointer."), *OperationName, *(AttributeKey.Key.ToString()));
+			}
 		}
 
-		if (HasAttributeStorageResult(Result, EAttributeStorageResult::Operation_Error_UnsupportedType))
+		FAttributeStorage::FAttributeStorage(const FAttributeStorage& Other)
 		{
-			UE_LOG(LogInterchangeCore, Warning, TEXT("Attribute storage operation [%s] Key[%s]: The storage do not support this type of value, please add unsupported type to EAttributeTypes."), *OperationName, *(AttributeKey.Key.ToString()));
+			//Lock both Storage mutex
+			FScopeLock ScopeLock(&StorageMutex);
+			FScopeLock ScopeLockOther(&Other.StorageMutex);
+
+			//Copy all value
+			AttributeAllocationTable = Other.AttributeAllocationTable;
+			AttributeStorage = Other.AttributeStorage;
+			FragmentedMemoryCost = Other.FragmentedMemoryCost;
+			DefragRatio = Other.DefragRatio;
 		}
 
-		if (HasAttributeStorageResult(Result, EAttributeStorageResult::Operation_Error_InvalidStorage))
+		FAttributeStorage& FAttributeStorage::operator=(const FAttributeStorage& Other)
 		{
-			UE_LOG(LogInterchangeCore, Warning, TEXT("Attribute storage operation [%s] Key[%s]: The storage is invalid (NULL)."), *OperationName, *(AttributeKey.Key.ToString()));
+			//Lock both Storage mutex
+			FScopeLock ScopeLock(&StorageMutex);
+			FScopeLock ScopeLockOther(&Other.StorageMutex);
+
+			//Copy all value
+			AttributeAllocationTable = Other.AttributeAllocationTable;
+			AttributeStorage = Other.AttributeStorage;
+			FragmentedMemoryCost = Other.FragmentedMemoryCost;
+			DefragRatio = Other.DefragRatio;
+
+			return *this;
 		}
 
-		if (HasAttributeStorageResult(Result, EAttributeStorageResult::Operation_Error_InvalidMultiSizeValueData))
+		EAttributeStorageResult FAttributeStorage::UnregisterAttribute(const FAttributeKey& ElementAttributeKey)
 		{
-			UE_LOG(LogInterchangeCore, Warning, TEXT("Attribute storage operation [%s] Key[%s]: Cannot retrieve the multisize value data pointer."), *OperationName, *(AttributeKey.Key.ToString()));
-		}
-	}
+			//Lock the storage
+			FScopeLock ScopeLock(&StorageMutex);
 
-	FAttributeStorage::FAttributeStorage(const FAttributeStorage& Other)
-	{
-		//Lock both Storage mutex
-		FScopeLock ScopeLock(&StorageMutex);
-		FScopeLock ScopeLockOther(&Other.StorageMutex);
-
-		//Copy all value
-		AttributeAllocationTable = Other.AttributeAllocationTable;
-		AttributeStorage = Other.AttributeStorage;
-		FragmentedMemoryCost = Other.FragmentedMemoryCost;
-		DefragRatio = Other.DefragRatio;
-	}
-
-	FAttributeStorage& FAttributeStorage::operator=(const FAttributeStorage& Other)
-	{
-		//Lock both Storage mutex
-		FScopeLock ScopeLock(&StorageMutex);
-		FScopeLock ScopeLockOther(&Other.StorageMutex);
+			FAttributeAllocationInfo* AttributeAllocationInfo = AttributeAllocationTable.Find(ElementAttributeKey);
+			if (!AttributeAllocationInfo)
+			{
+				return EAttributeStorageResult::Operation_Error_CannotFoundKey;
+			}
 
-		//Copy all value
-		AttributeAllocationTable = Other.AttributeAllocationTable;
-		AttributeStorage = Other.AttributeStorage;
-		FragmentedMemoryCost = Other.FragmentedMemoryCost;
-		DefragRatio = Other.DefragRatio;
+			FragmentedMemoryCost += AttributeAllocationInfo->Size;
+			if (AttributeAllocationTable.Remove(ElementAttributeKey) == 0)
+			{
+				return EAttributeStorageResult::Operation_Error_CannotRemoveAttribute;
+			}
 
-		return *this;
-	}
+			DefragInternal();
 
-	EAttributeStorageResult FAttributeStorage::UnregisterAttribute(const FAttributeKey& ElementAttributeKey)
-	{
-		//Lock the storage
-		FScopeLock ScopeLock(&StorageMutex);
-		
-		FAttributeAllocationInfo* AttributeAllocationInfo = AttributeAllocationTable.Find(ElementAttributeKey);
-		if (!AttributeAllocationInfo)
-		{
-			return EAttributeStorageResult::Operation_Error_CannotFoundKey;
+			return EAttributeStorageResult::Operation_Success;
 		}
 
-		FragmentedMemoryCost += AttributeAllocationInfo->Size;
-		if (AttributeAllocationTable.Remove(ElementAttributeKey) == 0)
+		EAttributeTypes FAttributeStorage::GetAttributeType(const FAttributeKey& ElementAttributeKey)
 		{
-			return EAttributeStorageResult::Operation_Error_CannotRemoveAttribute;
+			//Lock the storage
+			FScopeLock ScopeLock(&StorageMutex);
+			const FAttributeAllocationInfo* AttributeAllocationInfo = AttributeAllocationTable.Find(ElementAttributeKey);
+			if (AttributeAllocationInfo)
+			{
+				return AttributeAllocationInfo->Type;
+			}
+			return EAttributeTypes::None;
 		}
 
-		DefragInternal();
-		
-		return EAttributeStorageResult::Operation_Success;
-	}
-	
-	EAttributeTypes FAttributeStorage::GetAttributeType(const FAttributeKey& ElementAttributeKey)
-	{
-		//Lock the storage
-		FScopeLock ScopeLock(&StorageMutex);
-		const FAttributeAllocationInfo* AttributeAllocationInfo = AttributeAllocationTable.Find(ElementAttributeKey);
-		if (AttributeAllocationInfo)
+		bool FAttributeStorage::ContainAttribute(const FAttributeKey& ElementAttributeKey) const
 		{
-			return AttributeAllocationInfo->Type;
-		}
-		return EAttributeTypes::None;
-	}	
+			//Lock the storage
+			FScopeLock ScopeLock(&StorageMutex);
 
-	bool FAttributeStorage::ContainAttribute(const FAttributeKey& ElementAttributeKey) const
-	{
-		//Lock the storage
-		FScopeLock ScopeLock(&StorageMutex);
-
-		return AttributeAllocationTable.Contains(ElementAttributeKey);
-	}
-
-	void FAttributeStorage::GetAttributeKeys(TArray<FAttributeKey>& AttributeKeys) const
-	{
-		FScopeLock ScopeLock(&StorageMutex);
-		AttributeAllocationTable.GetKeys(AttributeKeys);
-	}
+			return AttributeAllocationTable.Contains(ElementAttributeKey);
+		}
 
-	FGuid FAttributeStorage::GetAttributeHash(const FAttributeKey& ElementAttributeKey) const
-	{
-		FGuid AttributeHash;
-		GetAttributeHash(ElementAttributeKey, AttributeHash);
-		return AttributeHash;
-	}
+		void FAttributeStorage::GetAttributeKeys(TArray<FAttributeKey>& AttributeKeys) const
+		{
+			FScopeLock ScopeLock(&StorageMutex);
+			AttributeAllocationTable.GetKeys(AttributeKeys);
+		}
 
-	bool FAttributeStorage::GetAttributeHash(const FAttributeKey& ElementAttributeKey, FGuid& OutGuid) const
-	{
-		//Lock the storage
-		FScopeLock ScopeLock(&StorageMutex);
-		
-		const FAttributeAllocationInfo* AttributeAllocationInfo = AttributeAllocationTable.Find(ElementAttributeKey);
-		if (AttributeAllocationInfo)
+		FGuid FAttributeStorage::GetAttributeHash(const FAttributeKey& ElementAttributeKey) const
 		{
-			OutGuid = AttributeAllocationInfo->Hash;
-			return true;
+			FGuid AttributeHash;
+			GetAttributeHash(ElementAttributeKey, AttributeHash);
+			return AttributeHash;
 		}
-		return false;
-	}
 
-	FGuid FAttributeStorage::GetStorageHash() const
-	{
-		//Lock the storage
-		FScopeLock ScopeLock(&StorageMutex);
-
-		TArray<FAttributeKey> OrderedKeys;
-		AttributeAllocationTable.GetKeys(OrderedKeys);
-		OrderedKeys.Sort();
-		FSHA1 Sha;
-		for (FAttributeKey Key : OrderedKeys)
+		bool FAttributeStorage::GetAttributeHash(const FAttributeKey& ElementAttributeKey, FGuid& OutGuid) const
 		{
-			const FAttributeAllocationInfo& AttributeAllocationInfo = AttributeAllocationTable.FindChecked(Key);
-			//Skip Attribute that has the no hash flag
-			if(HasAttributeProperty(AttributeAllocationInfo.Property, EAttributeProperty::NoHash))
+			//Lock the storage
+			FScopeLock ScopeLock(&StorageMutex);
+
+			const FAttributeAllocationInfo* AttributeAllocationInfo = AttributeAllocationTable.Find(ElementAttributeKey);
+			if (AttributeAllocationInfo)
 			{
-				continue;
+				OutGuid = AttributeAllocationInfo->Hash;
+				return true;
 			}
-			uint32 GuidData[4];
-			GuidData[0] = AttributeAllocationInfo.Hash.A;
-			GuidData[1] = AttributeAllocationInfo.Hash.B;
-			GuidData[2] = AttributeAllocationInfo.Hash.C;
-			GuidData[3] = AttributeAllocationInfo.Hash.D;
-			Sha.Update(reinterpret_cast<uint8*>(&GuidData[0]), 16);
+			return false;
 		}
-		Sha.Final();
-		// Retrieve the hash and use it to construct a pseudo-GUID. 
-		uint32 Hash[5];
-		Sha.GetHash(reinterpret_cast<uint8*>(Hash));
-		return FGuid(Hash[0] ^ Hash[4], Hash[1], Hash[2], Hash[3]);
-	}
-
-	void FAttributeStorage::CompareStorage(const FAttributeStorage& BaseStorage, const FAttributeStorage& VersionStorage, TArray<FAttributeKey>& RemovedAttributes, TArray<FAttributeKey>& AddedAttributes, TArray<FAttributeKey>& ModifiedAttributes)
-	{
-		//Lock the storage
-		FScopeLock ScopeLockVersion(&VersionStorage.StorageMutex);
-		//Lock the storage
-		FScopeLock ScopeLockBase(&BaseStorage.StorageMutex);
-		
-		for (const auto KvpBase : BaseStorage.AttributeAllocationTable)
+
+		FGuid FAttributeStorage::GetStorageHash() const
 		{
-			FAttributeKey KeyBase = KvpBase.Key;
-			const FAttributeAllocationInfo* AttributeAllocationInfoVersion = VersionStorage.AttributeAllocationTable.Find(KeyBase);
-			if (!AttributeAllocationInfoVersion)
+			//Lock the storage
+			FScopeLock ScopeLock(&StorageMutex);
+
+			TArray<FAttributeKey> OrderedKeys;
+			AttributeAllocationTable.GetKeys(OrderedKeys);
+			OrderedKeys.Sort();
+			FSHA1 Sha;
+			for (FAttributeKey Key : OrderedKeys)
 			{
-				//Add the attribute to RemovedAttributes
-				RemovedAttributes.Add(KeyBase);
-			}
-			else if(KvpBase.Value.Hash != AttributeAllocationInfoVersion->Hash)
-			{
-				//Add the attribute to ModifiedAttributes
-				ModifiedAttributes.Add(KeyBase);
+				const FAttributeAllocationInfo& AttributeAllocationInfo = AttributeAllocationTable.FindChecked(Key);
+				//Skip Attribute that has the no hash flag
+				if (HasAttributeProperty(AttributeAllocationInfo.Property, EAttributeProperty::NoHash))
+				{
+					continue;
+				}
+				uint32 GuidData[4];
+				GuidData[0] = AttributeAllocationInfo.Hash.A;
+				GuidData[1] = AttributeAllocationInfo.Hash.B;
+				GuidData[2] = AttributeAllocationInfo.Hash.C;
+				GuidData[3] = AttributeAllocationInfo.Hash.D;
+				Sha.Update(reinterpret_cast<uint8*>(&GuidData[0]), 16);
 			}
+			Sha.Final();
+			// Retrieve the hash and use it to construct a pseudo-GUID. 
+			uint32 Hash[5];
+			Sha.GetHash(reinterpret_cast<uint8*>(Hash));
+			return FGuid(Hash[0] ^ Hash[4], Hash[1], Hash[2], Hash[3]);
 		}
 
-		for (const auto KvpVersion : VersionStorage.AttributeAllocationTable)
+		void FAttributeStorage::CompareStorage(const FAttributeStorage& BaseStorage, const FAttributeStorage& VersionStorage, TArray<FAttributeKey>& RemovedAttributes, TArray<FAttributeKey>& AddedAttributes, TArray<FAttributeKey>& ModifiedAttributes)
 		{
-			FAttributeKey KeyVersion = KvpVersion.Key;
-			const FAttributeAllocationInfo* AttributeAllocationInfoBase = BaseStorage.AttributeAllocationTable.Find(KeyVersion);
-			if (!AttributeAllocationInfoBase)
+			//Lock the storage
+			FScopeLock ScopeLockVersion(&VersionStorage.StorageMutex);
+			//Lock the storage
+			FScopeLock ScopeLockBase(&BaseStorage.StorageMutex);
+
+			for (const auto KvpBase : BaseStorage.AttributeAllocationTable)
 			{
-				//Add the attribute to RemovedAttributes
-				AddedAttributes.Add(KeyVersion);
+				FAttributeKey KeyBase = KvpBase.Key;
+				const FAttributeAllocationInfo* AttributeAllocationInfoVersion = VersionStorage.AttributeAllocationTable.Find(KeyBase);
+				if (!AttributeAllocationInfoVersion)
+				{
+					//Add the attribute to RemovedAttributes
+					RemovedAttributes.Add(KeyBase);
+				}
+				else if (KvpBase.Value.Hash != AttributeAllocationInfoVersion->Hash)
+				{
+					//Add the attribute to ModifiedAttributes
+					ModifiedAttributes.Add(KeyBase);
+				}
 			}
-		}
-	}
 
-	void FAttributeStorage::SetDefragRatio(const float InDefragRatio)
-	{
-		FScopeLock ScopeLock(&StorageMutex);
-		DefragRatio = InDefragRatio;
-		DefragInternal();
-	}
+			for (const auto KvpVersion : VersionStorage.AttributeAllocationTable)
+			{
+				FAttributeKey KeyVersion = KvpVersion.Key;
+				const FAttributeAllocationInfo* AttributeAllocationInfoBase = BaseStorage.AttributeAllocationTable.Find(KeyVersion);
+				if (!AttributeAllocationInfoBase)
+				{
+					//Add the attribute to RemovedAttributes
+					AddedAttributes.Add(KeyVersion);
+				}
+			}
+		}
 
-	void FAttributeStorage::Reserve(int64 NewAttributeCount, int64 NewStorageSize)
-	{
-		if (NewAttributeCount > 0)
+		void FAttributeStorage::SetDefragRatio(const float InDefragRatio)
 		{
-			const int64 ReserveCount = AttributeAllocationTable.Num() + NewAttributeCount;
-			AttributeAllocationTable.Reserve(ReserveCount);
+			FScopeLock ScopeLock(&StorageMutex);
+			DefragRatio = InDefragRatio;
+			DefragInternal();
 		}
-		if (NewStorageSize > 0)
+
+		void FAttributeStorage::Reserve(int64 NewAttributeCount, int64 NewStorageSize)
 		{
-			const int64 ReserveCount = AttributeStorage.Num() + NewStorageSize;
-			AttributeStorage.Reserve(ReserveCount);
+			if (NewAttributeCount > 0)
+			{
+				const int64 ReserveCount = AttributeAllocationTable.Num() + NewAttributeCount;
+				AttributeAllocationTable.Reserve(ReserveCount);
+			}
+			if (NewStorageSize > 0)
+			{
+				const int64 ReserveCount = AttributeStorage.Num() + NewStorageSize;
+				AttributeStorage.Reserve(ReserveCount);
+			}
 		}
-	}
-
-	void FAttributeStorage::ExtractFStringAttributeFromStorage(const uint8* StorageData, const FAttributeAllocationInfo* AttributeAllocationInfo, FString& OutValue) const
-	{
-		//Allocate the ByteArray
-		const uint64 NumberOfChar = AttributeAllocationInfo->Size / sizeof(TCHAR);
-		check(NumberOfChar > 0);
 
-
-		//Empty and pre allocate the FString buffer (TArray<TCHAR>)
-		OutValue.Reset(NumberOfChar - 1);
-		//SLOW! We have to add the characters 1 x 1, there is no AddZeroed or similar stuff. We do not have a string to use the operator= which do what we need here
-		//We go up to NumberOfChar -1 because there is always a null terminating character add at the end of the FString
-		for (int32 CharIndex = 0; CharIndex < NumberOfChar - 1; ++CharIndex)
+		void FAttributeStorage::ExtractFStringAttributeFromStorage(const uint8* StorageData, const FAttributeAllocationInfo* AttributeAllocationInfo, FString& OutValue) const
 		{
-			//Adding spaces, we cannot add null character because the null character are not added
-			OutValue.AppendChar(' ');
-		}
+			//Allocate the ByteArray
+			const uint64 NumberOfChar = AttributeAllocationInfo->Size / sizeof(TCHAR);
+			check(NumberOfChar > 0);
 
-		//Copy into the buffer
-		FMemory::Memcpy((uint8*)(*OutValue), &StorageData[AttributeAllocationInfo->Offset], AttributeAllocationInfo->Size);
-	}
 
-	EAttributeStorageResult FAttributeStorage::GetAttribute(const FAttributeKey& ElementAttributeKey, FString& OutValue, TSpecializeType<FString >) const
-	{
-		//Lock the storage
-		FScopeLock ScopeLock(&StorageMutex);
+			//Empty and pre allocate the FString buffer (TArray<TCHAR>)
+			OutValue.Reset(NumberOfChar - 1);
+			//SLOW! We have to add the characters 1 x 1, there is no AddZeroed or similar stuff. We do not have a string to use the operator= which do what we need here
+			//We go up to NumberOfChar -1 because there is always a null terminating character add at the end of the FString
+			for (int32 CharIndex = 0; CharIndex < NumberOfChar - 1; ++CharIndex)
+			{
+				//Adding spaces, we cannot add null character because the null character are not added
+				OutValue.AppendChar(' ');
+			}
 
-		if (!IsSupportedType(OutValue))
-		{
-			return EAttributeStorageResult::Operation_Error_UnsupportedType;
+			//Copy into the buffer
+			FMemory::Memcpy((uint8*)(*OutValue), &StorageData[AttributeAllocationInfo->Offset], AttributeAllocationInfo->Size);
 		}
 
-		const EAttributeTypes ValueType = TAttributeTypeTraits<FString>::GetType();
-		const FAttributeAllocationInfo* AttributeAllocationInfo = AttributeAllocationTable.Find(ElementAttributeKey);
-		if (AttributeAllocationInfo)
+		EAttributeStorageResult FAttributeStorage::GetAttribute(const FAttributeKey& ElementAttributeKey, FString& OutValue, TSpecializeType<FString >) const
 		{
-			if (AttributeAllocationInfo->Type != ValueType)
+			//Lock the storage
+			FScopeLock ScopeLock(&StorageMutex);
+
+			if (!IsSupportedType(OutValue))
 			{
-				return EAttributeStorageResult::Operation_Error_WrongType;
+				return EAttributeStorageResult::Operation_Error_UnsupportedType;
 			}
-			if (!AttributeStorage.IsValidIndex(AttributeAllocationInfo->Offset))
+
+			const EAttributeTypes ValueType = TAttributeTypeTraits<FString>::GetType();
+			const FAttributeAllocationInfo* AttributeAllocationInfo = AttributeAllocationTable.Find(ElementAttributeKey);
+			if (AttributeAllocationInfo)
 			{
-				return EAttributeStorageResult::Operation_Error_AttributeAllocationCorrupted;
+				if (AttributeAllocationInfo->Type != ValueType)
+				{
+					return EAttributeStorageResult::Operation_Error_WrongType;
+				}
+				if (!AttributeStorage.IsValidIndex(AttributeAllocationInfo->Offset))
+				{
+					return EAttributeStorageResult::Operation_Error_AttributeAllocationCorrupted;
+				}
+			}
+			else
+			{
+				//The key do not exist
+				return EAttributeStorageResult::Operation_Error_CannotFoundKey;
 			}
-		}
-		else
-		{
-			//The key do not exist
-			return EAttributeStorageResult::Operation_Error_CannotFoundKey;
-		}
 
-		if (AttributeAllocationInfo->Size == 0)
-		{
-			OutValue = TEXT("");
+			if (AttributeAllocationInfo->Size == 0)
+			{
+				OutValue = TEXT("");
+				return EAttributeStorageResult::Operation_Success;
+			}
+
+			ExtractFStringAttributeFromStorage(AttributeStorage.GetData(), AttributeAllocationInfo, OutValue);
+
 			return EAttributeStorageResult::Operation_Success;
 		}
 
-		ExtractFStringAttributeFromStorage(AttributeStorage.GetData(), AttributeAllocationInfo, OutValue);
-
-		return EAttributeStorageResult::Operation_Success;
-	}
+		EAttributeStorageResult FAttributeStorage::GetAttribute(const FAttributeKey& ElementAttributeKey, FName& OutValue, TSpecializeType<FName >) const
+		{
+			//Lock the storage
+			FScopeLock ScopeLock(&StorageMutex);
 
-	EAttributeStorageResult FAttributeStorage::GetAttribute(const FAttributeKey& ElementAttributeKey, FName& OutValue, TSpecializeType<FName >) const
-	{
-		//Lock the storage
-		FScopeLock ScopeLock(&StorageMutex);
+			if (!IsSupportedType(OutValue))
+			{
+				return EAttributeStorageResult::Operation_Error_UnsupportedType;
+			}
 
-		if (!IsSupportedType(OutValue))
-		{
-			return EAttributeStorageResult::Operation_Error_UnsupportedType;
-		}
+			const EAttributeTypes ValueType = TAttributeTypeTraits<FName>::GetType();
 
-		const EAttributeTypes ValueType = TAttributeTypeTraits<FName>::GetType();
 
-		
-		const FAttributeAllocationInfo* AttributeAllocationInfo = AttributeAllocationTable.Find(ElementAttributeKey);
-		if (AttributeAllocationInfo)
-		{
-			if (AttributeAllocationInfo->Type != ValueType)
+			const FAttributeAllocationInfo* AttributeAllocationInfo = AttributeAllocationTable.Find(ElementAttributeKey);
+			if (AttributeAllocationInfo)
 			{
-				return EAttributeStorageResult::Operation_Error_WrongType;
+				if (AttributeAllocationInfo->Type != ValueType)
+				{
+					return EAttributeStorageResult::Operation_Error_WrongType;
+				}
+				if (!AttributeStorage.IsValidIndex(AttributeAllocationInfo->Offset))
+				{
+					return EAttributeStorageResult::Operation_Error_AttributeAllocationCorrupted;
+				}
 			}
-			if (!AttributeStorage.IsValidIndex(AttributeAllocationInfo->Offset))
+			else
 			{
-				return EAttributeStorageResult::Operation_Error_AttributeAllocationCorrupted;
+				//The key do not exist
+				return EAttributeStorageResult::Operation_Error_CannotFoundKey;
 			}
-		}
-		else
-		{
-			//The key do not exist
-			return EAttributeStorageResult::Operation_Error_CannotFoundKey;
-		}
 
-		
-		if (AttributeAllocationInfo->Size == 0)
-		{
-			OutValue = NAME_None;
-			return EAttributeStorageResult::Operation_Success;
-		}
 
-		FString ValueStr;
+			if (AttributeAllocationInfo->Size == 0)
+			{
+				OutValue = NAME_None;
+				return EAttributeStorageResult::Operation_Success;
+			}
 
-		//Share the code with FString from here
-		ExtractFStringAttributeFromStorage(AttributeStorage.GetData(), AttributeAllocationInfo, ValueStr);
+			FString ValueStr;
 
-		//Create the FName and copy it to OutValue
-		OutValue = FName(*ValueStr);
+			//Share the code with FString from here
+			ExtractFStringAttributeFromStorage(AttributeStorage.GetData(), AttributeAllocationInfo, ValueStr);
 
-		return EAttributeStorageResult::Operation_Success;
-	}
+			//Create the FName and copy it to OutValue
+			OutValue = FName(*ValueStr);
 
-	void FAttributeStorage::DefragInternal()
-	{
-		//Defrag if the fragmented memory cost is bigger then the DefragRatio.
-		if (FragmentedMemoryCost < 11 || FragmentedMemoryCost < AttributeStorage.Num()* DefragRatio)
-		{
-			return;
+			return EAttributeStorageResult::Operation_Success;
 		}
-		
-		//Sorts the allocation table per offset since we want to defrag using memory block
-		//The ValueSort is also calling CompactStable which remove all holes in the TMap
-		struct FAttributeAllocationInfoSmaller
-		{
-			FORCEINLINE bool operator()(const FAttributeAllocationInfo& A, const FAttributeAllocationInfo& B) const
-			{
-				return A.Offset < B.Offset;
-			}
-		};
-		AttributeAllocationTable.ValueSort(FAttributeAllocationInfoSmaller());
-
-		const uint64 AttributeNumber = AttributeAllocationTable.Num();
-		uint8* StorageData = AttributeStorage.GetData();
-		//Number of attribute we have done
-		uint64 AttributeCount = 0;
-		//Current storage offset we want to defrag
-		uint64 CurrentOffset = 0;
-		//Are we build a defrag block
-		bool bBuildBlock = false;
-		//Defrag block start
-		uint64 MemmoveBlockStart = 0;
-		//Replace CurrentOffset when we are moving a block
-		uint64 MemmoveBlockCurrentOffset = 0;
-		for (auto& Kvp : AttributeAllocationTable)
+
+		void FAttributeStorage::DefragInternal()
 		{
-			AttributeCount++;
-			const bool bLast = AttributeCount == AttributeNumber;
-			
-			FAttributeAllocationInfo& AttributeAllocationInfo = Kvp.Value;
-			if (AttributeAllocationInfo.Offset == CurrentOffset)
+			//Defrag if the fragmented memory cost is bigger then the DefragRatio.
+			if (FragmentedMemoryCost < 11 || FragmentedMemoryCost < AttributeStorage.Num() * DefragRatio)
 			{
-				CurrentOffset += AttributeAllocationInfo.Size;
-				continue;
+				return;
 			}
-			check(CurrentOffset < AttributeAllocationInfo.Offset);
-			if(!bBuildBlock)
+
+			//Sorts the allocation table per offset since we want to defrag using memory block
+			//The ValueSort is also calling CompactStable which remove all holes in the TMap
+			struct FAttributeAllocationInfoSmaller
 			{
-				bBuildBlock = true;
-				MemmoveBlockStart = AttributeAllocationInfo.Offset;
-				AttributeAllocationInfo.Offset = CurrentOffset;
-				MemmoveBlockCurrentOffset = MemmoveBlockStart + AttributeAllocationInfo.Size;
-			}
-			else
+				FORCEINLINE bool operator()(const FAttributeAllocationInfo& A, const FAttributeAllocationInfo& B) const
+				{
+					return A.Offset < B.Offset;
+				}
+			};
+			AttributeAllocationTable.ValueSort(FAttributeAllocationInfoSmaller());
+
+			const uint64 AttributeNumber = AttributeAllocationTable.Num();
+			uint8* StorageData = AttributeStorage.GetData();
+			//Number of attribute we have done
+			uint64 AttributeCount = 0;
+			//Current storage offset we want to defrag
+			uint64 CurrentOffset = 0;
+			//Are we build a defrag block
+			bool bBuildBlock = false;
+			//Defrag block start
+			uint64 MemmoveBlockStart = 0;
+			//Replace CurrentOffset when we are moving a block
+			uint64 MemmoveBlockCurrentOffset = 0;
+			for (auto& Kvp : AttributeAllocationTable)
 			{
-				if (AttributeAllocationInfo.Offset == MemmoveBlockCurrentOffset)
+				AttributeCount++;
+				const bool bLast = AttributeCount == AttributeNumber;
+
+				FAttributeAllocationInfo& AttributeAllocationInfo = Kvp.Value;
+				if (AttributeAllocationInfo.Offset == CurrentOffset)
+				{
+					CurrentOffset += AttributeAllocationInfo.Size;
+					continue;
+				}
+				check(CurrentOffset < AttributeAllocationInfo.Offset);
+				if (!bBuildBlock)
 				{
-					AttributeAllocationInfo.Offset -= MemmoveBlockStart - CurrentOffset;
-					MemmoveBlockCurrentOffset += AttributeAllocationInfo.Size;
+					bBuildBlock = true;
+					MemmoveBlockStart = AttributeAllocationInfo.Offset;
+					AttributeAllocationInfo.Offset = CurrentOffset;
+					MemmoveBlockCurrentOffset = MemmoveBlockStart + AttributeAllocationInfo.Size;
 				}
 				else
+				{
+					if (AttributeAllocationInfo.Offset == MemmoveBlockCurrentOffset)
+					{
+						AttributeAllocationInfo.Offset -= MemmoveBlockStart - CurrentOffset;
+						MemmoveBlockCurrentOffset += AttributeAllocationInfo.Size;
+					}
+					else
+					{
+						const uint64 BlockSize = MemmoveBlockCurrentOffset - MemmoveBlockStart;
+						//Memmove support overlap
+						FMemory::Memmove(&StorageData[CurrentOffset], &StorageData[MemmoveBlockStart], BlockSize);
+						CurrentOffset += BlockSize;
+
+						MemmoveBlockStart = AttributeAllocationInfo.Offset;
+						AttributeAllocationInfo.Offset = CurrentOffset;
+						MemmoveBlockCurrentOffset = MemmoveBlockStart + AttributeAllocationInfo.Size;
+					}
+				}
+
+				//We have to move the block if we are the last attribute
+				if (bLast)
 				{
 					const uint64 BlockSize = MemmoveBlockCurrentOffset - MemmoveBlockStart;
 					//Memmove support overlap
 					FMemory::Memmove(&StorageData[CurrentOffset], &StorageData[MemmoveBlockStart], BlockSize);
-					CurrentOffset += BlockSize;
-
-					MemmoveBlockStart = AttributeAllocationInfo.Offset;
-					AttributeAllocationInfo.Offset = CurrentOffset;
-					MemmoveBlockCurrentOffset = MemmoveBlockStart + AttributeAllocationInfo.Size;
+					const uint64 RemoveStartIndex = CurrentOffset + BlockSize;
+					const uint64 RemoveCount = AttributeStorage.Num() - RemoveStartIndex;
+					//Remove the Last items, allow shrinking so we are compact
+					AttributeStorage.RemoveAt(RemoveStartIndex, RemoveCount);
 				}
 			}
+			//Reset the fragmented cost
+			FragmentedMemoryCost = 0;
+		}
 
-			//We have to move the block if we are the last attribute
-			if (bLast)
-			{
-				const uint64 BlockSize = MemmoveBlockCurrentOffset - MemmoveBlockStart;
-				//Memmove support overlap
-				FMemory::Memmove(&StorageData[CurrentOffset], &StorageData[MemmoveBlockStart], BlockSize);
-				const uint64 RemoveStartIndex = CurrentOffset + BlockSize;
-				const uint64 RemoveCount = AttributeStorage.Num() - RemoveStartIndex;
-				//Remove the Last items, allow shrinking so we are compact
-				AttributeStorage.RemoveAt(RemoveStartIndex, RemoveCount);
-			}
+		FGuid FAttributeStorage::GetValueHash(const uint8* Value, uint64 ValueSize)
+		{
+			FSHA1 Sha;
+			Sha.Update(Value, ValueSize);
+			Sha.Final();
+			// Retrieve the hash and use it to construct a pseudo-GUID. 
+			uint32 Hash[5];
+			Sha.GetHash(reinterpret_cast<uint8*>(Hash));
+			return FGuid(Hash[0] ^ Hash[4], Hash[1], Hash[2], Hash[3]);
 		}
-		//Reset the fragmented cost
-		FragmentedMemoryCost = 0;
-	}
 
-	FGuid FAttributeStorage::GetValueHash(const uint8* Value, uint64 ValueSize)
-	{
-		FSHA1 Sha;
-		Sha.Update(Value, ValueSize);
-		Sha.Final();
-		// Retrieve the hash and use it to construct a pseudo-GUID. 
-		uint32 Hash[5];
-		Sha.GetHash(reinterpret_cast<uint8*>(Hash));
-		return FGuid(Hash[0] ^ Hash[4], Hash[1], Hash[2], Hash[3]);
-	}
-
-	uint64 FAttributeStorage::GetValueSize(const FString& Value, TSpecializeType<FString >)
-	{
-		//We must add the null character '/0' terminate string
-		return (Value.Len()+1) * sizeof(TCHAR);
-	}
+		uint64 FAttributeStorage::GetValueSize(const FString& Value, TSpecializeType<FString >)
+		{
+			//We must add the null character '/0' terminate string
+			return (Value.Len() + 1) * sizeof(TCHAR);
+		}
 
-	uint64 FAttributeStorage::GetValueSize(const TArray<uint8>& Value, TSpecializeType<TArray<uint8> >)
-	{
-		return Value.Num();
-	}
+		uint64 FAttributeStorage::GetValueSize(const TArray<uint8>& Value, TSpecializeType<TArray<uint8> >)
+		{
+			return Value.Num();
+		}
 
-	uint64 FAttributeStorage::GetValueSize(const TArray64<uint8>& Value, TSpecializeType<TArray64<uint8> >)
-	{
-		return Value.Num();
-	}
+		uint64 FAttributeStorage::GetValueSize(const TArray64<uint8>& Value, TSpecializeType<TArray64<uint8> >)
+		{
+			return Value.Num();
+		}
 
-} //End Interchange namespace
\ No newline at end of file
+	} //ns Interchange
+} //ns UE
\ No newline at end of file
diff --git a/Engine/Source/Runtime/Interchange/Core/Public/InterchangeManager.h b/Engine/Source/Runtime/Interchange/Core/Public/InterchangeManager.h
--- a/Engine/Source/Runtime/Interchange/Core/Public/InterchangeManager.h
+++ b/Engine/Source/Runtime/Interchange/Core/Public/InterchangeManager.h
@@ -19,114 +19,116 @@
 #include "InterchangeManager.generated.h"
 
 class FAsyncTaskNotification;
-
-namespace Interchange
+namespace UE
 {
-	class INTERCHANGECORE_API FScopedSourceData
-	{
-	public:
-		explicit FScopedSourceData(const FString& Filename);
-		UInterchangeSourceData* GetSourceData() const;
-	private:
-		TStrongObjectPtr<UInterchangeSourceData> SourceDataPtr = nullptr;
-	};
-
-	class INTERCHANGECORE_API FScopedTranslator
+	namespace Interchange
 	{
-	public:
-		explicit FScopedTranslator(const UInterchangeSourceData* SourceData);
-		UInterchangeTranslatorBase* GetTranslator();
+		class INTERCHANGECORE_API FScopedSourceData
+		{
+		public:
+			explicit FScopedSourceData(const FString& Filename);
+			UInterchangeSourceData* GetSourceData() const;
+		private:
+			TStrongObjectPtr<UInterchangeSourceData> SourceDataPtr = nullptr;
+		};
 
-	private:
-		TStrongObjectPtr<UInterchangeTranslatorBase> ScopedTranslatorPtr = nullptr;
-	};
+		class INTERCHANGECORE_API FScopedTranslator
+		{
+		public:
+			explicit FScopedTranslator(const UInterchangeSourceData* SourceData);
+			UInterchangeTranslatorBase* GetTranslator();
 
-	enum class EImportType : uint8
-	{
-		ImportType_None,
-		ImportType_Asset,
-		ImportType_Scene
-	};
+		private:
+			TStrongObjectPtr<UInterchangeTranslatorBase> ScopedTranslatorPtr = nullptr;
+		};
 
-	struct FImportAsyncHelperData
-	{
-		//True if the import process is unattended. We cannot show UI  if the import is automated
-		bool bIsAutomated = false;
+		enum class EImportType : uint8
+		{
+			ImportType_None,
+			ImportType_Asset,
+			ImportType_Scene
+		};
 
-		//We can import assets or full scene
-		EImportType ImportType = EImportType::ImportType_None;
+		struct FImportAsyncHelperData
+		{
+			//True if the import process is unattended. We cannot show UI  if the import is automated
+			bool bIsAutomated = false;
 
-		//True if we are reimporting assets or scene
-		UObject* ReimportObject = nullptr;
-	};
+			//We can import assets or full scene
+			EImportType ImportType = EImportType::ImportType_None;
 
-	class FImportAsyncHelper : public FGCObject
-	{
-	public:
-		FImportAsyncHelper();
+			//True if we are reimporting assets or scene
+			UObject* ReimportObject = nullptr;
+		};
 
-		~FImportAsyncHelper()
+		class FImportAsyncHelper : public FGCObject
 		{
-			CleanUp();
-		}
+		public:
+			FImportAsyncHelper();
 
-		/* FGCObject interface */
-		virtual void AddReferencedObjects(FReferenceCollector& Collector) override;
+			~FImportAsyncHelper()
+			{
+				CleanUp();
+			}
 
-		TArray <TStrongObjectPtr<UInterchangeBaseNodeContainer>> BaseNodeContainers;
+			/* FGCObject interface */
+			virtual void AddReferencedObjects(FReferenceCollector& Collector) override;
 
-		TArray<UInterchangeSourceData* > SourceDatas;
-		TArray<UInterchangeTranslatorBase* > Translators;
-		TArray<UInterchangePipelineBase* > Pipelines;
-		TArray<UInterchangeFactoryBase* > Factories;
+			TArray <TStrongObjectPtr<UInterchangeBaseNodeContainer>> BaseNodeContainers;
 
-		TArray<FGraphEventRef> TranslatorTasks;
-		TArray<FGraphEventRef> PipelineTasks;
-		FGraphEventRef ParsingTask;
-		TArray<FGraphEventRef> CreateAssetTasks;
+			TArray<UInterchangeSourceData* > SourceDatas;
+			TArray<UInterchangeTranslatorBase* > Translators;
+			TArray<UInterchangePipelineBase* > Pipelines;
+			TArray<UInterchangeFactoryBase* > Factories;
 
-		FGraphEventRef CompletionTask;
+			TArray<FGraphEventRef> TranslatorTasks;
+			TArray<FGraphEventRef> PipelineTasks;
+			FGraphEventRef ParsingTask;
+			TArray<FGraphEventRef> CreateAssetTasks;
 
-		//Create package map, Key is package name. We cannot create package asynchronously so we have to create a game thread task to do this
-		FCriticalSection CreatedPackagesLock;
-		TMap<FString, UPackage*> CreatedPackages;
+			FGraphEventRef CompletionTask;
 
-		struct FImportedAssetInfo
-		{
-			UObject* ImportAsset;
-			UInterchangeFactoryBase* Factory;
-			UInterchangeBaseNode* AssetNode;
-		};
+			//Create package map, Key is package name. We cannot create package asynchronously so we have to create a game thread task to do this
+			FCriticalSection CreatedPackagesLock;
+			TMap<FString, UPackage*> CreatedPackages;
 
-		FCriticalSection ImportedAssetsPerSourceIndexLock;
-		TMap<int32, TArray<FImportedAssetInfo>> ImportedAssetsPerSourceIndex;
+			struct FImportedAssetInfo
+			{
+				UObject* ImportAsset;
+				UInterchangeFactoryBase* Factory;
+				UInterchangeBaseNode* AssetNode;
+			};
 
-		FImportAsyncHelperData TaskData;
+			FCriticalSection ImportedAssetsPerSourceIndexLock;
+			TMap<int32, TArray<FImportedAssetInfo>> ImportedAssetsPerSourceIndex;
 
-		TPromise< UObject* > RootObject;
-		FGraphEventRef RootObjectCompletionEvent;
+			FImportAsyncHelperData TaskData;
 
-		void CleanUp();
-	};
+			TPromise< UObject* > RootObject;
+			FGraphEventRef RootObjectCompletionEvent;
 
-	class INTERCHANGECORE_API FAsyncImportResult
-	{
-	public:
-		FAsyncImportResult() = default;
-		explicit FAsyncImportResult( TFuture< UObject* >&& InFutureObject, const FGraphEventRef& InGraphEvent );
+			void CleanUp();
+		};
 
-		bool IsValid() const;
-		UObject* Get() const;
-		FAsyncImportResult Next( TFunction< UObject*( UObject* ) > Continuation );
+		class INTERCHANGECORE_API FAsyncImportResult
+		{
+		public:
+			FAsyncImportResult() = default;
+			explicit FAsyncImportResult(TFuture< UObject* >&& InFutureObject, const FGraphEventRef& InGraphEvent);
+
+			bool IsValid() const;
+			UObject* Get() const;
+			FAsyncImportResult Next(TFunction< UObject* (UObject*) > Continuation);
 
-	private:
-		TFuture< UObject* > FutureObject;
-		FGraphEventRef GraphEvent;
-	};
+		private:
+			TFuture< UObject* > FutureObject;
+			FGraphEventRef GraphEvent;
+		};
 
-	void SanitizeInvalidChar(FString& String);
+		void SanitizeInvalidChar(FString& String);
 
-} //End namespace interchange
+	} //ns interchange
+} //ns UE
 
 USTRUCT(BlueprintType)
 struct FImportAssetParameters
@@ -234,7 +236,7 @@ public:
 	 */
 	UFUNCTION(BlueprintCallable, Category = "Interchange | Import Manager")
 	bool ImportAsset(const FString& ContentPath, const UInterchangeSourceData* SourceData, const FImportAssetParameters& ImportAssetParameters);
-	Interchange::FAsyncImportResult ImportAssetAsync(const FString& ContentPath, const UInterchangeSourceData* SourceData, const FImportAssetParameters& ImportAssetParameters);
+	UE::Interchange::FAsyncImportResult ImportAssetAsync(const FString& ContentPath, const UInterchangeSourceData* SourceData, const FImportAssetParameters& ImportAssetParameters);
 
 	/**
 	 * Call this to start an import scene process, the caller must specify a source data.
@@ -298,10 +300,10 @@ public:
 	 * Return an FImportAsynHelper pointer. The pointer is deleted when ReleaseAsyncHelper is call.
 	 * @param Data - The data we want to pass to the different import tasks
 	 */
-	TWeakPtr<Interchange::FImportAsyncHelper, ESPMode::ThreadSafe> CreateAsyncHelper(const Interchange::FImportAsyncHelperData& Data);
+	TWeakPtr<UE::Interchange::FImportAsyncHelper, ESPMode::ThreadSafe> CreateAsyncHelper(const UE::Interchange::FImportAsyncHelperData& Data);
 
 	/** Delete the specified AsyncHelper and remove it from the array that was holding it. */
-	void ReleaseAsyncHelper(TWeakPtr<Interchange::FImportAsyncHelper, ESPMode::ThreadSafe> AsyncHelper);
+	void ReleaseAsyncHelper(TWeakPtr<UE::Interchange::FImportAsyncHelper, ESPMode::ThreadSafe> AsyncHelper);
 
 	/*
 	 * Return the first translator that can translate the source data.
@@ -324,7 +326,7 @@ protected:
 
 private:
 	//By using pointer, there is no issue if the array get resize
-	TArray<TSharedPtr<Interchange::FImportAsyncHelper, ESPMode::ThreadSafe> > ImportTasks;
+	TArray<TSharedPtr<UE::Interchange::FImportAsyncHelper, ESPMode::ThreadSafe> > ImportTasks;
 
 	TSharedPtr<FAsyncTaskNotification> Notification = nullptr;
 
@@ -345,5 +347,5 @@ private:
 	UPROPERTY()
 	TMap<const UClass*, UInterchangeWriterBase* > RegisteredWriters;
 
-	friend class Interchange::FScopedTranslator;
+	friend class UE::Interchange::FScopedTranslator;
 };
diff --git a/Engine/Source/Runtime/Interchange/Core/Public/Nodes/InterchangeBaseNode.h b/Engine/Source/Runtime/Interchange/Core/Public/Nodes/InterchangeBaseNode.h
--- a/Engine/Source/Runtime/Interchange/Core/Public/Nodes/InterchangeBaseNode.h
+++ b/Engine/Source/Runtime/Interchange/Core/Public/Nodes/InterchangeBaseNode.h
@@ -27,21 +27,21 @@ namespace InterchangePrivateNodeBase
 	 *           return false if the attribute do not exist or there is an error retriving it from the Storage
 	 */
 	template<typename ValueType>
-	bool GetCustomAttribute(const Interchange::FAttributeStorage& Attributes, const Interchange::FAttributeKey& AttributeKey, const FString& OperationName, ValueType& OutAttributeValue)
+	bool GetCustomAttribute(const UE::Interchange::FAttributeStorage& Attributes, const UE::Interchange::FAttributeKey& AttributeKey, const FString& OperationName, ValueType& OutAttributeValue)
 	{
 		if (!Attributes.ContainAttribute(AttributeKey))
 		{
 			return false;
 		}
-		Interchange::FAttributeStorage::TAttributeHandle<ValueType> AttributeHandle = Attributes.GetAttributeHandle<ValueType>(AttributeKey);
+		UE::Interchange::FAttributeStorage::TAttributeHandle<ValueType> AttributeHandle = Attributes.GetAttributeHandle<ValueType>(AttributeKey);
 		if (!AttributeHandle.IsValid())
 		{
 			return false;
 		}
-		Interchange::EAttributeStorageResult Result = AttributeHandle.Get(OutAttributeValue);
-		if (!Interchange::IsAttributeStorageResultSuccess(Result))
+		UE::Interchange::EAttributeStorageResult Result = AttributeHandle.Get(OutAttributeValue);
+		if (!UE::Interchange::IsAttributeStorageResultSuccess(Result))
 		{
-			Interchange::LogAttributeStorageErrors(Result, OperationName, AttributeKey);
+			UE::Interchange::LogAttributeStorageErrors(Result, OperationName, AttributeKey);
 			return false;
 		}
 		return true;
@@ -56,12 +56,12 @@ namespace InterchangePrivateNodeBase
 	 * @param AttributeValue - The value we want to add or update in the storage
 	 */
 	template<typename ValueType>
-	bool SetCustomAttribute(Interchange::FAttributeStorage& Attributes, const Interchange::FAttributeKey& AttributeKey, const FString& OperationName, const ValueType& AttributeValue)
+	bool SetCustomAttribute(UE::Interchange::FAttributeStorage& Attributes, const UE::Interchange::FAttributeKey& AttributeKey, const FString& OperationName, const ValueType& AttributeValue)
 	{
-		Interchange::EAttributeStorageResult Result = Attributes.RegisterAttribute(AttributeKey, AttributeValue);
-		if (!Interchange::IsAttributeStorageResultSuccess(Result))
+		UE::Interchange::EAttributeStorageResult Result = Attributes.RegisterAttribute(AttributeKey, AttributeValue);
+		if (!UE::Interchange::IsAttributeStorageResultSuccess(Result))
 		{
-			Interchange::LogAttributeStorageErrors(Result, OperationName, AttributeKey);
+			UE::Interchange::LogAttributeStorageErrors(Result, OperationName, AttributeKey);
 			return false;
 		}
 		return true;
@@ -83,7 +83,7 @@ namespace InterchangePrivateNodeBase
  * @param EnumType - Optional, specify it only if the AssetType member is an enum so we can type cast it in the apply function (we use uint8 to store the enum value)"
  */
 #define IMPLEMENT_NODE_ATTRIBUTE_KEY(AttributeName)																		\
-const Interchange::FAttributeKey Macro_Custom##AttributeName##Key = Interchange::FAttributeKey(TEXT(#AttributeName));	\
+const UE::Interchange::FAttributeKey Macro_Custom##AttributeName##Key = UE::Interchange::FAttributeKey(TEXT(#AttributeName));	\
 
 #if WITH_ENGINE
 #define IMPLEMENT_NODE_ATTRIBUTE_APPLY_UOBJECT(AttributeName, AttributeType, AssetType, EnumType)	\
@@ -129,8 +129,8 @@ bool ApplyCustom##AttributeName##ToAsset(UObject* Asset) const										\
 	{																																				\
 		if(bAddApplyDelegate)																														\
 		{																																			\
-			TArray<Interchange::FApplyAttributeToAsset>& Delegates = ApplyCustomAttributeDelegates.FindOrAdd(AssetType::StaticClass());				\
-			Delegates.Add(Interchange::FApplyAttributeToAsset::CreateUObject(this, &NodeClassName::ApplyCustom##AttributeName##ToAsset));			\
+			TArray<UE::Interchange::FApplyAttributeToAsset>& Delegates = ApplyCustomAttributeDelegates.FindOrAdd(AssetType::StaticClass());				\
+			Delegates.Add(UE::Interchange::FApplyAttributeToAsset::CreateUObject(this, &NodeClassName::ApplyCustom##AttributeName##ToAsset));			\
 		}																																			\
 		return true;																																\
 	}																																				\
@@ -146,60 +146,63 @@ bool ApplyCustom##AttributeName##ToAsset(UObject* Asset) const										\
 
 
 //Interchange namespace
-namespace Interchange
+namespace UE
 {
-
-DECLARE_DELEGATE_RetVal_OneParam(bool, FApplyAttributeToAsset, UObject*);
-
-/**
- * Helper struct use to declare static const data we use in the UInterchangeBaseNode
- * Node that derive from UInterchangeBaseNode can also add a struct that derive from this one to add there static data
- * @note: The static data are mainly for holding Attribute keys. All attributes that are always available for a node should be in this class or a derived class.
- */
-struct FBaseNodeStaticData
-{
-	static const FAttributeKey& UniqueIDKey()
-	{
-		static FAttributeKey AttributeKey(TEXT("__UNQ_ID_"));
-		return AttributeKey;
-	}
-
-	static const FAttributeKey& DisplayLabelKey()
-	{
-		static FAttributeKey AttributeKey(TEXT("__DSPL_LBL_"));
-		return AttributeKey;
-	}
-
-	static const FAttributeKey& ParentIDKey()
+	namespace Interchange
 	{
-		static FAttributeKey AttributeKey(TEXT("__PARENT_UID_"));
-		return AttributeKey;
-	}
-
-	static const FAttributeKey& IsEnabledKey()
-	{
-		static FAttributeKey AttributeKey(TEXT("__IS_NBLD_"));
-		return AttributeKey;
-	}
 
-	static const FAttributeKey& DependencyCountKey()
-	{
-		static FAttributeKey AttributeKey(TEXT("__DEPENDENCY_COUNT_"));
-		return AttributeKey;
-	}
+		DECLARE_DELEGATE_RetVal_OneParam(bool, FApplyAttributeToAsset, UObject*);
 
-	static const FAttributeKey& DependencyBaseKey()
-	{
-		static FAttributeKey AttributeKey(TEXT("__DEPENDENCY_INDEX_"));
-		return AttributeKey;
-	}
-};
-
-} //ns Interchange
+		/**
+		 * Helper struct use to declare static const data we use in the UInterchangeBaseNode
+		 * Node that derive from UInterchangeBaseNode can also add a struct that derive from this one to add there static data
+		 * @note: The static data are mainly for holding Attribute keys. All attributes that are always available for a node should be in this class or a derived class.
+		 */
+		struct FBaseNodeStaticData
+		{
+			static const FAttributeKey& UniqueIDKey()
+			{
+				static FAttributeKey AttributeKey(TEXT("__UNQ_ID_"));
+				return AttributeKey;
+			}
+
+			static const FAttributeKey& DisplayLabelKey()
+			{
+				static FAttributeKey AttributeKey(TEXT("__DSPL_LBL_"));
+				return AttributeKey;
+			}
+
+			static const FAttributeKey& ParentIDKey()
+			{
+				static FAttributeKey AttributeKey(TEXT("__PARENT_UID_"));
+				return AttributeKey;
+			}
+
+			static const FAttributeKey& IsEnabledKey()
+			{
+				static FAttributeKey AttributeKey(TEXT("__IS_NBLD_"));
+				return AttributeKey;
+			}
+
+			static const FAttributeKey& DependencyCountKey()
+			{
+				static FAttributeKey AttributeKey(TEXT("__DEPENDENCY_COUNT_"));
+				return AttributeKey;
+			}
+
+			static const FAttributeKey& DependencyBaseKey()
+			{
+				static FAttributeKey AttributeKey(TEXT("__DEPENDENCY_INDEX_"));
+				return AttributeKey;
+			}
+		};
+
+	} //ns Interchange
+} //ns UE
 
 /**
  * This struct is used to store and retrieve key value attributes. The attributes are store in a generic FAttributeStorage which serialize the value in a TArray64<uint8>
- * See Interchange::EAttributeTypes to know the supported template types
+ * See UE::Interchange::EAttributeTypes to know the supported template types
  * This is an abstract class. This is the base class of the interchange node graph format, all class in this format should derive from this class
  */
 UCLASS(BlueprintType)
@@ -234,16 +237,16 @@ public:
 	 *
 	 */
 	template<typename T>
-	Interchange::FAttributeStorage::TAttributeHandle<T> RegisterAttribute(const Interchange::FAttributeKey& NodeAttributeKey, const T& Value)
+	UE::Interchange::FAttributeStorage::TAttributeHandle<T> RegisterAttribute(const UE::Interchange::FAttributeKey& NodeAttributeKey, const T& Value)
 	{
-		const Interchange::EAttributeStorageResult Result = Attributes.RegisterAttribute(NodeAttributeKey, Value);
+		const UE::Interchange::EAttributeStorageResult Result = Attributes.RegisterAttribute(NodeAttributeKey, Value);
 		
 		if (IsAttributeStorageResultSuccess(Result))
 		{
 			return Attributes.GetAttributeHandle<T>(NodeAttributeKey);
 		}
 		LogAttributeStorageErrors(Result, TEXT("RegisterAttribute"), NodeAttributeKey);
-		return Interchange::FAttributeStorage::TAttributeHandle<T>();
+		return UE::Interchange::FAttributeStorage::TAttributeHandle<T>();
 	}
 
 	/**
@@ -251,7 +254,7 @@ public:
 	 * @param nodeAttributeKey - The key of the searched attribute
 	 *
 	 */
-	virtual bool HasAttribute(const Interchange::FAttributeKey& NodeAttributeKey) const;
+	virtual bool HasAttribute(const UE::Interchange::FAttributeKey& NodeAttributeKey) const;
 
 
 	/**
@@ -261,7 +264,7 @@ public:
 	 * @param NodeAttributeKey - The key of the attribute
 	 */
 	template<typename T>
-	Interchange::FAttributeStorage::TAttributeHandle<T> GetAttributeHandle(const Interchange::FAttributeKey& NodeAttributeKey) const
+	UE::Interchange::FAttributeStorage::TAttributeHandle<T> GetAttributeHandle(const UE::Interchange::FAttributeKey& NodeAttributeKey) const
 	{
 		return Attributes.GetAttributeHandle<T>(NodeAttributeKey);
 	}
@@ -350,10 +353,10 @@ public:
 	mutable FSoftObjectPath ReferenceObject;
 protected:
 	/** The storage use to store the Key value attribute for this node. */
-	Interchange::FAttributeStorage Attributes;
+	UE::Interchange::FAttributeStorage Attributes;
 
 	/* This array hold the delegate to apply the attribute that has to be set on an UObject */
-	TMap<UClass*, TArray<Interchange::FApplyAttributeToAsset>> ApplyCustomAttributeDelegates;
+	TMap<UClass*, TArray<UE::Interchange::FApplyAttributeToAsset>> ApplyCustomAttributeDelegates;
 
 	bool bIsInitialized = false;
 };
diff --git a/Engine/Source/Runtime/Interchange/Core/Public/PackageUtils/PackageUtils.h b/Engine/Source/Runtime/Interchange/Core/Public/PackageUtils/PackageUtils.h
--- a/Engine/Source/Runtime/Interchange/Core/Public/PackageUtils/PackageUtils.h
+++ b/Engine/Source/Runtime/Interchange/Core/Public/PackageUtils/PackageUtils.h
@@ -7,48 +7,51 @@
 #include "Misc/Paths.h"
 
 //Interchange namespace
-namespace Interchange
+namespace UE
 {
-	class FPackageUtils
+	namespace Interchange
 	{
-	public:
-		static bool IsMapPackageAsset(const FString& ObjectPath)
+		class FPackageUtils
 		{
-			FString MapFilePath;
-			return FPackageUtils::IsMapPackageAsset(ObjectPath, MapFilePath);
-		}
+		public:
+			static bool IsMapPackageAsset(const FString& ObjectPath)
+			{
+				FString MapFilePath;
+				return FPackageUtils::IsMapPackageAsset(ObjectPath, MapFilePath);
+			}
 
-		static bool IsMapPackageAsset(const FString& ObjectPath, FString& MapFilePath)
-		{
-			const FString PackageName = FPackageUtils::ExtractPackageName(ObjectPath);
-			if ( PackageName.Len() > 0 )
+			static bool IsMapPackageAsset(const FString& ObjectPath, FString& MapFilePath)
 			{
-				FString PackagePath;
-				if ( FPackageName::DoesPackageExist(PackageName, NULL, &PackagePath) )
+				const FString PackageName = FPackageUtils::ExtractPackageName(ObjectPath);
+				if (PackageName.Len() > 0)
 				{
-					const FString FileExtension = FPaths::GetExtension(PackagePath, true);
-					if ( FileExtension == FPackageName::GetMapPackageExtension() )
+					FString PackagePath;
+					if (FPackageName::DoesPackageExist(PackageName, NULL, &PackagePath))
 					{
-						MapFilePath = PackagePath;
-						return true;
+						const FString FileExtension = FPaths::GetExtension(PackagePath, true);
+						if (FileExtension == FPackageName::GetMapPackageExtension())
+						{
+							MapFilePath = PackagePath;
+							return true;
+						}
 					}
 				}
-			}
 
-			return false;
-		}
+				return false;
+			}
 
-		static FString ExtractPackageName(const FString& ObjectPath)
-		{
-			// To find the package name in an object path we need to find the path left of the FIRST delimiter.
-			// Assets like BSPs, lightmaps etc. can have multiple '.' delimiters.
-			const int32 PackageDelimiterPos = ObjectPath.Find(TEXT("."), ESearchCase::CaseSensitive, ESearchDir::FromStart);
-			if ( PackageDelimiterPos != INDEX_NONE )
+			static FString ExtractPackageName(const FString& ObjectPath)
 			{
-				return ObjectPath.Left(PackageDelimiterPos);
-			}
+				// To find the package name in an object path we need to find the path left of the FIRST delimiter.
+				// Assets like BSPs, lightmaps etc. can have multiple '.' delimiters.
+				const int32 PackageDelimiterPos = ObjectPath.Find(TEXT("."), ESearchCase::CaseSensitive, ESearchDir::FromStart);
+				if (PackageDelimiterPos != INDEX_NONE)
+				{
+					return ObjectPath.Left(PackageDelimiterPos);
+				}
 
-			return ObjectPath;
-		}
-	};
-} //End Interchange namespace
\ No newline at end of file
+				return ObjectPath;
+			}
+		};
+	} //ns Interchange
+} //ns UE
\ No newline at end of file
