Split InterchangeCore module and add an InterchangeEngine module to make dependencies more easy to manage. The out of process worker cannot compile with Engine and do not need the interchangemanager.

diff --git a/Engine/Plugins/Enterprise/DatasmithImporter/Source/DatasmithImporter/DatasmithImporter.Build.cs b/Engine/Plugins/Enterprise/DatasmithImporter/Source/DatasmithImporter/DatasmithImporter.Build.cs
--- a/Engine/Plugins/Enterprise/DatasmithImporter/Source/DatasmithImporter/DatasmithImporter.Build.cs
+++ b/Engine/Plugins/Enterprise/DatasmithImporter/Source/DatasmithImporter/DatasmithImporter.Build.cs
@@ -26,6 +26,7 @@ namespace UnrealBuildTool.Rules
 					"FreeImage",
 					"InputCore",
 					"InterchangeCore",
+					"InterchangeEngine",
 					"InterchangeNodePlugin",
 					"Json",
 					"Landscape",
diff --git a/Engine/Plugins/Runtime/Interchange/Source/Export/InterchangeExportPlugin.Build.cs b/Engine/Plugins/Runtime/Interchange/Source/Export/InterchangeExportPlugin.Build.cs
--- a/Engine/Plugins/Runtime/Interchange/Source/Export/InterchangeExportPlugin.Build.cs
+++ b/Engine/Plugins/Runtime/Interchange/Source/Export/InterchangeExportPlugin.Build.cs
@@ -13,6 +13,7 @@ namespace UnrealBuildTool.Rules
 					"CoreUObject",
 					"Engine",
 					"InterchangeCore",
+					"InterchangeEngine"
 				}
 				);
 		}
diff --git a/Engine/Plugins/Runtime/Interchange/Source/Import/InterchangeImportPlugin.Build.cs b/Engine/Plugins/Runtime/Interchange/Source/Import/InterchangeImportPlugin.Build.cs
--- a/Engine/Plugins/Runtime/Interchange/Source/Import/InterchangeImportPlugin.Build.cs
+++ b/Engine/Plugins/Runtime/Interchange/Source/Import/InterchangeImportPlugin.Build.cs
@@ -14,6 +14,7 @@ namespace UnrealBuildTool.Rules
 					"Engine",
 					"InterchangeCore",
 					"InterchangeDispatcher",
+					"InterchangeEngine",
 					"InterchangeNodePlugin",
 					"MeshDescription",
 				}
diff --git a/Engine/Plugins/Runtime/Interchange/Source/Nodes/Public/InterchangeSceneNode.h b/Engine/Plugins/Runtime/Interchange/Source/Nodes/Public/InterchangeSceneNode.h
--- a/Engine/Plugins/Runtime/Interchange/Source/Nodes/Public/InterchangeSceneNode.h
+++ b/Engine/Plugins/Runtime/Interchange/Source/Nodes/Public/InterchangeSceneNode.h
@@ -5,11 +5,6 @@
 #include "CoreMinimal.h"
 #include "Nodes/InterchangeBaseNode.h"
 
-#if WITH_ENGINE
-
-
-#endif
-
 #include "InterchangeSceneNode.generated.h"
 
 
diff --git a/Engine/Plugins/Runtime/Interchange/Source/Parsers/Fbx/InterchangeFbxParser.Build.cs b/Engine/Plugins/Runtime/Interchange/Source/Parsers/Fbx/InterchangeFbxParser.Build.cs
--- a/Engine/Plugins/Runtime/Interchange/Source/Parsers/Fbx/InterchangeFbxParser.Build.cs
+++ b/Engine/Plugins/Runtime/Interchange/Source/Parsers/Fbx/InterchangeFbxParser.Build.cs
@@ -26,8 +26,11 @@ namespace UnrealBuildTool.Rules
 
 			if (Target.bCompileAgainstEngine)
 			{
-				PrivateDependencyModuleNames.Add("Engine");
-
+				PrivateDependencyModuleNames.AddRange(
+					new string[] {
+					"Engine",
+					}
+				);
 			}
 
 			AddEngineThirdPartyPrivateStaticDependencies(Target,
diff --git a/Engine/Source/Developer/AssetTools/AssetTools.Build.cs b/Engine/Source/Developer/AssetTools/AssetTools.Build.cs
--- a/Engine/Source/Developer/AssetTools/AssetTools.Build.cs
+++ b/Engine/Source/Developer/AssetTools/AssetTools.Build.cs
@@ -38,7 +38,8 @@ public class AssetTools : ModuleRules
 				"ToolMenus",
 				"PhysicsCore",
 				"DeveloperSettings",
-				"InterchangeCore"
+				"InterchangeCore",
+				"InterchangeEngine"
 			}
 		);
 
diff --git a/Engine/Source/Editor/UnrealEd/UnrealEd.Build.cs b/Engine/Source/Editor/UnrealEd/UnrealEd.Build.cs
--- a/Engine/Source/Editor/UnrealEd/UnrealEd.Build.cs
+++ b/Engine/Source/Editor/UnrealEd/UnrealEd.Build.cs
@@ -112,7 +112,6 @@ public class UnrealEd : ModuleRules
 				"TypedElementInterfaces",
 				"ToolMenusEditor",
 				"StatusBar",
-				"InterchangeCore",
 			}
 		);
 
@@ -134,6 +133,8 @@ public class UnrealEd : ModuleRules
 				"ImageWriteQueue",
 				"InputCore",
 				"InputBindingEditor",
+				"InterchangeCore",
+				"InterchangeEngine",
 				"LauncherServices",
 				"MaterialEditor",
 				"MessageLog",
diff --git a/Engine/Source/Runtime/Interchange/Core/InterchangeCore.Build.cs b/Engine/Source/Runtime/Interchange/Core/InterchangeCore.Build.cs
--- a/Engine/Source/Runtime/Interchange/Core/InterchangeCore.Build.cs
+++ b/Engine/Source/Runtime/Interchange/Core/InterchangeCore.Build.cs
@@ -7,20 +7,10 @@ public class InterchangeCore : ModuleRules
 	public InterchangeCore(ReadOnlyTargetRules Target) : base(Target)
 	{
 		PublicDependencyModuleNames.AddRange(
-				new string[] {
-					"Core",
-					"CoreUObject",
-				}
-			);
-
-		if(Target.bCompileAgainstEngine)
-        {
-			PrivateDependencyModuleNames.AddRange(
-				new string[] {
-					"AssetRegistry",
-					"Engine",
-				}
-			);
-		}
+			new string[] {
+				"Core",
+				"CoreUObject",
+			}
+		);
 	}
 }
diff --git a/Engine/Source/Runtime/Interchange/Core/Private/InterchangeManager.cpp b/Engine/Source/Runtime/Interchange/Core/Private/InterchangeManager.cpp
deleted file mode 100644
diff --git a/Engine/Source/Runtime/Interchange/Core/Private/Tasks/InterchangeTaskCompletion.cpp b/Engine/Source/Runtime/Interchange/Core/Private/Tasks/InterchangeTaskCompletion.cpp
deleted file mode 100644
diff --git a/Engine/Source/Runtime/Interchange/Core/Private/Tasks/InterchangeTaskCompletion.h b/Engine/Source/Runtime/Interchange/Core/Private/Tasks/InterchangeTaskCompletion.h
deleted file mode 100644
diff --git a/Engine/Source/Runtime/Interchange/Core/Private/Tasks/InterchangeTaskCreateAsset.cpp b/Engine/Source/Runtime/Interchange/Core/Private/Tasks/InterchangeTaskCreateAsset.cpp
deleted file mode 100644
diff --git a/Engine/Source/Runtime/Interchange/Core/Private/Tasks/InterchangeTaskCreateAsset.h b/Engine/Source/Runtime/Interchange/Core/Private/Tasks/InterchangeTaskCreateAsset.h
deleted file mode 100644
diff --git a/Engine/Source/Runtime/Interchange/Core/Private/Tasks/InterchangeTaskParsing.cpp b/Engine/Source/Runtime/Interchange/Core/Private/Tasks/InterchangeTaskParsing.cpp
deleted file mode 100644
diff --git a/Engine/Source/Runtime/Interchange/Core/Private/Tasks/InterchangeTaskParsing.h b/Engine/Source/Runtime/Interchange/Core/Private/Tasks/InterchangeTaskParsing.h
deleted file mode 100644
diff --git a/Engine/Source/Runtime/Interchange/Core/Private/Tasks/InterchangeTaskPipeline.cpp b/Engine/Source/Runtime/Interchange/Core/Private/Tasks/InterchangeTaskPipeline.cpp
deleted file mode 100644
diff --git a/Engine/Source/Runtime/Interchange/Core/Private/Tasks/InterchangeTaskPipeline.h b/Engine/Source/Runtime/Interchange/Core/Private/Tasks/InterchangeTaskPipeline.h
deleted file mode 100644
diff --git a/Engine/Source/Runtime/Interchange/Core/Private/Tasks/InterchangeTaskTranslator.cpp b/Engine/Source/Runtime/Interchange/Core/Private/Tasks/InterchangeTaskTranslator.cpp
deleted file mode 100644
diff --git a/Engine/Source/Runtime/Interchange/Core/Private/Tasks/InterchangeTaskTranslator.h b/Engine/Source/Runtime/Interchange/Core/Private/Tasks/InterchangeTaskTranslator.h
deleted file mode 100644
diff --git a/Engine/Source/Runtime/Interchange/Core/Public/InterchangeManager.h b/Engine/Source/Runtime/Interchange/Core/Public/InterchangeManager.h
deleted file mode 100644
diff --git a/Engine/Source/Runtime/Interchange/Engine/InterchangeEngine.Build.cs b/Engine/Source/Runtime/Interchange/Engine/InterchangeEngine.Build.cs
new file mode 100644
--- /dev/null
+++ b/Engine/Source/Runtime/Interchange/Engine/InterchangeEngine.Build.cs
@@ -0,0 +1,24 @@
+// Copyright Epic Games, Inc. All Rights Reserved.
+
+using UnrealBuildTool;
+
+public class InterchangeEngine : ModuleRules
+{
+	public InterchangeEngine(ReadOnlyTargetRules Target) : base(Target)
+	{
+		PublicDependencyModuleNames.AddRange(
+				new string[] {
+					"Core",
+					"CoreUObject",
+					"InterchangeCore"
+				}
+			);
+
+		PrivateDependencyModuleNames.AddRange(
+			new string[] {
+				"AssetRegistry",
+				"Engine",
+			}
+		);
+	}
+}
diff --git a/Engine/Source/Runtime/Interchange/Engine/Private/InterchangeEngineLogPrivate.h b/Engine/Source/Runtime/Interchange/Engine/Private/InterchangeEngineLogPrivate.h
new file mode 100644
--- /dev/null
+++ b/Engine/Source/Runtime/Interchange/Engine/Private/InterchangeEngineLogPrivate.h
@@ -0,0 +1,7 @@
+// Copyright Epic Games, Inc. All Rights Reserved.
+
+#pragma once
+
+#include "CoreMinimal.h"
+
+DECLARE_LOG_CATEGORY_EXTERN(LogInterchangeEngine, Log, All);
diff --git a/Engine/Source/Runtime/Interchange/Engine/Private/InterchangeEngineModule.cpp b/Engine/Source/Runtime/Interchange/Engine/Private/InterchangeEngineModule.cpp
new file mode 100644
--- /dev/null
+++ b/Engine/Source/Runtime/Interchange/Engine/Private/InterchangeEngineModule.cpp
@@ -0,0 +1,30 @@
+// Copyright Epic Games, Inc. All Rights Reserved.
+#include "CoreMinimal.h"
+#include "InterchangeEngineLogPrivate.h"
+#include "Modules/ModuleInterface.h"
+#include "Modules/ModuleManager.h"
+
+DEFINE_LOG_CATEGORY(LogInterchangeEngine);
+
+
+/**
+ * Implements the Interchange module.
+ */
+class FInterchangeEngineModule
+	: public IModuleInterface
+{
+public:
+
+	// IModuleInterface interface
+
+	virtual void StartupModule() override { }
+	virtual void ShutdownModule() override { }
+
+	virtual bool SupportsDynamicReloading() override
+	{
+		return false;
+	}
+};
+
+
+IMPLEMENT_MODULE(FInterchangeEngineModule, InterchangeEngine);
diff --git a/Engine/Source/Runtime/Interchange/Engine/Private/InterchangeManager.cpp b/Engine/Source/Runtime/Interchange/Engine/Private/InterchangeManager.cpp
new file mode 100644
--- /dev/null
+++ b/Engine/Source/Runtime/Interchange/Engine/Private/InterchangeManager.cpp
@@ -0,0 +1,536 @@
+// Copyright Epic Games, Inc. All Rights Reserved.
+
+#include "InterchangeManager.h"
+
+#include "AssetDataTagMap.h"
+#include "AssetRegistryModule.h"
+#include "CoreMinimal.h"
+#include "Engine/Blueprint.h"
+#include "InterchangeFactoryBase.h"
+#include "InterchangeEngineLogPrivate.h"
+#include "InterchangeSourceData.h"
+#include "InterchangeTranslatorBase.h"
+#include "InterchangeWriterBase.h"
+#include "Internationalization/Internationalization.h"
+#include "Misc/App.h"
+#include "Misc/AsyncTaskNotification.h"
+#include "Misc/MessageDialog.h"
+#include "Misc/ScopedSlowTask.h"
+#include "Nodes/InterchangeBaseNodeContainer.h"
+#include "PackageUtils/PackageUtils.h"
+#include "Tasks/InterchangeTaskParsing.h"
+#include "Tasks/InterchangeTaskPipeline.h"
+#include "Tasks/InterchangeTaskTranslator.h"
+#include "UObject/Class.h"
+#include "UObject/NameTypes.h"
+#include "UObject/Object.h"
+#include "UObject/ObjectMacros.h"
+#include "UObject/UObjectIterator.h"
+#include "UObject/WeakObjectPtrTemplates.h"
+
+namespace InternalInterchangePrivate
+{
+	const FLogCategoryBase* GetLogInterchangePtr()
+	{
+#if NO_LOGGING
+		return nullptr;
+#else
+		return &LogInterchangeEngine;
+#endif
+	}
+}
+
+UE::Interchange::FScopedSourceData::FScopedSourceData(const FString& Filename)
+{
+	//Found the translator
+	SourceDataPtr = TStrongObjectPtr<UInterchangeSourceData>(UInterchangeManager::GetInterchangeManager().CreateSourceData(Filename));
+	check(SourceDataPtr.IsValid());
+}
+
+UInterchangeSourceData* UE::Interchange::FScopedSourceData::GetSourceData() const
+{
+	return SourceDataPtr.Get();
+}
+
+UE::Interchange::FScopedTranslator::FScopedTranslator(const UInterchangeSourceData* SourceData)
+{
+	//Found the translator
+	ScopedTranslatorPtr = TStrongObjectPtr<UInterchangeTranslatorBase>(UInterchangeManager::GetInterchangeManager().GetTranslatorForSourceData(SourceData));
+}
+
+UInterchangeTranslatorBase* UE::Interchange::FScopedTranslator::GetTranslator()
+{
+	return ScopedTranslatorPtr.Get();
+}
+
+UE::Interchange::FImportAsyncHelper::FImportAsyncHelper()
+{
+	RootObjectCompletionEvent = FGraphEvent::CreateGraphEvent();
+}
+
+void UE::Interchange::FImportAsyncHelper::AddReferencedObjects(FReferenceCollector& Collector)
+{
+	for (UInterchangeSourceData* SourceData : SourceDatas)
+	{
+		Collector.AddReferencedObject(SourceData);
+	}
+	for (UInterchangeTranslatorBase* Translator : Translators)
+	{
+		Collector.AddReferencedObject(Translator);
+	}
+	for (UInterchangePipelineBase* Pipeline : Pipelines)
+	{
+		Collector.AddReferencedObject(Pipeline);
+	}
+	for (UInterchangeFactoryBase* Factory : Factories)
+	{
+		Collector.AddReferencedObject(Factory);
+	}
+}
+
+void UE::Interchange::FImportAsyncHelper::CleanUp()
+{
+	//Release the graph
+	BaseNodeContainers.Empty();
+
+	for (UInterchangeSourceData* SourceData : SourceDatas)
+	{
+		if (SourceData)
+		{
+			SourceData->RemoveFromRoot();
+			SourceData->MarkPendingKill();
+			SourceData = nullptr;
+		}
+	}
+	SourceDatas.Empty();
+
+	for (UInterchangeTranslatorBase* Translator : Translators)
+	{
+		if(Translator)
+		{
+			Translator->ImportFinish();
+			Translator->RemoveFromRoot();
+			Translator->MarkPendingKill();
+			Translator = nullptr;
+		}
+	}
+	Translators.Empty();
+
+	for (UInterchangePipelineBase* Pipeline : Pipelines)
+	{
+		if(Pipeline)
+		{
+			Pipeline->RemoveFromRoot();
+			Pipeline->MarkPendingKill();
+			Pipeline = nullptr;
+		}
+	}
+	Pipelines.Empty();
+
+	//Factories are not instantiate, we use the registered one directly
+	Factories.Empty();
+}
+
+UE::Interchange::FAsyncImportResult::FAsyncImportResult( TFuture< UObject* >&& InFutureObject, const FGraphEventRef& InGraphEvent )
+	: FutureObject( MoveTemp( InFutureObject ) )
+	, GraphEvent( InGraphEvent )
+{
+}
+
+bool UE::Interchange::FAsyncImportResult::IsValid() const
+{
+	return FutureObject.IsValid();
+}
+
+UObject* UE::Interchange::FAsyncImportResult::Get() const
+{
+	if ( !FutureObject.IsReady() )
+	{
+		// Tick the task graph until our FutureObject is ready
+		FTaskGraphInterface::Get().WaitUntilTaskCompletes( GraphEvent );
+	}
+
+	return FutureObject.Get();
+}
+
+UE::Interchange::FAsyncImportResult UE::Interchange::FAsyncImportResult::Next( TFunction< UObject*( UObject* ) > Continuation )
+{
+	return UE::Interchange::FAsyncImportResult{ FutureObject.Next( Continuation ), GraphEvent };
+}
+
+void UE::Interchange::SanitizeInvalidChar(FString& String)
+{
+	const TCHAR* InvalidChar = INVALID_OBJECTPATH_CHARACTERS;
+	while (*InvalidChar)
+	{
+		String.ReplaceCharInline(*InvalidChar, TCHAR('_'), ESearchCase::CaseSensitive);
+		++InvalidChar;
+	}
+}
+
+bool UInterchangeManager::RegisterTranslator(const UClass* TranslatorClass)
+{
+	if(!TranslatorClass)
+	{
+		return false;
+	}
+
+	if(RegisteredTranslators.Contains(TranslatorClass))
+	{
+		return true;
+	}
+	UInterchangeTranslatorBase* TranslatorToRegister = NewObject<UInterchangeTranslatorBase>(GetTransientPackage(), TranslatorClass, NAME_None);
+	if(!TranslatorToRegister)
+	{
+		return false;
+	}
+	RegisteredTranslators.Add(TranslatorClass, TranslatorToRegister);
+	return true;
+}
+
+bool UInterchangeManager::RegisterFactory(const UClass* FactoryClass)
+{
+	if (!FactoryClass)
+	{
+		return false;
+	}
+
+	UInterchangeFactoryBase* FactoryToRegister = NewObject<UInterchangeFactoryBase>(GetTransientPackage(), FactoryClass, NAME_None);
+	if (!FactoryToRegister)
+	{
+		return false;
+	}
+	if (FactoryToRegister->GetFactoryClass() == nullptr || RegisteredFactories.Contains(FactoryToRegister->GetFactoryClass()))
+	{
+		FactoryToRegister->MarkPendingKill();
+		return FactoryToRegister->GetFactoryClass() == nullptr ? false : true;
+	}
+	RegisteredFactories.Add(FactoryToRegister->GetFactoryClass(), FactoryToRegister);
+	return true;
+}
+
+bool UInterchangeManager::RegisterWriter(const UClass* WriterClass)
+{
+	if (!WriterClass)
+	{
+		return false;
+	}
+
+	if (RegisteredFactories.Contains(WriterClass))
+	{
+		return true;
+	}
+	UInterchangeWriterBase* WriterToRegister = NewObject<UInterchangeWriterBase>(GetTransientPackage(), WriterClass, NAME_None);
+	if (!WriterToRegister)
+	{
+		return false;
+	}
+	RegisteredWriters.Add(WriterClass, WriterToRegister);
+	return true;
+}
+
+bool UInterchangeManager::CanTranslateSourceData(const UInterchangeSourceData* SourceData) const
+{
+	//Found the translator
+	UE::Interchange::FScopedTranslator ScopeDataTranslator(SourceData);
+	const UInterchangeTranslatorBase* SourceDataTranslator = ScopeDataTranslator.GetTranslator();
+	if (SourceDataTranslator)
+	{
+		return true;
+	}
+	return false;
+}
+
+bool UInterchangeManager::ImportAsset(const FString& ContentPath, const UInterchangeSourceData* SourceData, const FImportAssetParameters& ImportAssetParameters)
+{
+	return ImportAssetAsync( ContentPath, SourceData, ImportAssetParameters ).IsValid();
+}
+
+UE::Interchange::FAsyncImportResult UInterchangeManager::ImportAssetAsync(const FString& ContentPath, const UInterchangeSourceData* SourceData, const FImportAssetParameters& ImportAssetParameters)
+{
+	FString PackageBasePath = ContentPath;
+	if(!ImportAssetParameters.ReimportAsset)
+	{
+		UE::Interchange::SanitizeInvalidChar(PackageBasePath);
+	}
+
+	bool bCanShowDialog = !ImportAssetParameters.bIsAutomated && !IsAttended();
+
+	//Create a task for every source data
+	UE::Interchange::FImportAsyncHelperData TaskData;
+	TaskData.bIsAutomated = ImportAssetParameters.bIsAutomated;
+	TaskData.ImportType = UE::Interchange::EImportType::ImportType_Asset;
+	TaskData.ReimportObject = ImportAssetParameters.ReimportAsset;
+	TWeakPtr<UE::Interchange::FImportAsyncHelper, ESPMode::ThreadSafe> WeakAsyncHelper = CreateAsyncHelper(TaskData);
+	TSharedPtr<UE::Interchange::FImportAsyncHelper, ESPMode::ThreadSafe> AsyncHelper = WeakAsyncHelper.Pin();
+	check(AsyncHelper.IsValid());
+
+	FText TitleText = NSLOCTEXT("Interchange", "Asynchronous_import_start", "Importing");
+	if(!Notification.IsValid())
+	{
+		FAsyncTaskNotificationConfig NotificationConfig;
+		NotificationConfig.bIsHeadless = false;
+		NotificationConfig.bKeepOpenOnFailure = true;
+		NotificationConfig.TitleText = TitleText;
+		NotificationConfig.LogCategory = InternalInterchangePrivate::GetLogInterchangePtr();
+		NotificationConfig.bCanCancel.Set(false);
+		NotificationConfig.bKeepOpenOnFailure.Set(true);
+
+		Notification = MakeShared<FAsyncTaskNotification>(NotificationConfig);
+		Notification->SetNotificationState(FAsyncNotificationStateData(TitleText, FText::GetEmpty(), EAsyncTaskNotificationState::Pending));
+	}
+	
+	//Create a duplicate of the source data, we need to be multithread safe so we copy it to control the life cycle. The async helper will hold it and delete it when the import task will be completed.
+	UInterchangeSourceData* DuplicateSourceData = Cast<UInterchangeSourceData>(StaticDuplicateObject(SourceData, GetTransientPackage()));
+	//Array of source data to build one graph per source
+	AsyncHelper->SourceDatas.Add(DuplicateSourceData);
+		
+
+	//Get all the translators for the source datas
+	for (int32 SourceDataIndex = 0; SourceDataIndex < AsyncHelper->SourceDatas.Num(); ++SourceDataIndex)
+	{
+		ensure(AsyncHelper->Translators.Add(GetTranslatorForSourceData(AsyncHelper->SourceDatas[SourceDataIndex])) == SourceDataIndex);
+	}
+
+	//Create the node graphs for each source data (StrongObjectPtr has to be created on the main thread)
+	for (int32 SourceDataIndex = 0; SourceDataIndex < AsyncHelper->SourceDatas.Num(); ++SourceDataIndex)
+	{
+		AsyncHelper->BaseNodeContainers.Add(TStrongObjectPtr<UInterchangeBaseNodeContainer>(NewObject<UInterchangeBaseNodeContainer>(GetTransientPackage(), NAME_None)));
+		check(AsyncHelper->BaseNodeContainers[SourceDataIndex].IsValid());
+	}
+
+	if ( ImportAssetParameters.OverridePipeline == nullptr )
+	{
+		//Get all pipeline candidate we want for this import
+		/*TArray<UClass*> PipelineCandidates;
+		FindPipelineCandidate(PipelineCandidates);
+
+		// Stack all pipelines, for this import proto. TODO: We need to be able to control which pipeline we use for the import
+		// It can be set in the project settings and can also be set by a UI where the user create the pipeline stack he want.
+		// This should be a list of available pipelines that can be drop into a stack where you can control the order.
+		for (int32 GraphPipelineNumber = 0; GraphPipelineNumber < PipelineCandidates.Num(); ++GraphPipelineNumber)
+		{
+			UInterchangePipelineBase* GeneratedPipeline = NewObject<UInterchangePipelineBase>(GetTransientPackage(), PipelineCandidates[GraphPipelineNumber], NAME_None, RF_NoFlags);
+			AsyncHelper->Pipelines.Add(GeneratedPipeline);
+		}*/
+	}
+	else
+	{
+		AsyncHelper->Pipelines.Add(ImportAssetParameters.OverridePipeline);
+	}
+
+	//Create/Start import tasks
+	FGraphEventArray PipelinePrerequistes;
+	check(AsyncHelper->Translators.Num() == AsyncHelper->SourceDatas.Num());
+	for (int32 SourceDataIndex = 0; SourceDataIndex < AsyncHelper->SourceDatas.Num(); ++SourceDataIndex)
+	{
+		int32 TranslatorTaskIndex = AsyncHelper->TranslatorTasks.Add(TGraphTask<UE::Interchange::FTaskTranslator>::CreateTask().ConstructAndDispatchWhenReady(SourceDataIndex, WeakAsyncHelper));
+		PipelinePrerequistes.Add(AsyncHelper->TranslatorTasks[TranslatorTaskIndex]);
+	}
+		
+	FGraphEventArray GraphParsingPrerequistes;
+	for(int32 GraphPipelineIndex = 0; GraphPipelineIndex < AsyncHelper->Pipelines.Num(); ++GraphPipelineIndex)
+	{
+		UInterchangePipelineBase* GraphPipeline = AsyncHelper->Pipelines[GraphPipelineIndex];
+		TWeakObjectPtr<UInterchangePipelineBase> WeakPipelinePtr = GraphPipeline;
+		int32 GraphPipelineTaskIndex = INDEX_NONE;
+		GraphPipelineTaskIndex = AsyncHelper->PipelineTasks.Add(TGraphTask<UE::Interchange::FTaskPipeline>::CreateTask(&PipelinePrerequistes).ConstructAndDispatchWhenReady(WeakPipelinePtr, WeakAsyncHelper));
+		//Ensure we run the pipeline in the same order we create the task, since pipeline modify the node container, its important that its not process in parallel, Adding the one we start to the prerequisites
+		//is the way to go here
+		PipelinePrerequistes.Add(AsyncHelper->PipelineTasks[GraphPipelineTaskIndex]);
+
+		//Add pipeline to the graph parsing prerequisites
+		GraphParsingPrerequistes.Add(AsyncHelper->PipelineTasks[GraphPipelineTaskIndex]);
+	}
+
+	if (GraphParsingPrerequistes.Num() > 0)
+	{
+		AsyncHelper->ParsingTask = TGraphTask<UE::Interchange::FTaskParsing>::CreateTask(&GraphParsingPrerequistes).ConstructAndDispatchWhenReady(this, PackageBasePath, WeakAsyncHelper);
+	}
+	else
+	{
+		//Fallback on the translator pipeline prerequisites (translator must be done if there is no pipeline)
+		AsyncHelper->ParsingTask = TGraphTask<UE::Interchange::FTaskParsing>::CreateTask(&PipelinePrerequistes).ConstructAndDispatchWhenReady(this, PackageBasePath, WeakAsyncHelper);
+	}
+
+	//The graph parsing task will create the FCreateAssetTask that will run after them, the FAssetImportTask will call the appropriate Post asset import pipeline when the asset is completed
+
+	return UE::Interchange::FAsyncImportResult{ AsyncHelper->RootObject.GetFuture(), AsyncHelper->RootObjectCompletionEvent };
+}
+
+bool UInterchangeManager::ImportScene(const FString& ImportContext, const UInterchangeSourceData* SourceData, bool bIsReimport, bool bIsAutomated)
+{
+	return false;
+}
+
+bool UInterchangeManager::ExportAsset(const UObject* Asset, bool bIsAutomated)
+{
+	return false;
+}
+
+bool UInterchangeManager::ExportScene(const UObject* World, bool bIsAutomated)
+{
+	return false;
+}
+
+UInterchangeSourceData* UInterchangeManager::CreateSourceData(const FString& InFileName)
+{
+	UInterchangeSourceData* SourceDataAsset = NewObject<UInterchangeSourceData>(GetTransientPackage(), NAME_None);
+	if(!InFileName.IsEmpty())
+	{
+		SourceDataAsset->SetFilename(InFileName);
+	}
+	return SourceDataAsset;
+}
+
+TWeakPtr<UE::Interchange::FImportAsyncHelper, ESPMode::ThreadSafe> UInterchangeManager::CreateAsyncHelper(const UE::Interchange::FImportAsyncHelperData& Data)
+{
+	TSharedPtr<UE::Interchange::FImportAsyncHelper, ESPMode::ThreadSafe> AsyncHelper = MakeShared<UE::Interchange::FImportAsyncHelper, ESPMode::ThreadSafe>();
+	//Copy the task data
+	AsyncHelper->TaskData = Data;
+	int32 AsyncHelperIndex = ImportTasks.Add(AsyncHelper);
+	//Update the asynchronous notification
+	if (Notification.IsValid())
+	{
+		int32 ImportTaskNumber = ImportTasks.Num();
+		FString ImportTaskNumberStr = TEXT(" (") + FString::FromInt(ImportTaskNumber) + TEXT(")");
+		Notification->SetProgressText(FText::FromString(ImportTaskNumberStr));
+	}
+	
+	return ImportTasks[AsyncHelperIndex];
+}
+
+void UInterchangeManager::ReleaseAsyncHelper(TWeakPtr<UE::Interchange::FImportAsyncHelper, ESPMode::ThreadSafe> AsyncHelper)
+{
+	check(AsyncHelper.IsValid());
+	ImportTasks.RemoveSingle(AsyncHelper.Pin());
+	check(!AsyncHelper.IsValid());
+
+	int32 ImportTaskNumber = ImportTasks.Num();
+	FString ImportTaskNumberStr = TEXT(" (") + FString::FromInt(ImportTaskNumber) + TEXT(")");
+	if (ImportTasks.Num() == 0 && Notification.IsValid())
+	{
+		FText TitleText = NSLOCTEXT("Interchange", "Asynchronous_import_end", "Import Done");
+		//TODO make sure any error are reported so we can control success or not
+		const bool bSuccess = true;
+		Notification->SetComplete(TitleText, FText::GetEmpty(), bSuccess);
+		Notification = nullptr; //This should delete the notification
+	}
+	else if(Notification.IsValid())
+	{
+		Notification->SetProgressText(FText::FromString(ImportTaskNumberStr));
+	}
+}
+
+UInterchangeTranslatorBase* UInterchangeManager::GetTranslatorForSourceData(const UInterchangeSourceData* SourceData) const
+{
+	if (RegisteredTranslators.Num() == 0)
+	{
+		return nullptr;
+	}
+	//Found the translator
+	for (const auto Kvp : RegisteredTranslators)
+	{
+		if (Kvp.Value->CanImportSourceData(SourceData))
+		{
+			UInterchangeTranslatorBase* SourceDataTranslator = NewObject<UInterchangeTranslatorBase>(GetTransientPackage(), Kvp.Key, NAME_None);
+			return SourceDataTranslator;
+		}
+	}
+	return nullptr;
+}
+
+bool UInterchangeManager::IsAttended()
+{
+	if (FApp::IsGame())
+	{
+		return false;
+	}
+	if (FApp::IsUnattended())
+	{
+		return false;
+	}
+	return true;
+}
+
+void UInterchangeManager::FindPipelineCandidate(TArray<UClass*>& PipelineCandidates)
+{
+	//Find in memory pipeline class
+	for (TObjectIterator< UClass > ClassIt; ClassIt; ++ClassIt)
+	{
+		UClass* Class = *ClassIt;
+		// Only interested in native C++ classes
+		if (!Class->IsNative())
+		{
+			continue;
+		}
+		// Ignore deprecated
+		if (Class->HasAnyClassFlags(CLASS_Deprecated | CLASS_NewerVersionExists))
+		{
+			continue;
+		}
+
+		// Check this class is a subclass of Base and not the base itself
+		if (Class == UInterchangePipelineBase::StaticClass() || !Class->IsChildOf(UInterchangePipelineBase::StaticClass()))
+		{
+			continue;
+		}
+
+		//We found a candidate
+		PipelineCandidates.AddUnique(Class);
+	}
+
+//Blueprint and python script discoverability is available only if we compile with the engine
+	// Load the asset registry module
+	FAssetRegistryModule& AssetRegistryModule = FModuleManager::LoadModuleChecked< FAssetRegistryModule >(FName("AssetRegistry"));
+	IAssetRegistry& AssetRegistry = AssetRegistryModule.Get();
+
+	TArray< FString > ContentPaths;
+	ContentPaths.Add(TEXT("/Game"));
+	//TODO do we have an other alternative, this call is synchronous and will wait unitl the registry database have finish the initial scan. If there is a lot of asset it can take multiple second the first time we call it.
+	AssetRegistry.ScanPathsSynchronous(ContentPaths);
+
+	FName BaseClassName = UInterchangePipelineBase::StaticClass()->GetFName();
+
+	// Use the asset registry to get the set of all class names deriving from Base
+	TSet< FName > DerivedNames;
+	{
+		TArray< FName > BaseNames;
+		BaseNames.Add(BaseClassName);
+
+		TSet< FName > Excluded;
+		AssetRegistry.GetDerivedClassNames(BaseNames, Excluded, DerivedNames);
+	}
+
+	FARFilter Filter;
+	Filter.ClassNames.Add(UBlueprint::StaticClass()->GetFName());
+	Filter.bRecursiveClasses = true;
+	Filter.bRecursivePaths = true;
+
+	TArray< FAssetData > AssetList;
+	AssetRegistry.GetAssets(Filter, AssetList);
+
+	// Iterate over retrieved blueprint assets
+	for (const FAssetData& Asset : AssetList)
+	{
+		//Only get the asset with the native parent class using UInterchangePipelineBase
+		FAssetDataTagMapSharedView::FFindTagResult GeneratedClassPath = Asset.TagsAndValues.FindTag(TEXT("GeneratedClass"));
+		if (GeneratedClassPath.IsSet())
+		{
+			// Convert path to just the name part
+			const FString ClassObjectPath = FPackageName::ExportTextPathToObjectPath(*GeneratedClassPath.GetValue());
+			const FString ClassName = FPackageName::ObjectPathToObjectName(ClassObjectPath);
+
+			// Check if this class is in the derived set
+			if (!DerivedNames.Contains(*ClassName))
+			{
+				continue;
+			}
+
+			UBlueprint* Blueprint = Cast<UBlueprint>(Asset.GetAsset());
+			check(Blueprint);
+			check(Blueprint->ParentClass == UInterchangePipelineBase::StaticClass());
+			PipelineCandidates.AddUnique(Blueprint->GeneratedClass);
+		}
+	}
+}
diff --git a/Engine/Source/Runtime/Interchange/Engine/Private/Tasks/InterchangeTaskCompletion.cpp b/Engine/Source/Runtime/Interchange/Engine/Private/Tasks/InterchangeTaskCompletion.cpp
new file mode 100644
--- /dev/null
+++ b/Engine/Source/Runtime/Interchange/Engine/Private/Tasks/InterchangeTaskCompletion.cpp
@@ -0,0 +1,74 @@
+// Copyright Epic Games, Inc. All Rights Reserved.
+#include "InterchangeTaskCompletion.h"
+
+#include "AssetRegistryModule.h"
+#include "Async/TaskGraphInterfaces.h"
+#include "CoreMinimal.h"
+#include "InterchangeFactoryBase.h"
+#include "InterchangeManager.h"
+#include "Stats/Stats.h"
+#include "Templates/SharedPointer.h"
+#include "UObject/Object.h"
+#include "UObject/ObjectMacros.h"
+#include "UObject/WeakObjectPtrTemplates.h"
+
+void UE::Interchange::FTaskCompletion::DoTask(ENamedThreads::Type CurrentThread, const FGraphEventRef& MyCompletionGraphEvent)
+{
+	TSharedPtr<FImportAsyncHelper, ESPMode::ThreadSafe> AsyncHelper = WeakAsyncHelper.Pin();
+	check(AsyncHelper.IsValid());
+	bool bIsFutureRootObjectSet = false;
+	for(TPair<int32, TArray<FImportAsyncHelper::FImportedAssetInfo>>& AssetInfosPerSourceIndexPair : AsyncHelper->ImportedAssetsPerSourceIndex)
+	{
+		const int32 SourceIndex = AssetInfosPerSourceIndexPair.Key;
+		const bool bCallPostImportGameThreadCallback = ensure(AsyncHelper->SourceDatas.IsValidIndex(SourceIndex));
+		for (const FImportAsyncHelper::FImportedAssetInfo& AssetInfo : AssetInfosPerSourceIndexPair.Value)
+		{
+			UObject* Asset = AssetInfo.ImportAsset;
+			//In case Some factory code cannot run outside of the main thread we offer this callback to finish the work before calling post edit change (building the asset)
+			if(bCallPostImportGameThreadCallback && AssetInfo.Factory)
+			{
+				UInterchangeFactoryBase::FPostImportGameThreadCallbackParams Arguments;
+				Arguments.ReimportObject = Asset;
+				Arguments.SourceData = AsyncHelper->SourceDatas[SourceIndex];
+				AssetInfo.Factory->PostImportGameThreadCallback(Arguments);
+			}
+
+			//Clear any async flag from the created asset
+			Asset->ClearInternalFlags(EInternalObjectFlags::Async);
+			//Make sure the package is dirty
+			Asset->MarkPackageDirty();
+#if WITH_EDITOR
+			//Make sure the asset is built correctly
+			Asset->PostEditChange();
+#endif
+			//Post import broadcast
+			if (AsyncHelper->TaskData.ReimportObject)
+			{
+				InterchangeManager->OnAssetPostReimport.Broadcast(Asset);
+			}
+			else
+			{
+				InterchangeManager->OnAssetPostImport.Broadcast(Asset);
+			}
+			//Notify the asset registry
+			FAssetRegistryModule::AssetCreated(Asset);
+
+			if (!bIsFutureRootObjectSet && SourceIndex == 0)
+			{
+				bIsFutureRootObjectSet = true;
+				AsyncHelper->RootObject.SetValue(Asset);
+				AsyncHelper->RootObjectCompletionEvent->DispatchSubsequents();
+			}
+		}
+	}
+
+	if (!AsyncHelper->RootObjectCompletionEvent->IsComplete())
+	{
+		AsyncHelper->RootObject.SetValue(nullptr);
+		AsyncHelper->RootObjectCompletionEvent->DispatchSubsequents();
+	}
+
+	//Release the async helper
+	AsyncHelper = nullptr;
+	InterchangeManager->ReleaseAsyncHelper(WeakAsyncHelper);
+}
diff --git a/Engine/Source/Runtime/Interchange/Engine/Private/Tasks/InterchangeTaskCompletion.h b/Engine/Source/Runtime/Interchange/Engine/Private/Tasks/InterchangeTaskCompletion.h
new file mode 100644
--- /dev/null
+++ b/Engine/Source/Runtime/Interchange/Engine/Private/Tasks/InterchangeTaskCompletion.h
@@ -0,0 +1,47 @@
+// Copyright Epic Games, Inc. All Rights Reserved.
+
+#pragma once
+
+#include "Async/TaskGraphInterfaces.h"
+#include "CoreMinimal.h"
+#include "InterchangeManager.h"
+#include "Stats/Stats.h"
+#include "UObject/WeakObjectPtrTemplates.h"
+
+namespace UE
+{
+	namespace Interchange
+	{
+
+		class FTaskCompletion
+		{
+		private:
+			UInterchangeManager* InterchangeManager;
+			TWeakPtr<FImportAsyncHelper, ESPMode::ThreadSafe> WeakAsyncHelper;
+		public:
+			FTaskCompletion(UInterchangeManager* InInterchangeManager, TWeakPtr<FImportAsyncHelper, ESPMode::ThreadSafe> InAsyncHelper)
+				: InterchangeManager(InInterchangeManager)
+				, WeakAsyncHelper(InAsyncHelper)
+			{
+			}
+
+			static FORCEINLINE ENamedThreads::Type GetDesiredThread()
+			{
+				return ENamedThreads::GameThread;
+			}
+			static FORCEINLINE ESubsequentsMode::Type GetSubsequentsMode()
+			{
+				return ESubsequentsMode::FireAndForget;
+			}
+
+			FORCEINLINE TStatId GetStatId() const
+			{
+				RETURN_QUICK_DECLARE_CYCLE_STAT(FTaskCompletion, STATGROUP_TaskGraphTasks);
+			}
+
+			void DoTask(ENamedThreads::Type CurrentThread, const FGraphEventRef& MyCompletionGraphEvent);
+		};
+
+	} //ns Interchange
+}//ns UE
+
diff --git a/Engine/Source/Runtime/Interchange/Engine/Private/Tasks/InterchangeTaskCreateAsset.cpp b/Engine/Source/Runtime/Interchange/Engine/Private/Tasks/InterchangeTaskCreateAsset.cpp
new file mode 100644
--- /dev/null
+++ b/Engine/Source/Runtime/Interchange/Engine/Private/Tasks/InterchangeTaskCreateAsset.cpp
@@ -0,0 +1,172 @@
+// Copyright Epic Games, Inc. All Rights Reserved.
+#include "InterchangeTaskCreateAsset.h"
+
+#include "Async/TaskGraphInterfaces.h"
+#include "CoreMinimal.h"
+#include "InterchangeFactoryBase.h"
+#include "InterchangeEngineLogPrivate.h"
+#include "InterchangeManager.h"
+#include "InterchangeSourceData.h"
+#include "InterchangeTranslatorBase.h"
+#include "Misc/Paths.h"
+#include "PackageUtils/PackageUtils.h"
+#include "Stats/Stats.h"
+#include "Templates/SharedPointer.h"
+#include "UObject/Object.h"
+#include "UObject/UObjectGlobals.h"
+#include "UObject/ObjectMacros.h"
+#include "UObject/WeakObjectPtrTemplates.h"
+#include "Nodes/InterchangeBaseNode.h"
+
+namespace UE
+{
+	namespace Interchange
+	{
+		namespace Private
+		{
+			void InternalGetPackageName(const UE::Interchange::FImportAsyncHelper& AsyncHelper, const int32 SourceIndex, const FString& PackageBasePath, const UInterchangeBaseNode* Node, FString& OutPackageName, FString& OutAssetName)
+			{
+				const UInterchangeSourceData* SourceData = AsyncHelper.SourceDatas[SourceIndex];
+				check(SourceData);
+				FString NodeDisplayName = Node->GetDisplayLabel().ToString();
+				const FString BaseFileName = FPaths::GetBaseFilename(SourceData->GetFilename());
+
+				//Set the asset name and the package name
+				if (NodeDisplayName.Equals(BaseFileName) || BaseFileName.IsEmpty())
+				{
+					OutAssetName = NodeDisplayName;
+				}
+				else
+				{
+					OutAssetName = BaseFileName.IsEmpty() ? NodeDisplayName : BaseFileName + TEXT("_") + NodeDisplayName;
+				}
+				OutPackageName = FPaths::Combine(*PackageBasePath, *OutAssetName);
+
+				//Sanitize only the package name
+				UE::Interchange::SanitizeInvalidChar(OutPackageName);
+			}
+		}//ns Private
+	}//ns Interchange
+}//ns UE
+
+void UE::Interchange::FTaskCreatePackage::DoTask(ENamedThreads::Type CurrentThread, const FGraphEventRef& MyCompletionGraphEvent)
+{
+	TSharedPtr<UE::Interchange::FImportAsyncHelper, ESPMode::ThreadSafe> AsyncHelper = WeakAsyncHelper.Pin();
+	check(AsyncHelper.IsValid());
+
+	//The create package thread must always execute on the game thread
+	check(IsInGameThread());
+
+	UPackage* Pkg = nullptr;
+	FString PackageName;
+	FString AssetName;
+	//If we do a reimport no need to create a package
+	if (AsyncHelper->TaskData.ReimportObject)
+	{
+		Pkg = AsyncHelper->TaskData.ReimportObject->GetPackage();
+		PackageName = Pkg->GetPathName();
+	}
+	else
+	{
+		Private::InternalGetPackageName(*AsyncHelper, SourceIndex, PackageBasePath, Node, PackageName, AssetName);
+		// We can not create assets that share the name of a map file in the same location
+		if (UE::Interchange::FPackageUtils::IsMapPackageAsset(PackageName))
+		{
+			const FText Message = FText::Format(NSLOCTEXT("Interchange", "AssetNameInUseByMap", "You can not create an asset named '{0}' because there is already a map file with this name in this folder."), FText::FromString(AssetName));
+			UE_LOG(LogInterchangeEngine, Warning, TEXT("%s"), *Message.ToString());
+			//Skip this asset
+			return;
+		}
+
+		Pkg = CreatePackage(*PackageName);
+		if (Pkg == nullptr)
+		{
+			const FText Message = FText::Format(NSLOCTEXT("Interchange", "CannotCreatePackageErrorMsg", "Cannot create package named '{0}', will not import asset {1}."), FText::FromString(PackageName), FText::FromString(AssetName));
+			UE_LOG(LogInterchangeEngine, Warning, TEXT("%s"), *Message.ToString());
+			//Skip this asset
+			return;
+		}
+
+		//Import Asset describe by the node
+		UInterchangeFactoryBase::FCreateAssetParams CreateAssetParams;
+		CreateAssetParams.AssetName = AssetName;
+		CreateAssetParams.AssetNode = Node;
+		CreateAssetParams.Parent = Pkg;
+		CreateAssetParams.SourceData = AsyncHelper->SourceDatas[SourceIndex];
+		CreateAssetParams.Translator = nullptr;
+		if (AsyncHelper->BaseNodeContainers.IsValidIndex(SourceIndex))
+		{
+			CreateAssetParams.NodeContainer = AsyncHelper->BaseNodeContainers[SourceIndex].Get();
+		}
+		CreateAssetParams.ReimportObject = AsyncHelper->TaskData.ReimportObject;
+		//Make sure the asset UObject is created with the correct type on the main thread
+		Factory->CreateEmptyAsset(CreateAssetParams);
+	}
+	// Make sure the destination package is loaded
+	Pkg->FullyLoad();
+
+	FScopeLock Lock(&AsyncHelper->CreatedPackagesLock);
+	AsyncHelper->CreatedPackages.Add(PackageName, Pkg);
+}
+
+void UE::Interchange::FTaskCreateAsset::DoTask(ENamedThreads::Type CurrentThread, const FGraphEventRef& MyCompletionGraphEvent)
+{
+	TSharedPtr<UE::Interchange::FImportAsyncHelper, ESPMode::ThreadSafe> AsyncHelper = WeakAsyncHelper.Pin();
+	check(WeakAsyncHelper.IsValid());
+
+	FString PackageName;
+	FString AssetName;
+	Private::InternalGetPackageName(*AsyncHelper, SourceIndex, PackageBasePath, Node, PackageName, AssetName);
+	if (AsyncHelper->TaskData.ReimportObject)
+	{
+		UPackage* Pkg = AsyncHelper->TaskData.ReimportObject->GetPackage();
+		PackageName = Pkg->GetPathName();
+	}
+
+	UPackage* Pkg = nullptr;
+	{
+		FScopeLock Lock(&AsyncHelper->CreatedPackagesLock);
+		UPackage** PkgPtr = AsyncHelper->CreatedPackages.Find(PackageName);
+
+		if (!PkgPtr || !(*PkgPtr))
+		{
+			const FText Message = FText::Format(NSLOCTEXT("Interchange", "CannotCreateAssetNoPackageErrorMsg", "Cannot create asset named '{1}', package '{0}'was not created properly."), FText::FromString(PackageName), FText::FromString(AssetName));
+			UE_LOG(LogInterchangeEngine, Warning, TEXT("%s"), *Message.ToString());
+			return;
+		}
+
+		if (!AsyncHelper->SourceDatas.IsValidIndex(SourceIndex) || !AsyncHelper->Translators.IsValidIndex(SourceIndex))
+		{
+			const FText Message = FText::Format(NSLOCTEXT("Interchange", "CannotCreateAssetMissingDataErrorMsg", "Cannot create asset named '{0}', Source data or translator is invalid."), FText::FromString(AssetName));
+			UE_LOG(LogInterchangeEngine, Warning, TEXT("%s"), *Message.ToString());
+			return;
+		}
+
+		Pkg = *PkgPtr;
+	}
+
+	UInterchangeTranslatorBase* Translator = AsyncHelper->Translators[SourceIndex];
+	//Import Asset describe by the node
+	UInterchangeFactoryBase::FCreateAssetParams CreateAssetParams;
+	CreateAssetParams.AssetName = AssetName;
+	CreateAssetParams.AssetNode = Node;
+	CreateAssetParams.Parent = Pkg;
+	CreateAssetParams.SourceData = AsyncHelper->SourceDatas[SourceIndex];
+	CreateAssetParams.Translator = Translator;
+	if (AsyncHelper->BaseNodeContainers.IsValidIndex(SourceIndex))
+	{
+		CreateAssetParams.NodeContainer = AsyncHelper->BaseNodeContainers[SourceIndex].Get();
+	}
+	CreateAssetParams.ReimportObject = AsyncHelper->TaskData.ReimportObject;
+
+	UObject* NodeAsset = Factory->CreateAsset(CreateAssetParams);
+	if (NodeAsset)
+	{
+		FScopeLock Lock(&AsyncHelper->ImportedAssetsPerSourceIndexLock);
+		TArray<UE::Interchange::FImportAsyncHelper::FImportedAssetInfo>& ImportedInfos = AsyncHelper->ImportedAssetsPerSourceIndex.FindOrAdd(SourceIndex);
+		UE::Interchange::FImportAsyncHelper::FImportedAssetInfo& AssetInfo = ImportedInfos.AddZeroed_GetRef();
+		AssetInfo.ImportAsset = NodeAsset;
+		Node->ReferenceObject = FSoftObjectPath(NodeAsset);
+		AssetInfo.Factory = Factory;
+	}
+}
\ No newline at end of file
diff --git a/Engine/Source/Runtime/Interchange/Engine/Private/Tasks/InterchangeTaskCreateAsset.h b/Engine/Source/Runtime/Interchange/Engine/Private/Tasks/InterchangeTaskCreateAsset.h
new file mode 100644
--- /dev/null
+++ b/Engine/Source/Runtime/Interchange/Engine/Private/Tasks/InterchangeTaskCreateAsset.h
@@ -0,0 +1,106 @@
+// Copyright Epic Games, Inc. All Rights Reserved.
+
+#pragma once
+
+#include "Async/TaskGraphInterfaces.h"
+#include "CoreMinimal.h"
+#include "InterchangeFactoryBase.h"
+#include "InterchangeManager.h"
+#include "Stats/Stats.h"
+#include "UObject/WeakObjectPtrTemplates.h"
+#include "Nodes/InterchangeBaseNode.h"
+
+namespace UE
+{
+	namespace Interchange
+	{
+
+		/**
+		 * This task create package, Cook::PackageTracker::NotifyUObjectCreated is not thread safe, so we need to create the packages on the main thread
+		 */
+		class FTaskCreatePackage
+		{
+		private:
+			FString PackageBasePath;
+			int32 SourceIndex;
+			TWeakPtr<FImportAsyncHelper, ESPMode::ThreadSafe> WeakAsyncHelper;
+			const UInterchangeBaseNode* Node;
+			UInterchangeFactoryBase* Factory;
+
+		public:
+			FTaskCreatePackage(const FString& InPackageBasePath, const int32 InSourceIndex, TWeakPtr<FImportAsyncHelper, ESPMode::ThreadSafe> InAsyncHelper, const UInterchangeBaseNode* InNode, UInterchangeFactoryBase* InFactory)
+				: PackageBasePath(InPackageBasePath)
+				, SourceIndex(InSourceIndex)
+				, WeakAsyncHelper(InAsyncHelper)
+				, Node(InNode)
+				, Factory(InFactory)
+			{
+				check(Node);
+				check(Factory);
+			}
+
+			ENamedThreads::Type GetDesiredThread()
+			{
+				if (WeakAsyncHelper.IsValid() && WeakAsyncHelper.Pin()->TaskData.ReimportObject)
+				{
+					//When doing a reimport the package already exist, so we can get it outside of the main thread
+					return ENamedThreads::AnyBackgroundThreadNormalTask;
+				}
+				return ENamedThreads::GameThread;
+			}
+
+			static ESubsequentsMode::Type GetSubsequentsMode()
+			{
+				return ESubsequentsMode::TrackSubsequents;
+			}
+
+			TStatId GetStatId() const
+			{
+				RETURN_QUICK_DECLARE_CYCLE_STAT(FTaskCreatePackage, STATGROUP_TaskGraphTasks);
+			}
+
+			void DoTask(ENamedThreads::Type CurrentThread, const FGraphEventRef& MyCompletionGraphEvent);
+		};
+
+		class FTaskCreateAsset
+		{
+		private:
+			FString PackageBasePath;
+			int32 SourceIndex;
+			TWeakPtr<FImportAsyncHelper, ESPMode::ThreadSafe> WeakAsyncHelper;
+			const UInterchangeBaseNode* Node;
+			UInterchangeFactoryBase* Factory;
+
+		public:
+			FTaskCreateAsset(const FString& InPackageBasePath, const int32 InSourceIndex, TWeakPtr<FImportAsyncHelper, ESPMode::ThreadSafe> InAsyncHelper, const UInterchangeBaseNode* InNode, UInterchangeFactoryBase* InFactory)
+				: PackageBasePath(InPackageBasePath)
+				, SourceIndex(InSourceIndex)
+				, WeakAsyncHelper(InAsyncHelper)
+				, Node(InNode)
+				, Factory(InFactory)
+			{
+				check(Node);
+				check(Factory);
+			}
+
+			ENamedThreads::Type GetDesiredThread()
+			{
+				return Factory->CanExecuteOnAnyThread() ? ENamedThreads::AnyBackgroundThreadNormalTask : ENamedThreads::GameThread;
+			}
+
+			static ESubsequentsMode::Type GetSubsequentsMode()
+			{
+				return ESubsequentsMode::TrackSubsequents;
+			}
+
+			TStatId GetStatId() const
+			{
+				RETURN_QUICK_DECLARE_CYCLE_STAT(FTaskCreateAsset, STATGROUP_TaskGraphTasks);
+			}
+
+			void DoTask(ENamedThreads::Type CurrentThread, const FGraphEventRef& MyCompletionGraphEvent);
+		};
+
+
+	} //ns Interchange
+}//ns UE
diff --git a/Engine/Source/Runtime/Interchange/Engine/Private/Tasks/InterchangeTaskParsing.cpp b/Engine/Source/Runtime/Interchange/Engine/Private/Tasks/InterchangeTaskParsing.cpp
new file mode 100644
--- /dev/null
+++ b/Engine/Source/Runtime/Interchange/Engine/Private/Tasks/InterchangeTaskParsing.cpp
@@ -0,0 +1,115 @@
+// Copyright Epic Games, Inc. All Rights Reserved.
+#include "InterchangeTaskParsing.h"
+
+#include "Async/TaskGraphInterfaces.h"
+#include "CoreMinimal.h"
+#include "InterchangeFactoryBase.h"
+#include "InterchangeManager.h"
+#include "InterchangeSourceData.h"
+#include "InterchangeTaskCompletion.h"
+#include "InterchangeTaskCreateAsset.h"
+#include "InterchangeTranslatorBase.h"
+#include "Misc/Paths.h"
+#include "PackageUtils/PackageUtils.h"
+#include "Stats/Stats.h"
+#include "Templates/SharedPointer.h"
+#include "UObject/Object.h"
+#include "UObject/UObjectGlobals.h"
+#include "UObject/ObjectMacros.h"
+#include "UObject/WeakObjectPtrTemplates.h"
+#include "Nodes/InterchangeBaseNode.h"
+#include "Nodes/InterchangeBaseNodeContainer.h"
+
+void UE::Interchange::FTaskParsing::DoTask(ENamedThreads::Type CurrentThread, const FGraphEventRef& MyCompletionGraphEvent)
+{
+	TSharedPtr<FImportAsyncHelper, ESPMode::ThreadSafe> AsyncHelper = WeakAsyncHelper.Pin();
+	check(AsyncHelper.IsValid());
+
+	//Parse each graph and prepare import task data, we will then be able to create all the task with the correct dependencies
+	struct FTaskData
+	{
+		FName UniqueID;
+		int32 SourceIndex;
+		const UInterchangeBaseNode* Node;
+		TArray<FName> Dependencies;
+		FGraphEventRef GraphEventRef;
+		FGraphEventArray Prerequistes;
+		UInterchangeFactoryBase* Factory;
+	};
+
+
+
+	FGraphEventArray CompletionPrerequistes;
+	TArray<FTaskData> TaskDatas;
+	for (int32 SourceIndex = 0; SourceIndex < AsyncHelper->SourceDatas.Num(); ++SourceIndex)
+	{
+		if (!AsyncHelper->BaseNodeContainers.IsValidIndex(SourceIndex))
+		{
+			continue;
+		}
+		UInterchangeBaseNodeContainer* BaseNodeContainer = AsyncHelper->BaseNodeContainers[SourceIndex].Get();
+		if (!BaseNodeContainer)
+		{
+			continue;
+		}
+		BaseNodeContainer->IterateNodes([&](const FName NodeUID, const UInterchangeBaseNode* Node)
+		{
+			if (Node->GetAssetClass() != nullptr)
+			{
+				UInterchangeFactoryBase* NodeFactory = InterchangeManager->GetRegisterFactory(Node->GetAssetClass());
+				if (!NodeFactory)
+				{
+					//nothing we can import from this element
+					return;
+				}
+				FTaskData& NodeTaskData = TaskDatas.AddDefaulted_GetRef();
+				NodeTaskData.UniqueID = Node->GetUniqueID();
+				NodeTaskData.SourceIndex = SourceIndex;
+				NodeTaskData.Node = Node;
+				Node->GetDependecies(NodeTaskData.Dependencies);
+				NodeTaskData.Factory = NodeFactory;
+			}
+		});
+	}
+
+	//Sort per dependencies
+	TaskDatas.Sort([](const FTaskData& A, const FTaskData& B)
+	{
+		//if A is a dependency of B then return true to do A before B
+		if (B.Dependencies.Contains(A.UniqueID))
+		{
+			return true;
+		}
+		if (A.Dependencies.Contains(B.UniqueID))
+		{
+			return false;
+		}
+		return A.Dependencies.Num() <= B.Dependencies.Num();
+	});
+
+	for (int32 TaskIndex = 0; TaskIndex < TaskDatas.Num(); ++TaskIndex)
+	{
+		if (TaskDatas[TaskIndex].Dependencies.Num() > 0)
+		{
+			//Search the previous node to find the dependence
+			for (int32 DepTaskIndex = 0; DepTaskIndex < TaskIndex; ++DepTaskIndex)
+			{
+				if (TaskDatas[TaskIndex].Dependencies.Contains(TaskDatas[DepTaskIndex].UniqueID))
+				{
+					//Add has prerequisite
+					TaskDatas[TaskIndex].Prerequistes.Add(TaskDatas[DepTaskIndex].GraphEventRef);
+				}
+			}
+		}
+
+		//Add create package task has a prerequisite of FTaskCreateAsset. Create package task is a game thread task
+		FGraphEventArray CreatePackagePrerequistes;
+		CreatePackagePrerequistes.Add(TGraphTask<FTaskCreatePackage>::CreateTask(&(TaskDatas[TaskIndex].Prerequistes)).ConstructAndDispatchWhenReady(PackageBasePath, TaskDatas[TaskIndex].SourceIndex, WeakAsyncHelper, TaskDatas[TaskIndex].Node, TaskDatas[TaskIndex].Factory));
+
+		int32 CreateTaskIndex = AsyncHelper->CreateAssetTasks.Add(TGraphTask<FTaskCreateAsset>::CreateTask(&(CreatePackagePrerequistes)).ConstructAndDispatchWhenReady(PackageBasePath, TaskDatas[TaskIndex].SourceIndex, WeakAsyncHelper, TaskDatas[TaskIndex].Node, TaskDatas[TaskIndex].Factory));
+		TaskDatas[TaskIndex].GraphEventRef = AsyncHelper->CreateAssetTasks[CreateTaskIndex];
+		CompletionPrerequistes.Add(TaskDatas[TaskIndex].GraphEventRef);
+	}
+
+	AsyncHelper->CompletionTask = TGraphTask<FTaskCompletion>::CreateTask(&CompletionPrerequistes).ConstructAndDispatchWhenReady(InterchangeManager, WeakAsyncHelper);
+}
diff --git a/Engine/Source/Runtime/Interchange/Engine/Private/Tasks/InterchangeTaskParsing.h b/Engine/Source/Runtime/Interchange/Engine/Private/Tasks/InterchangeTaskParsing.h
new file mode 100644
--- /dev/null
+++ b/Engine/Source/Runtime/Interchange/Engine/Private/Tasks/InterchangeTaskParsing.h
@@ -0,0 +1,50 @@
+// Copyright Epic Games, Inc. All Rights Reserved.
+
+#pragma once
+
+#include "Async/TaskGraphInterfaces.h"
+#include "CoreMinimal.h"
+#include "InterchangeManager.h"
+#include "Stats/Stats.h"
+#include "UObject/WeakObjectPtrTemplates.h"
+
+namespace UE
+{
+	namespace Interchange
+	{
+
+		class FTaskParsing
+		{
+		private:
+			UInterchangeManager* InterchangeManager;
+			TWeakPtr<FImportAsyncHelper, ESPMode::ThreadSafe> WeakAsyncHelper;
+			FString PackageBasePath;
+		public:
+			FTaskParsing(UInterchangeManager* InInterchangeManager, const FString& InPackageBasePath, TWeakPtr<FImportAsyncHelper, ESPMode::ThreadSafe> InAsyncHelper)
+				: InterchangeManager(InInterchangeManager)
+				, WeakAsyncHelper(InAsyncHelper)
+				, PackageBasePath(InPackageBasePath)
+			{
+				check(InterchangeManager);
+			}
+
+			static FORCEINLINE ENamedThreads::Type GetDesiredThread()
+			{
+				return ENamedThreads::AnyBackgroundThreadNormalTask;
+			}
+			static FORCEINLINE ESubsequentsMode::Type GetSubsequentsMode()
+			{
+
+				return ESubsequentsMode::TrackSubsequents;
+			}
+
+			FORCEINLINE TStatId GetStatId() const
+			{
+				RETURN_QUICK_DECLARE_CYCLE_STAT(FTaskParsing, STATGROUP_TaskGraphTasks);
+			}
+
+			void DoTask(ENamedThreads::Type CurrentThread, const FGraphEventRef& MyCompletionGraphEvent);
+		};
+
+	} //ns Interchange
+}//ns UE
diff --git a/Engine/Source/Runtime/Interchange/Engine/Private/Tasks/InterchangeTaskPipeline.cpp b/Engine/Source/Runtime/Interchange/Engine/Private/Tasks/InterchangeTaskPipeline.cpp
new file mode 100644
--- /dev/null
+++ b/Engine/Source/Runtime/Interchange/Engine/Private/Tasks/InterchangeTaskPipeline.cpp
@@ -0,0 +1,26 @@
+// Copyright Epic Games, Inc. All Rights Reserved.
+#include "InterchangeTaskPipeline.h"
+
+#include "Async/TaskGraphInterfaces.h"
+#include "CoreMinimal.h"
+#include "InterchangeManager.h"
+#include "InterchangePipelineBase.h"
+#include "Stats/Stats.h"
+#include "Templates/SharedPointer.h"
+#include "UObject/WeakObjectPtrTemplates.h"
+#include "Nodes/InterchangeBaseNodeContainer.h"
+
+void UE::Interchange::FTaskPipeline::DoTask(ENamedThreads::Type CurrentThread, const FGraphEventRef& MyCompletionGraphEvent)
+{
+	TSharedPtr<FImportAsyncHelper, ESPMode::ThreadSafe> AsyncHelper = WeakAsyncHelper.Pin();
+	check(AsyncHelper.IsValid());
+
+	if (UInterchangePipelineBase* Pipeline = PipelineBase.Get())
+	{
+		for (int32 GraphIndex = 0; GraphIndex < AsyncHelper->BaseNodeContainers.Num(); ++GraphIndex)
+		{
+			check(AsyncHelper->BaseNodeContainers[GraphIndex].IsValid());
+			Pipeline->ScriptedExecuteImportPipeline(AsyncHelper->BaseNodeContainers[GraphIndex].Get());
+		}
+	}
+}
diff --git a/Engine/Source/Runtime/Interchange/Engine/Private/Tasks/InterchangeTaskPipeline.h b/Engine/Source/Runtime/Interchange/Engine/Private/Tasks/InterchangeTaskPipeline.h
new file mode 100644
--- /dev/null
+++ b/Engine/Source/Runtime/Interchange/Engine/Private/Tasks/InterchangeTaskPipeline.h
@@ -0,0 +1,49 @@
+// Copyright Epic Games, Inc. All Rights Reserved.
+
+#pragma once
+
+#include "Async/TaskGraphInterfaces.h"
+#include "CoreMinimal.h"
+#include "InterchangeManager.h"
+#include "InterchangePipelineBase.h"
+#include "Stats/Stats.h"
+#include "UObject/WeakObjectPtrTemplates.h"
+
+namespace UE
+{
+	namespace Interchange
+	{
+
+		class FTaskPipeline
+		{
+		private:
+			TWeakObjectPtr<UInterchangePipelineBase> PipelineBase;
+			TWeakPtr<FImportAsyncHelper, ESPMode::ThreadSafe> WeakAsyncHelper;
+		public:
+			FTaskPipeline(TWeakObjectPtr<UInterchangePipelineBase> InPipelineBase, TWeakPtr<FImportAsyncHelper, ESPMode::ThreadSafe> InAsyncHelper)
+				: PipelineBase(InPipelineBase)
+				, WeakAsyncHelper(InAsyncHelper)
+			{
+			}
+
+			ENamedThreads::Type GetDesiredThread()
+			{
+				return PipelineBase.Get()->ScriptedCanExecuteOnAnyThread() ? ENamedThreads::AnyBackgroundThreadNormalTask : ENamedThreads::GameThread;
+			}
+
+			static ESubsequentsMode::Type GetSubsequentsMode()
+			{
+				return ESubsequentsMode::TrackSubsequents;
+			}
+
+			FORCEINLINE TStatId GetStatId() const
+			{
+				RETURN_QUICK_DECLARE_CYCLE_STAT(FTaskPipeline, STATGROUP_TaskGraphTasks);
+			}
+
+			void DoTask(ENamedThreads::Type CurrentThread, const FGraphEventRef& MyCompletionGraphEvent);
+		};
+
+
+	} //ns Interchange
+}//ns UE
diff --git a/Engine/Source/Runtime/Interchange/Engine/Private/Tasks/InterchangeTaskTranslator.cpp b/Engine/Source/Runtime/Interchange/Engine/Private/Tasks/InterchangeTaskTranslator.cpp
new file mode 100644
--- /dev/null
+++ b/Engine/Source/Runtime/Interchange/Engine/Private/Tasks/InterchangeTaskTranslator.cpp
@@ -0,0 +1,38 @@
+// Copyright Epic Games, Inc. All Rights Reserved.
+#include "InterchangeTaskTranslator.h"
+
+#include "Async/TaskGraphInterfaces.h"
+#include "CoreMinimal.h"
+#include "InterchangeManager.h"
+#include "InterchangeSourceData.h"
+#include "InterchangeTranslatorBase.h"
+#include "Stats/Stats.h"
+#include "Templates/SharedPointer.h"
+#include "UObject/WeakObjectPtrTemplates.h"
+#include "Nodes/InterchangeBaseNodeContainer.h"
+
+void UE::Interchange::FTaskTranslator::DoTask(ENamedThreads::Type CurrentThread, const FGraphEventRef& MyCompletionGraphEvent)
+{
+	TSharedPtr<FImportAsyncHelper, ESPMode::ThreadSafe> AsyncHelper = WeakAsyncHelper.Pin();
+	check(AsyncHelper.IsValid());
+
+	UInterchangeTranslatorBase* Translator = AsyncHelper->Translators.IsValidIndex(SourceIndex) ? AsyncHelper->Translators[SourceIndex] : nullptr;
+	if (!Translator)
+	{
+		return;
+	}
+	const UInterchangeSourceData* SourceData = AsyncHelper->SourceDatas.IsValidIndex(SourceIndex) ? AsyncHelper->SourceDatas[SourceIndex] : nullptr;
+	if (!SourceData)
+	{
+		return;
+	}
+
+	if (!AsyncHelper->BaseNodeContainers.IsValidIndex(SourceIndex) || !AsyncHelper->BaseNodeContainers[SourceIndex].IsValid())
+	{
+		return;
+	}
+
+	//Translate the source data
+	UInterchangeBaseNodeContainer& BaseNodeContainer = *(AsyncHelper->BaseNodeContainers[SourceIndex].Get());
+	Translator->Translate(SourceData, BaseNodeContainer);
+}
diff --git a/Engine/Source/Runtime/Interchange/Engine/Private/Tasks/InterchangeTaskTranslator.h b/Engine/Source/Runtime/Interchange/Engine/Private/Tasks/InterchangeTaskTranslator.h
new file mode 100644
--- /dev/null
+++ b/Engine/Source/Runtime/Interchange/Engine/Private/Tasks/InterchangeTaskTranslator.h
@@ -0,0 +1,48 @@
+// Copyright Epic Games, Inc. All Rights Reserved.
+
+#pragma once
+
+#include "Async/TaskGraphInterfaces.h"
+#include "CoreMinimal.h"
+#include "InterchangeManager.h"
+#include "Stats/Stats.h"
+#include "UObject/WeakObjectPtrTemplates.h"
+
+namespace UE
+{
+	namespace Interchange
+	{
+
+		class FTaskTranslator
+		{
+		private:
+			int32 SourceIndex = INDEX_NONE;
+			TWeakPtr<FImportAsyncHelper, ESPMode::ThreadSafe> WeakAsyncHelper;
+
+		public:
+			FTaskTranslator(int32 InSourceIndex, TWeakPtr<FImportAsyncHelper, ESPMode::ThreadSafe> InAsyncHelper)
+				: SourceIndex(InSourceIndex)
+				, WeakAsyncHelper(InAsyncHelper)
+			{
+			}
+
+			static FORCEINLINE ENamedThreads::Type GetDesiredThread()
+			{
+				return ENamedThreads::AnyBackgroundThreadNormalTask;
+			}
+			static FORCEINLINE ESubsequentsMode::Type GetSubsequentsMode()
+			{
+				return ESubsequentsMode::TrackSubsequents;
+			}
+
+			FORCEINLINE TStatId GetStatId() const
+			{
+				RETURN_QUICK_DECLARE_CYCLE_STAT(FTaskTranslator, STATGROUP_TaskGraphTasks);
+			}
+
+			void DoTask(ENamedThreads::Type CurrentThread, const FGraphEventRef& MyCompletionGraphEvent);
+		};
+
+
+	} //ns Interchange
+} //ns UE
diff --git a/Engine/Source/Runtime/Interchange/Engine/Public/InterchangeManager.h b/Engine/Source/Runtime/Interchange/Engine/Public/InterchangeManager.h
new file mode 100644
--- /dev/null
+++ b/Engine/Source/Runtime/Interchange/Engine/Public/InterchangeManager.h
@@ -0,0 +1,351 @@
+// Copyright Epic Games, Inc. All Rights Reserved.
+
+#pragma once
+
+#include "Async/TaskGraphInterfaces.h"
+#include "CoreMinimal.h"
+#include "Delegates/DelegateCombinations.h"
+#include "InterchangeFactoryBase.h"
+#include "InterchangePipelineBase.h"
+#include "InterchangeSourceData.h"
+#include "InterchangeTranslatorBase.h"
+#include "InterchangeWriterBase.h"
+#include "Nodes/InterchangeBaseNodeContainer.h"
+#include "UObject/Package.h"
+#include "UObject/Object.h"
+#include "UObject/ObjectMacros.h"
+#include "UObject/StrongObjectPtr.h"
+
+#include "InterchangeManager.generated.h"
+
+class FAsyncTaskNotification;
+namespace UE
+{
+	namespace Interchange
+	{
+		class INTERCHANGEENGINE_API FScopedSourceData
+		{
+		public:
+			explicit FScopedSourceData(const FString& Filename);
+			UInterchangeSourceData* GetSourceData() const;
+		private:
+			TStrongObjectPtr<UInterchangeSourceData> SourceDataPtr = nullptr;
+		};
+
+		class INTERCHANGEENGINE_API FScopedTranslator
+		{
+		public:
+			explicit FScopedTranslator(const UInterchangeSourceData* SourceData);
+			UInterchangeTranslatorBase* GetTranslator();
+
+		private:
+			TStrongObjectPtr<UInterchangeTranslatorBase> ScopedTranslatorPtr = nullptr;
+		};
+
+		enum class EImportType : uint8
+		{
+			ImportType_None,
+			ImportType_Asset,
+			ImportType_Scene
+		};
+
+		struct FImportAsyncHelperData
+		{
+			//True if the import process is unattended. We cannot show UI  if the import is automated
+			bool bIsAutomated = false;
+
+			//We can import assets or full scene
+			EImportType ImportType = EImportType::ImportType_None;
+
+			//True if we are reimporting assets or scene
+			UObject* ReimportObject = nullptr;
+		};
+
+		class FImportAsyncHelper : public FGCObject
+		{
+		public:
+			FImportAsyncHelper();
+
+			~FImportAsyncHelper()
+			{
+				CleanUp();
+			}
+
+			/* FGCObject interface */
+			virtual void AddReferencedObjects(FReferenceCollector& Collector) override;
+
+			TArray <TStrongObjectPtr<UInterchangeBaseNodeContainer>> BaseNodeContainers;
+
+			TArray<UInterchangeSourceData* > SourceDatas;
+			TArray<UInterchangeTranslatorBase* > Translators;
+			TArray<UInterchangePipelineBase* > Pipelines;
+			TArray<UInterchangeFactoryBase* > Factories;
+
+			TArray<FGraphEventRef> TranslatorTasks;
+			TArray<FGraphEventRef> PipelineTasks;
+			FGraphEventRef ParsingTask;
+			TArray<FGraphEventRef> CreateAssetTasks;
+
+			FGraphEventRef CompletionTask;
+
+			//Create package map, Key is package name. We cannot create package asynchronously so we have to create a game thread task to do this
+			FCriticalSection CreatedPackagesLock;
+			TMap<FString, UPackage*> CreatedPackages;
+
+			struct FImportedAssetInfo
+			{
+				UObject* ImportAsset;
+				UInterchangeFactoryBase* Factory;
+				UInterchangeBaseNode* AssetNode;
+			};
+
+			FCriticalSection ImportedAssetsPerSourceIndexLock;
+			TMap<int32, TArray<FImportedAssetInfo>> ImportedAssetsPerSourceIndex;
+
+			FImportAsyncHelperData TaskData;
+
+			TPromise< UObject* > RootObject;
+			FGraphEventRef RootObjectCompletionEvent;
+
+			void CleanUp();
+		};
+
+		class INTERCHANGEENGINE_API FAsyncImportResult
+		{
+		public:
+			FAsyncImportResult() = default;
+			explicit FAsyncImportResult(TFuture< UObject* >&& InFutureObject, const FGraphEventRef& InGraphEvent);
+
+			bool IsValid() const;
+			UObject* Get() const;
+			FAsyncImportResult Next(TFunction< UObject* (UObject*) > Continuation);
+
+		private:
+			TFuture< UObject* > FutureObject;
+			FGraphEventRef GraphEvent;
+		};
+
+		void SanitizeInvalidChar(FString& String);
+
+	} //ns interchange
+} //ns UE
+
+USTRUCT(BlueprintType)
+struct FImportAssetParameters
+{
+	GENERATED_USTRUCT_BODY()
+
+	UPROPERTY(EditAnywhere, BlueprintReadWrite, Category = "Interchange|ImportAsset")
+	UObject* ReimportAsset = nullptr;
+
+	UPROPERTY(EditAnywhere, BlueprintReadWrite, Category = "Interchange|ImportAsset")
+	bool bIsAutomated = false;
+
+	UPROPERTY(EditAnywhere, BlueprintReadWrite, Category = "Interchange|ImportAsset")
+	UInterchangePipelineBase* OverridePipeline = nullptr;
+};
+
+UCLASS(Transient, BlueprintType)
+class INTERCHANGEENGINE_API UInterchangeManager : public UObject
+{
+	GENERATED_BODY()
+public:
+
+	/**
+	 * Return the interchange manager singleton pointer.
+	 *
+	 * @note - We need to return a pointer to have a valid Blueprint callable function
+	 */
+	UFUNCTION(BlueprintCallable, Category = "Interchange | Import Manager")
+	static UInterchangeManager* GetInterchangeManagerScripted()
+	{
+		return &GetInterchangeManager();
+	}
+
+	/** Return the interchange manager singleton.*/
+	static UInterchangeManager& GetInterchangeManager()
+	{
+		static TStrongObjectPtr<UInterchangeManager> InterchangeManager = nullptr;
+		if (!InterchangeManager.IsValid())
+		{
+			//We cannot create a TStrongObjectPtr outside of the main thread, we also need a valid Transient package
+			check(IsInGameThread() && GetTransientPackage());
+			InterchangeManager = TStrongObjectPtr<UInterchangeManager>(NewObject<UInterchangeManager>(GetTransientPackage(), NAME_None, EObjectFlags::RF_NoFlags));
+		}
+		//When we get here we should be valid
+		check(InterchangeManager.IsValid());
+
+		return *(InterchangeManager.Get());
+	}
+
+
+	/** delegate type fired when new assets have been imported. Note: InCreatedObject can be NULL if import failed. Params: UFactory* InFactory, UObject* InCreatedObject */
+	DECLARE_MULTICAST_DELEGATE_OneParam(FInterchangeOnAssetPostImport, UObject*);
+	/** delegate type fired when new assets have been reimported. Note: InCreatedObject can be NULL if import failed. Params: UObject* InCreatedObject */
+	DECLARE_MULTICAST_DELEGATE_OneParam(FInterchangeOnAssetPostReimport, UObject*);
+	
+
+	// Delegates used to register and unregister
+
+	FInterchangeOnAssetPostImport OnAssetPostImport;
+	FInterchangeOnAssetPostReimport OnAssetPostReimport;
+
+	/**
+	 * Any translator must register to the manager
+	 * @Param Translator - The UClass of the translator you want to register
+	 * @return true if the translator class can be register false otherwise.
+	 *
+	 * @Note if you register multiple time the same class it will return true for every call
+	 */
+	bool RegisterTranslator(const UClass* TranslatorClass);
+
+	/**
+	 * Any factory must register to the manager
+	 * @Param Factory - The UClass of the factory you want to register
+	 * @return true if the factory class can be register false otherwise.
+	 *
+	 * @Note if you register multiple time the same class it will return true for every call
+	 */
+	bool RegisterFactory(const UClass* Factory);
+
+	/**
+	 * Any writer must register to the manager
+	 * @Param Writer - The UClass of the writer you want to register
+	 * @return true if the writer class can be register false otherwise.
+	 *
+	 * @Note if you register multiple time the same class it will return true for every call
+	 */
+	bool RegisterWriter(const UClass* Writer);
+
+	/**
+	 * Look if there is a registered translator for this source data.
+	 * This allow us to by pass the original asset tools system to import supported asset.
+	 * @Param SourceData - The source data input we want to translate to Uod
+	 * @return True if there is a registered translator that can handle handle this source data, false otherwise.
+	 */
+	bool CanTranslateSourceData(const UInterchangeSourceData* SourceData) const;
+
+	/**
+	 * Call this to start an import asset process, the caller must specify a source data.
+	 * This import process can import many different asset, but all in the game content.
+	 *
+	 * @Param ContentPath - The content path where to import the assets
+	 * @Param SourceData - The source data input we want to translate
+	 * @param ImportAssetParameters - All import asset parameter we need to pass to the import asset function
+	 * @return true if the import succeed, false otherwise.
+	 */
+	UFUNCTION(BlueprintCallable, Category = "Interchange | Import Manager")
+	bool ImportAsset(const FString& ContentPath, const UInterchangeSourceData* SourceData, const FImportAssetParameters& ImportAssetParameters);
+	UE::Interchange::FAsyncImportResult ImportAssetAsync(const FString& ContentPath, const UInterchangeSourceData* SourceData, const FImportAssetParameters& ImportAssetParameters);
+
+	/**
+	 * Call this to start an import scene process, the caller must specify a source data.
+	 * This import process can import many different asset and there transform (USceneComponent) and store the result in a blueprint and add the blueprint to the level.
+	 *
+	 * @Param ContentPath - The content path where to import the assets
+	 * @Param SourceData - The source data input we want to translate, this object will be duplicate to allow multithread safe operations
+	 * @param ImportAssetParameters - All import asset parameter we need to pass to the import asset function
+	 * @return true if the import succeed, false otherwise.
+	 */
+	UFUNCTION(BlueprintCallable, Category = "Interchange | Import Manager")
+	bool ImportScene(const FString& ContentPath, const UInterchangeSourceData* SourceData, bool bIsReimport = false, bool bIsAutomated = false);
+
+	/**
+	 * Call this to start an export asset process, the caller must specify a source data.
+	 * 
+	 * @Param SourceData - The source data output 
+	 * @Param bIsAutomated - If true the exporter will not show any UI or dialog
+	 * @return true if the import succeed, false otherwise.
+	 */
+	UFUNCTION(BlueprintCallable, Category = "Interchange | Export Manager")
+	bool ExportAsset(const UObject* Asset, bool bIsAutomated = false);
+
+	/**
+	 * Call this to start an export scene process, the caller must specify a source data
+	 * This import process can import many different asset and there transform (USceneComponent) and store the result in a blueprint and add the blueprint to the level.
+	 * @Param SourceData - The source data input we want to translate
+	 * @Param bIsAutomated - If true the import asset will not show any UI or dialog
+	 * @return true if the import succeed, false otherwise.
+	 */
+	UFUNCTION(BlueprintCallable, Category = "Interchange | Export Manager")
+	bool ExportScene(const UObject* World, bool bIsAutomated = false);
+
+	/*
+	* Script helper to create a source data object pointing on a file on disk
+	* @Param InFilename: Specify a file on disk
+	* @return: A new UInterchangeSourceData.
+	*/
+	UFUNCTION(BlueprintCallable, Category = "Interchange | Import Manager")
+	UInterchangeSourceData* CreateSourceData(const FString& InFileName);
+
+	/**
+	* Script helper to get a registered factory for a specified class
+	* @Param FactoryClass: The class we search a registerd factory
+	* @return: if found, we return the factory that is registered. Return NULL if nothing found.
+	*/
+	UFUNCTION(BlueprintCallable, Category = "Interchange | Import Manager")
+	UInterchangeFactoryBase* GetRegisterFactory(UClass* FactoryClass)
+	{
+		for (auto Kvp : RegisteredFactories)
+		{
+			if (FactoryClass->IsChildOf(Kvp.Key))
+			{
+				return Kvp.Value;
+			}
+		}
+		return nullptr;
+	}
+
+	/**
+	 * Return an FImportAsynHelper pointer. The pointer is deleted when ReleaseAsyncHelper is call.
+	 * @param Data - The data we want to pass to the different import tasks
+	 */
+	TWeakPtr<UE::Interchange::FImportAsyncHelper, ESPMode::ThreadSafe> CreateAsyncHelper(const UE::Interchange::FImportAsyncHelperData& Data);
+
+	/** Delete the specified AsyncHelper and remove it from the array that was holding it. */
+	void ReleaseAsyncHelper(TWeakPtr<UE::Interchange::FImportAsyncHelper, ESPMode::ThreadSafe> AsyncHelper);
+
+	/*
+	 * Return the first translator that can translate the source data.
+	 * @Param SourceData - The source data for which we search a translator.
+	 * @return return a matching translator or nullptr if there is no match.
+	 */
+	UInterchangeTranslatorBase* GetTranslatorForSourceData(const UInterchangeSourceData* SourceData) const;
+
+protected:
+
+	/** Return true if we can show some UI */
+	static bool IsAttended();
+
+	/*
+	 * Find all Pipeline candidate (c++, blueprint and python).
+	 * @Param SourceData - The source data for which we search a translator.
+	 * @return return a matching translator or nullptr if there is no match.
+	 */
+	void FindPipelineCandidate(TArray<UClass*>& PipelineCandidates);
+
+private:
+	//By using pointer, there is no issue if the array get resize
+	TArray<TSharedPtr<UE::Interchange::FImportAsyncHelper, ESPMode::ThreadSafe> > ImportTasks;
+
+	TSharedPtr<FAsyncTaskNotification> Notification = nullptr;
+
+	//The manager will create translator at every import, translator must be able to retrieve payload information when the factory ask for it.
+	//The translator stored has value is only use to know if we can use this type of translator.
+	UPROPERTY()
+	TMap<const UClass*, UInterchangeTranslatorBase* > RegisteredTranslators;
+	
+	//The manager will create only one pipeline per type
+	UPROPERTY()
+	TMap<const UClass*, UInterchangePipelineBase* > RegisteredPipelines;
+
+	//The manager will create only one factory per type
+	UPROPERTY()
+	TMap<const UClass*, UInterchangeFactoryBase* > RegisteredFactories;
+
+	//The manager will create only one writer per type
+	UPROPERTY()
+	TMap<const UClass*, UInterchangeWriterBase* > RegisteredWriters;
+
+	friend class UE::Interchange::FScopedTranslator;
+};
