Implement 3 reimport strategies:
-Do not change any asset property
-Set all the asset properties set by the pipelines
-Set all the asset properties set by the pipelines, except the properties change by the user in the editor

Make sure interchange scope of life is explicit

Support cancelling task in progress

diff --git a/Engine/Plugins/Runtime/Interchange/Source/Dispatcher/Private/InterchangeDispatcher.cpp b/Engine/Plugins/Runtime/Interchange/Source/Dispatcher/Private/InterchangeDispatcher.cpp
--- a/Engine/Plugins/Runtime/Interchange/Source/Dispatcher/Private/InterchangeDispatcher.cpp
+++ b/Engine/Plugins/Runtime/Interchange/Source/Dispatcher/Private/InterchangeDispatcher.cpp
@@ -103,6 +103,14 @@ namespace UE
 			SpawnHandler();
 		}
 
+		void FInterchangeDispatcher::StopProcess()
+		{
+			if (IsHandlerAlive())
+			{
+				WorkerHandler->Stop();
+			}
+		}
+
 		void FInterchangeDispatcher::TerminateProcess()
 		{
 			//Empty the cache folder
@@ -166,10 +174,7 @@ namespace UE
 
 		void FInterchangeDispatcher::CloseHandler()
 		{
-			if (IsHandlerAlive())
-			{
-				WorkerHandler->Stop();
-			}
+			StopProcess();
 			WorkerHandler.Reset();
 		}
 
diff --git a/Engine/Plugins/Runtime/Interchange/Source/Dispatcher/Public/InterchangeDispatcher.h b/Engine/Plugins/Runtime/Interchange/Source/Dispatcher/Public/InterchangeDispatcher.h
--- a/Engine/Plugins/Runtime/Interchange/Source/Dispatcher/Public/InterchangeDispatcher.h
+++ b/Engine/Plugins/Runtime/Interchange/Source/Dispatcher/Public/InterchangeDispatcher.h
@@ -25,6 +25,7 @@ namespace UE
 			void GetTaskState(int32 TaskIndex, ETaskState& TaskState, FString& JsonResult, TArray<FString>& JSonMessages);
 
 			void StartProcess();
+			void StopProcess();
 			void TerminateProcess();
 			void WaitAllTaskToCompleteExecution();
 			bool IsOver();
diff --git a/Engine/Plugins/Runtime/Interchange/Source/Export/Private/InterchangeExportPlugin.cpp b/Engine/Plugins/Runtime/Interchange/Source/Export/Private/InterchangeExportPlugin.cpp
--- a/Engine/Plugins/Runtime/Interchange/Source/Export/Private/InterchangeExportPlugin.cpp
+++ b/Engine/Plugins/Runtime/Interchange/Source/Export/Private/InterchangeExportPlugin.cpp
@@ -19,10 +19,21 @@ IMPLEMENT_MODULE( FInterchangeExportPlugin, InterchangeExportPlugin)
 
 void FInterchangeExportPlugin::StartupModule()
 {
-	UInterchangeManager& InterchangeManager = UInterchangeManager::GetInterchangeManager();
-
-	//Register the Writers
-	InterchangeManager.RegisterWriter(UInterchangeTextureWriter::StaticClass());
+	auto RegisterItems = []()
+	{
+		UInterchangeManager& InterchangeManager = UInterchangeManager::GetInterchangeManager();
+		//Register the Writers
+		InterchangeManager.RegisterWriter(UInterchangeTextureWriter::StaticClass());
+	};
+
+	if (GEngine)
+	{
+		RegisterItems();
+	}
+	else
+	{
+		FCoreDelegates::OnPostEngineInit.AddLambda(RegisterItems);
+	}
 }
 
 
diff --git a/Engine/Plugins/Runtime/Interchange/Source/Import/Private/Fbx/InterchangeFbxTranslator.cpp b/Engine/Plugins/Runtime/Interchange/Source/Import/Private/Fbx/InterchangeFbxTranslator.cpp
--- a/Engine/Plugins/Runtime/Interchange/Source/Import/Private/Fbx/InterchangeFbxTranslator.cpp
+++ b/Engine/Plugins/Runtime/Interchange/Source/Import/Private/Fbx/InterchangeFbxTranslator.cpp
@@ -109,6 +109,14 @@ bool UInterchangeFbxTranslator::Translate(const UInterchangeSourceData* SourceDa
 	return true;
 }
 
+void UInterchangeFbxTranslator::ReleaseSource()
+{
+	if (Dispatcher.IsValid())
+	{
+		Dispatcher->StopProcess();
+	}
+}
+
 void UInterchangeFbxTranslator::ImportFinish()
 {
 	if (Dispatcher.IsValid())
diff --git a/Engine/Plugins/Runtime/Interchange/Source/Import/Private/InterchangeImportCommon.cpp b/Engine/Plugins/Runtime/Interchange/Source/Import/Private/InterchangeImportCommon.cpp
new file mode 100644
--- /dev/null
+++ b/Engine/Plugins/Runtime/Interchange/Source/Import/Private/InterchangeImportCommon.cpp
@@ -0,0 +1,138 @@
+// Copyright Epic Games, Inc. All Rights Reserved.
+#include "InterchangeImportCommon.h"
+
+#include "CoreMinimal.h"
+#include "EditorFramework/AssetImportData.h"
+#include "InterchangeAssetImportData.h"
+#include "InterchangeSourceData.h"
+#include "Nodes/InterchangeBaseNodeContainer.h"
+#include "UObject/Object.h"
+
+namespace UE
+{
+	namespace Interchange
+	{
+		FFactoryCommon::FUpdateImportAssetDataParameters::FUpdateImportAssetDataParameters(UObject* InAssetImportDataOuter
+																							, UAssetImportData** InAssetImportData
+																							, const UInterchangeSourceData* InSourceData
+																							, FString InNodeUniqueID
+																							, UInterchangeBaseNodeContainer* InNodeContainer)
+			: AssetImportDataOuter(InAssetImportDataOuter)
+			, AssetImportData(InAssetImportData)
+			, SourceData(InSourceData)
+			, NodeUniqueID(InNodeUniqueID)
+			, NodeContainer(InNodeContainer)
+		{
+			ensure(AssetImportDataOuter);
+			ensure(AssetImportData);
+			ensure(SourceData);
+			ensure(!NodeUniqueID.IsEmpty());
+			ensure(NodeContainer);
+		}
+
+		void FFactoryCommon::UpdateImportAssetData(FUpdateImportAssetDataParameters& Parameters)
+		{
+#if WITH_EDITORONLY_DATA
+			if (!ensure(IsInGameThread()))
+			{
+				return;
+			}
+			if (!ensure(Parameters.SourceData && Parameters.AssetImportDataOuter && Parameters.AssetImportData))
+			{
+				return;
+			}
+			//Set the asset import data file source to allow reimport. TODO: manage MD5 Hash properly
+			TOptional<FMD5Hash> FileContentHash = Parameters.SourceData->GetFileContentHash();
+
+			UInterchangeAssetImportData* AssetImportData = nullptr;
+				
+			if (*(Parameters.AssetImportData) != nullptr)
+			{
+				UAssetImportData* ParametersAssetImportData = *(Parameters.AssetImportData);
+				AssetImportData = Cast<UInterchangeAssetImportData>(ParametersAssetImportData);
+				if (!AssetImportData)
+				{
+					AssetImportData = NewObject<UInterchangeAssetImportData>(Parameters.AssetImportDataOuter, NAME_None);
+					TArray<FString> Filenames;
+					ParametersAssetImportData->ExtractFilenames(Filenames);
+					TArray<FString> Labels;
+					ParametersAssetImportData->ExtractDisplayLabels(Labels);
+					if (Filenames.Num() > 1)
+					{
+						for (int32 FileIndex = 0; FileIndex < Filenames.Num(); ++FileIndex)
+						{
+							FString Filename = Filenames[FileIndex];
+							FString Label;
+							if (Labels.IsValidIndex(FileIndex))
+							{
+								Label = Labels[FileIndex];
+							}
+							//This is slow since it will hash the file, TODO make sure hashing is done in the create asset task
+							AssetImportData->AddFileName(Filename, FileIndex, Label);
+						}
+					}
+				}
+			}
+			else
+			{
+				AssetImportData = NewObject<UInterchangeAssetImportData>(Parameters.AssetImportDataOuter, NAME_None);
+			}
+
+			ensure(AssetImportData);
+			//Update the first filename, TODO for asset using multiple source file we have to update the correct index
+			AssetImportData->Update(Parameters.SourceData->GetFilename(), FileContentHash.IsSet() ? &FileContentHash.GetValue() : nullptr);
+			//Set the interchange node graph data
+			AssetImportData->NodeUniqueID = Parameters.NodeUniqueID;
+			FObjectDuplicationParameters DupParam(Parameters.NodeContainer, AssetImportData);
+			AssetImportData->NodeContainer = CastChecked<UInterchangeBaseNodeContainer>(StaticDuplicateObjectEx(DupParam));
+			//Assign the asset import data to the imported asset.
+			*(Parameters.AssetImportData) = AssetImportData;
+#endif //#if WITH_EDITORONLY_DATA
+		}
+
+		void FFactoryCommon::ApplyReimportStrategyToAsset(const EReimportStrategyFlags ReimportStrategyFlags
+										  , UObject* Asset
+										  , UInterchangeBaseNode* PreviousAssetNode
+										  , UInterchangeBaseNode* CurrentAssetNode
+										  , UInterchangeBaseNode* PipelineAssetNode)
+		{
+			if (!ensure(PreviousAssetNode) || !ensure(PipelineAssetNode) || !ensure(CurrentAssetNode))
+			{
+				return;
+			}
+
+			switch (ReimportStrategyFlags)
+			{
+				case EReimportStrategyFlags::ApplyNoProperties:
+				{
+					//We want to have no effect, we want to keep the original pipeline node.
+					//So we copy the previous asset node into the pipeline mode, the pipeline node will be saved
+					//in the import asset data, and it will save the original import node.
+					UInterchangeBaseNode::CopyStorage(PreviousAssetNode, PipelineAssetNode);
+					break;
+				}
+					
+				case EReimportStrategyFlags::ApplyPipelineProperties:
+				{
+					//Directly apply pipeline node attribute to the asset
+					PipelineAssetNode->ApplyAllCustomAttributeToAsset(Asset);
+					break;
+				}
+				
+				case EReimportStrategyFlags::ApplyEditorChangedProperties:
+				{
+					TArray<FAttributeKey> RemovedAttributes;
+					TArray<FAttributeKey> AddedAttributes;
+					TArray<FAttributeKey> ModifiedAttributes;
+					UInterchangeBaseNode::CompareNodeStorage(PreviousAssetNode, CurrentAssetNode, RemovedAttributes, AddedAttributes, ModifiedAttributes);
+
+					//set all ModifedAttributes from the CurrentAssetNode to the pipeline node. This will put back all user changes
+					UInterchangeBaseNode::CopyStorageAttributes(CurrentAssetNode, PipelineAssetNode, ModifiedAttributes);
+					//Now apply the pipeline node attribute to the asset
+					PipelineAssetNode->ApplyAllCustomAttributeToAsset(Asset);
+					break;
+				}
+			}
+		}
+	}
+}
diff --git a/Engine/Plugins/Runtime/Interchange/Source/Import/Private/InterchangeImportCommon.h b/Engine/Plugins/Runtime/Interchange/Source/Import/Private/InterchangeImportCommon.h
new file mode 100644
--- /dev/null
+++ b/Engine/Plugins/Runtime/Interchange/Source/Import/Private/InterchangeImportCommon.h
@@ -0,0 +1,55 @@
+// Copyright Epic Games, Inc. All Rights Reserved.
+
+#include "CoreMinimal.h"
+#include "InterchangeFactoryBase.h"
+
+class UAssetImportData;
+class UInterchangeBaseNode;
+class UInterchangeBaseNodeContainer;
+class UInterchangeSourceData;
+class UObject;
+
+namespace UE
+{
+	namespace Interchange
+	{
+		/**
+		 * All the code we cannot put in the base factory class because of dependencies (like Engine dep)
+		 * Will be available here.
+		 */
+		class FFactoryCommon
+		{
+		public:
+
+			struct FUpdateImportAssetDataParameters
+			{
+				UObject* AssetImportDataOuter = nullptr;
+				UAssetImportData** AssetImportData = nullptr;
+				const UInterchangeSourceData* SourceData = nullptr;
+				FString NodeUniqueID;
+				UInterchangeBaseNodeContainer* NodeContainer;
+
+				FUpdateImportAssetDataParameters(UObject* InAssetImportDataOuter
+																	, UAssetImportData** InAssetImportData
+																	, const UInterchangeSourceData* InSourceData
+																	, FString InNodeUniqueID
+																	, UInterchangeBaseNodeContainer* InNodeContainer);
+			};
+
+			/**
+			 * Update the AssetImportData source file of the specified asset in the parameters. Also update the node container and the node unique id.
+			 * If the AssetImportData is null it will create one. If the AssetImportData is not an UInterchangeAssetImportData it will create a new one.
+			 */
+			static void UpdateImportAssetData(FUpdateImportAssetDataParameters& Parameters);
+
+			/**
+			 * Apply the current strategy to the PipelineAssetNode
+			 */
+			static void ApplyReimportStrategyToAsset(const EReimportStrategyFlags ReimportStrategyFlags
+											  , UObject* Asset
+											  , UInterchangeBaseNode* PreviousAssetNode
+											  , UInterchangeBaseNode* CurrentAssetNode
+											  , UInterchangeBaseNode* PipelineAssetNode);
+		};
+	}
+}
diff --git a/Engine/Plugins/Runtime/Interchange/Source/Import/Private/InterchangeImportPlugin.cpp b/Engine/Plugins/Runtime/Interchange/Source/Import/Private/InterchangeImportPlugin.cpp
--- a/Engine/Plugins/Runtime/Interchange/Source/Import/Private/InterchangeImportPlugin.cpp
+++ b/Engine/Plugins/Runtime/Interchange/Source/Import/Private/InterchangeImportPlugin.cpp
@@ -31,24 +31,36 @@ IMPLEMENT_MODULE( FInterchangeImportPlugin, InterchangeImportPlugin)
 
 void FInterchangeImportPlugin::StartupModule()
 {
-	UInterchangeManager& InterchangeManager = UInterchangeManager::GetInterchangeManager();
-
-	//Register the translators
-	//Scenes
-	//InterchangeManager.RegisterTranslator(UInterchangeFbxTranslator::StaticClass()); //Do not submit uncommented until we replace completly fbx importer (staticmesh + skeletalMesh + animation)
-	//Textures
-	InterchangeManager.RegisterTranslator(UInterchangeBMPTranslator::StaticClass());
-	InterchangeManager.RegisterTranslator(UInterchangeDDSTranslator::StaticClass());
-	InterchangeManager.RegisterTranslator(UInterchangeEXRTranslator::StaticClass());
-	InterchangeManager.RegisterTranslator(UInterchangeJPGTranslator::StaticClass());
-	InterchangeManager.RegisterTranslator(UInterchangePCXTranslator::StaticClass());
-	InterchangeManager.RegisterTranslator(UInterchangePNGTranslator::StaticClass());
-	InterchangeManager.RegisterTranslator(UInterchangePSDTranslator::StaticClass());
-	InterchangeManager.RegisterTranslator(UInterchangeTGATranslator::StaticClass());
-
-	//Register the factories
-	InterchangeManager.RegisterFactory(UInterchangeTextureFactory::StaticClass());
-	InterchangeManager.RegisterFactory(UInterchangeMaterialFactory::StaticClass());
+	auto RegisterItems = []()
+	{
+		UInterchangeManager& InterchangeManager = UInterchangeManager::GetInterchangeManager();
+
+		//Register the translators
+		//Scenes
+		InterchangeManager.RegisterTranslator(UInterchangeFbxTranslator::StaticClass()); //Do not submit uncommented until we replace completly fbx importer (staticmesh + skeletalMesh + animation)
+		//Textures
+		InterchangeManager.RegisterTranslator(UInterchangeBMPTranslator::StaticClass());
+		InterchangeManager.RegisterTranslator(UInterchangeDDSTranslator::StaticClass());
+		InterchangeManager.RegisterTranslator(UInterchangeEXRTranslator::StaticClass());
+		InterchangeManager.RegisterTranslator(UInterchangeJPGTranslator::StaticClass());
+		InterchangeManager.RegisterTranslator(UInterchangePCXTranslator::StaticClass());
+		InterchangeManager.RegisterTranslator(UInterchangePNGTranslator::StaticClass());
+		InterchangeManager.RegisterTranslator(UInterchangePSDTranslator::StaticClass());
+		InterchangeManager.RegisterTranslator(UInterchangeTGATranslator::StaticClass());
+
+		//Register the factories
+		InterchangeManager.RegisterFactory(UInterchangeTextureFactory::StaticClass());
+		InterchangeManager.RegisterFactory(UInterchangeMaterialFactory::StaticClass());
+	};
+
+	if (GEngine)
+	{
+		RegisterItems();
+	}
+	else
+	{
+		FCoreDelegates::OnPostEngineInit.AddLambda(RegisterItems);
+	}
 }
 
 
diff --git a/Engine/Plugins/Runtime/Interchange/Source/Import/Private/Material/InterchangeMaterialFactory.cpp b/Engine/Plugins/Runtime/Interchange/Source/Import/Private/Material/InterchangeMaterialFactory.cpp
--- a/Engine/Plugins/Runtime/Interchange/Source/Import/Private/Material/InterchangeMaterialFactory.cpp
+++ b/Engine/Plugins/Runtime/Interchange/Source/Import/Private/Material/InterchangeMaterialFactory.cpp
@@ -1,7 +1,9 @@
 // Copyright Epic Games, Inc. All Rights Reserved. 
 #include "Material/InterchangeMaterialFactory.h"
 
+#include "InterchangeImportCommon.h"
 #include "InterchangeMaterialNode.h"
+#include "InterchangeSourceData.h"
 #include "LogInterchangeImportPlugin.h"
 #include "Materials/Material.h"
 #include "Materials/MaterialExpressionScalarParameter.h"
@@ -13,6 +15,7 @@
 #include "Nodes/InterchangeBaseNode.h"
 #include "Nodes/InterchangeBaseNodeContainer.h"
 
+
 #if WITH_EDITORONLY_DATA
 
 #include "EditorFramework/AssetImportData.h"
@@ -118,7 +121,8 @@ UObject* UInterchangeMaterialFactory::CreateAsset(const UInterchangeMaterialFact
 
 	if (MaterialObject)
 	{
-		//The re-import change only the source data and will not apply any custom attribute to the asset.
+		//Currently material re-import will not touch the material at all
+		//TODO design a re-import process for the material (expressions and input connections)
 		if(!Arguments.ReimportObject)
 		{
 			UMaterial* Material = Cast<UMaterial>(MaterialObject);
@@ -281,18 +285,17 @@ void UInterchangeMaterialFactory::PostImportGameThreadCallback(const FPostImport
 {
 	check(IsInGameThread());
 	Super::PostImportGameThreadCallback(Arguments);
-#if WITH_EDITORONLY_DATA
-	if (ensure(Arguments.ReimportObject && Arguments.SourceData))
+	if (ensure(Arguments.ImportedObject && Arguments.SourceData))
 	{
 		//We must call the Update of the asset source file in the main thread because UAssetImportData::Update execute some delegate we do not control
-		UMaterialInterface* ImportedMaterial = CastChecked<UMaterialInterface>(Arguments.ReimportObject);
-		//Set the asset import data file source to allow reimport. TODO: manage MD5 Hash properly
-		TOptional<FMD5Hash> FileContentHash = Arguments.SourceData->GetFileContentHash();
-		if (ImportedMaterial->AssetImportData)
-		{
-			ImportedMaterial->AssetImportData->Update(Arguments.SourceData->GetFilename(), FileContentHash.IsSet() ? &FileContentHash.GetValue() : nullptr);
-		}
+		UMaterialInterface* ImportedMaterial = CastChecked<UMaterialInterface>(Arguments.ImportedObject);
+
+		UE::Interchange::FFactoryCommon::FUpdateImportAssetDataParameters UpdateImportAssetDataParameters(ImportedMaterial
+																										  , &ImportedMaterial->AssetImportData
+																										  , Arguments.SourceData
+																										  , Arguments.NodeUniqueID
+																										  , Arguments.NodeContainer);
+		UE::Interchange::FFactoryCommon::UpdateImportAssetData(UpdateImportAssetDataParameters);
 	}
-#endif //#if WITH_EDITORONLY_DATA
 	return;
 }
\ No newline at end of file
diff --git a/Engine/Plugins/Runtime/Interchange/Source/Import/Private/Texture/InterchangeTextureFactory.cpp b/Engine/Plugins/Runtime/Interchange/Source/Import/Private/Texture/InterchangeTextureFactory.cpp
--- a/Engine/Plugins/Runtime/Interchange/Source/Import/Private/Texture/InterchangeTextureFactory.cpp
+++ b/Engine/Plugins/Runtime/Interchange/Source/Import/Private/Texture/InterchangeTextureFactory.cpp
@@ -3,7 +3,10 @@
 
 #include "Engine/Texture.h"
 #include "Engine/Texture2D.h"
+#include "InterchangeAssetImportData.h"
+#include "InterchangeImportCommon.h"
 #include "InterchangeTextureNode.h"
+#include "InterchangeTranslatorBase.h"
 #include "LogInterchangeImportPlugin.h"
 #include "Nodes/InterchangeBaseNode.h"
 #include "Texture/InterchangeTexturePayloadInterface.h"
@@ -89,7 +92,7 @@ UObject* UInterchangeTextureFactory::CreateAsset(const UInterchangeTextureFactor
 		return nullptr;
 	}
 
-	const UInterchangeTextureNode* TextureNode = Cast<UInterchangeTextureNode>(Arguments.AssetNode);
+	UInterchangeTextureNode* TextureNode = Cast<UInterchangeTextureNode>(Arguments.AssetNode);
 	if (TextureNode == nullptr)
 	{
 		return nullptr;
@@ -164,12 +167,25 @@ UObject* UInterchangeTextureFactory::CreateAsset(const UInterchangeTextureFactor
 			Texture2D->MipGenSettings = Image.MipGenSettings.GetValue();
 		}
 
-		//The re-import change only the source data and will not apply any custom attribute to the asset.
 		if(!Arguments.ReimportObject)
 		{
 			/** Apply all TextureNode custom attributes to the texture asset */
 			TextureNode->ApplyAllCustomAttributeToAsset(Texture2D);
 		}
+		else
+		{
+			UInterchangeAssetImportData* InterchangeAssetImportData = Cast<UInterchangeAssetImportData>(Texture2D->AssetImportData);
+			UInterchangeBaseNode* PreviousNode = nullptr;
+			if (InterchangeAssetImportData)
+			{
+				PreviousNode = InterchangeAssetImportData->NodeContainer->GetNode(FName(*InterchangeAssetImportData->NodeUniqueID));
+			}
+			UInterchangeTextureNode* CurrentNode = NewObject<UInterchangeTextureNode>();
+			UInterchangeBaseNode::CopyStorage(TextureNode, CurrentNode);
+			CurrentNode->FillAllCustomAttributeFromAsset(Texture2D);
+			//Apply reimport strategy
+			UE::Interchange::FFactoryCommon::ApplyReimportStrategyToAsset(Arguments.ReimportStrategyFlags, Texture2D, PreviousNode, CurrentNode, TextureNode);
+		}
 		
 		//Getting the file Hash will cache it into the source data
 		Arguments.SourceData->GetFileContentHash();
@@ -191,15 +207,17 @@ void UInterchangeTextureFactory::PostImportGameThreadCallback(const FPostImportG
 {
 	check(IsInGameThread());
 	Super::PostImportGameThreadCallback(Arguments);
-#if WITH_EDITORONLY_DATA
-	if(ensure(Arguments.ReimportObject && Arguments.SourceData))
+	if (ensure(Arguments.ImportedObject && Arguments.SourceData))
 	{
 		//We must call the Update of the asset source file in the main thread because UAssetImportData::Update execute some delegate we do not control
-		UTexture* ImportedTexture = CastChecked<UTexture>(Arguments.ReimportObject);
-		//Set the asset import data file source to allow reimport. TODO: manage MD5 Hash properly
-		TOptional<FMD5Hash> FileContentHash = Arguments.SourceData->GetFileContentHash();
-		ImportedTexture->AssetImportData->Update(Arguments.SourceData->GetFilename(), FileContentHash.IsSet() ? &FileContentHash.GetValue() : nullptr);
+		UTexture* ImportedTexture = CastChecked<UTexture>(Arguments.ImportedObject);
+		
+		UE::Interchange::FFactoryCommon::FUpdateImportAssetDataParameters UpdateImportAssetDataParameters(ImportedTexture
+																						 , &ImportedTexture->AssetImportData
+																						 , Arguments.SourceData
+																						 , Arguments.NodeUniqueID
+																						 , Arguments.NodeContainer);
+		UE::Interchange::FFactoryCommon::UpdateImportAssetData(UpdateImportAssetDataParameters);
 	}
-#endif //#if WITH_EDITORONLY_DATA
 	return;
 }
\ No newline at end of file
diff --git a/Engine/Plugins/Runtime/Interchange/Source/Import/Public/Fbx/InterchangeFbxTranslator.h b/Engine/Plugins/Runtime/Interchange/Source/Import/Public/Fbx/InterchangeFbxTranslator.h
--- a/Engine/Plugins/Runtime/Interchange/Source/Import/Public/Fbx/InterchangeFbxTranslator.h
+++ b/Engine/Plugins/Runtime/Interchange/Source/Import/Public/Fbx/InterchangeFbxTranslator.h
@@ -32,6 +32,7 @@ public:
 	/** Begin UInterchangeTranslatorBase API*/
 	virtual bool CanImportSourceData(const UInterchangeSourceData* InSourceData) const override;
 	virtual bool Translate(const UInterchangeSourceData* SourceData, UInterchangeBaseNodeContainer& BaseNodeContainer) const override;
+	virtual void ReleaseSource() override;
 	virtual void ImportFinish() override;
 	/** End UInterchangeTranslatorBase API*/
 
diff --git a/Engine/Plugins/Runtime/Interchange/Source/Nodes/Public/InterchangeMaterialNode.h b/Engine/Plugins/Runtime/Interchange/Source/Nodes/Public/InterchangeMaterialNode.h
--- a/Engine/Plugins/Runtime/Interchange/Source/Nodes/Public/InterchangeMaterialNode.h
+++ b/Engine/Plugins/Runtime/Interchange/Source/Nodes/Public/InterchangeMaterialNode.h
@@ -370,6 +370,9 @@ private:
 	const UE::Interchange::FAttributeKey Macro_CustomBlendModeKey = UE::Interchange::FAttributeKey(TEXT("BlendMode"));
 
 #if WITH_ENGINE
+
+	//Material Ussage is not handle properly, you can have several usage check
+	//Currently I can set only one and retrieve the first, TODO make this work (maybe one key per usage...)
 	bool ApplyCustomMaterialUsageToAsset(UObject * Asset) const
 	{
 		if (!Asset)
@@ -390,6 +393,35 @@ private:
 		return false;
 	}
 
+	bool FillCustomMaterialUsageFromAsset(UObject* Asset)
+	{
+		if (!Asset)
+		{
+			return false;
+		}
+		UMaterialInterface* TypedObject = Cast<UMaterialInterface>(Asset);
+		if (!TypedObject)
+		{
+			return false;
+		}
+		const UMaterial* Material = TypedObject->GetMaterial_Concurrent();
+		if (!Material)
+		{
+			return false;
+		}
+		for (int32 Usage = 0; Usage < MATUSAGE_MAX; Usage++)
+		{
+			const EMaterialUsage UsageEnum = (EMaterialUsage)Usage;
+			if (Material->GetUsageByFlag(UsageEnum))
+			{
+				if (SetCustomMaterialUsage((int8)UsageEnum, false))
+				{
+					return true;
+				}
+			}
+		}
+		return false;
+	}
 #endif
 
 	IMPLEMENT_NODE_ATTRIBUTE_APPLY_UOBJECT(BlendMode, uint8, UMaterial, TEnumAsByte<enum EBlendMode>);
diff --git a/Engine/Plugins/Runtime/Interchange/Source/Nodes/Public/InterchangeTextureNode.h b/Engine/Plugins/Runtime/Interchange/Source/Nodes/Public/InterchangeTextureNode.h
--- a/Engine/Plugins/Runtime/Interchange/Source/Nodes/Public/InterchangeTextureNode.h
+++ b/Engine/Plugins/Runtime/Interchange/Source/Nodes/Public/InterchangeTextureNode.h
@@ -884,8 +884,49 @@ private:
 #endif //#if WITH_EDITORONLY_DATA
 	IMPLEMENT_NODE_ATTRIBUTE_APPLY_UOBJECT(LODBias, int32, UTexture, );
 	IMPLEMENT_NODE_ATTRIBUTE_APPLY_UOBJECT(LODGroup, uint8, UTexture, TEnumAsByte<enum TextureGroup>);
+	
 	//TODO support per platform data in the FAttributeStorage, so we can set different value per platform at the pipeline stage, We set only the default value for now
-	IMPLEMENT_NODE_ATTRIBUTE_APPLY_UOBJECT(Downscale, float, UTexture, );
+	//IMPLEMENT_NODE_ATTRIBUTE_APPLY_UOBJECT(Downscale, float, UTexture, );
+#if WITH_ENGINE
+	bool ApplyCustomDownscaleToAsset(UObject * Asset) const
+	{
+		if (!Asset)
+		{
+			return false;
+		}
+		UTexture* TypedObject = Cast<UTexture>(Asset);
+		if (!TypedObject)
+		{
+			return false;
+		}
+		float ValueData;
+		if (GetCustomDownscale(ValueData))
+		{
+			TypedObject->Downscale.Default = ValueData;
+			return true;
+		}
+		return false;
+	}
+
+	bool FillCustomDownscaleFromAsset(UObject* Asset)
+	{
+		if (!Asset)
+		{
+			return false;
+		}
+		UTexture* TypedObject = Cast<UTexture>(Asset);
+		if (!TypedObject)
+		{
+			return false;
+		}
+		if (SetCustomDownscale(TypedObject->Downscale.Default, false))
+		{
+			return true;
+		}
+		return false;
+	}
+#endif //WITH_ENGINE
+
 	IMPLEMENT_NODE_ATTRIBUTE_APPLY_UOBJECT(DownscaleOptions, uint8, UTexture, ETextureDownscaleOptions);
 #if WITH_EDITORONLY_DATA
 	//Compositing
diff --git a/Engine/Source/Editor/MainFrame/MainFrame.Build.cs b/Engine/Source/Editor/MainFrame/MainFrame.Build.cs
--- a/Engine/Source/Editor/MainFrame/MainFrame.Build.cs
+++ b/Engine/Source/Editor/MainFrame/MainFrame.Build.cs
@@ -44,6 +44,8 @@ public class MainFrame : ModuleRules
 				"ToolMenus",
 				"LauncherServices",
 				"DerivedDataCache",
+				"InterchangeCore",
+				"InterchangeEngine"
 			}
 		);
 
diff --git a/Engine/Source/Editor/MainFrame/Private/Frame/MainFrameHandler.h b/Engine/Source/Editor/MainFrame/Private/Frame/MainFrameHandler.h
--- a/Engine/Source/Editor/MainFrame/Private/Frame/MainFrameHandler.h
+++ b/Engine/Source/Editor/MainFrame/Private/Frame/MainFrameHandler.h
@@ -11,6 +11,7 @@
 #include "Framework/Docking/TabManager.h"
 #include "Framework/Docking/LayoutService.h"
 #include "EngineGlobals.h"
+#include "InterchangeManager.h"
 #include "Interfaces/IMainFrameModule.h"
 #include "Editor/UnrealEdEngine.h"
 #include "EditorModeManager.h"
@@ -99,6 +100,8 @@ public:
 		{
 			// NORMAL EXIT PATH
 
+			if (UInterchangeManager::GetInterchangeManager().WarnIfInterchangeIsActive()) { return false; }
+
 			// Unattented mode can always exit.
 			if (FApp::IsUnattended())
 			{
diff --git a/Engine/Source/Editor/UnrealEd/Private/Subsystems/ImportSubsystem.cpp b/Engine/Source/Editor/UnrealEd/Private/Subsystems/ImportSubsystem.cpp
--- a/Engine/Source/Editor/UnrealEd/Private/Subsystems/ImportSubsystem.cpp
+++ b/Engine/Source/Editor/UnrealEd/Private/Subsystems/ImportSubsystem.cpp
@@ -112,17 +112,36 @@ UImportSubsystem::UImportSubsystem()
 
 void UImportSubsystem::Initialize(FSubsystemCollectionBase& Collection)
 {
-	//Hook on Interchange manager to know when to broadcast import/reimport
-	UInterchangeManager& InterchangeManager = UInterchangeManager::GetInterchangeManager();
-	InterchangeManager.OnAssetPostImport.AddUObject(this, &UImportSubsystem::InterchangeBroadcastAssetPostImport);
-	InterchangeManager.OnAssetPostReimport.AddUObject(this, &UImportSubsystem::InterchangeBroadcastAssetPostReimport);
+	auto AttachDelegates = [this]()
+	{
+		//Hook on Interchange manager to know when to broadcast import/reimport
+		UInterchangeManager& InterchangeManager = UInterchangeManager::GetInterchangeManager();
+		InterchangeManager.OnAssetPostImport.AddUObject(this, &UImportSubsystem::InterchangeBroadcastAssetPostImport);
+		InterchangeManager.OnAssetPostReimport.AddUObject(this, &UImportSubsystem::InterchangeBroadcastAssetPostReimport);
+	};
+
+	if (GEngine)
+	{
+		AttachDelegates();
+	}
+	else
+	{
+		FCoreDelegates::OnPostEngineInit.AddLambda(AttachDelegates);
+	}
+
+	//Since subsystem deinitialize after the FCoreDelegates::OnExit and interchange manager is destroy during FCoreDelegates::OnExit
+	//We want to unregister delegate using the FCoreDelegates::OnPreExit.
+	FCoreDelegates::OnPreExit.AddLambda([this]()
+	{
+		//Unhook interchange manager import/reimport delegates
+		UInterchangeManager& InterchangeManager = UInterchangeManager::GetInterchangeManager();
+		InterchangeManager.OnAssetPostImport.RemoveAll(this);
+		InterchangeManager.OnAssetPostReimport.RemoveAll(this);
+	});
 }
 
 void UImportSubsystem::Deinitialize()
 {
-	UInterchangeManager& InterchangeManager = UInterchangeManager::GetInterchangeManager();
-	InterchangeManager.OnAssetPostImport.RemoveAll(this);
-	InterchangeManager.OnAssetPostReimport.RemoveAll(this);
 }
 
 void UImportSubsystem::ImportNextTick(const TArray<FString>& Files, const FString& DestinationPath)
diff --git a/Engine/Source/Runtime/Interchange/Core/Private/Nodes/InterchangeBaseNode.cpp b/Engine/Source/Runtime/Interchange/Core/Private/Nodes/InterchangeBaseNode.cpp
--- a/Engine/Source/Runtime/Interchange/Core/Private/Nodes/InterchangeBaseNode.cpp
+++ b/Engine/Source/Runtime/Interchange/Core/Private/Nodes/InterchangeBaseNode.cpp
@@ -192,6 +192,24 @@ void UInterchangeBaseNode::ApplyAllCustomAttributeToAsset(UObject* Object) const
 	}
 }
 
+void UInterchangeBaseNode::FillAllCustomAttributeFromAsset(UObject* Object) const
+{
+	UClass* ObjectClass = Object->GetClass();
+	for (const TPair<UClass*, TArray<UE::Interchange::FFillAttributeToAsset>>& ClassDelegatePair : FillCustomAttributeDelegates)
+	{
+		if (ObjectClass->IsChildOf(ClassDelegatePair.Key))
+		{
+			for (const UE::Interchange::FFillAttributeToAsset& Delegate : ClassDelegatePair.Value)
+			{
+				if (Delegate.IsBound())
+				{
+					Delegate.Execute(Object);
+				}
+			}
+		}
+	}
+}
+
 void UInterchangeBaseNode::Serialize(FArchive& Ar)
 {
 	Ar << Attributes;
diff --git a/Engine/Source/Runtime/Interchange/Core/Private/Nodes/InterchangeBaseNodeContainer.cpp b/Engine/Source/Runtime/Interchange/Core/Private/Nodes/InterchangeBaseNodeContainer.cpp
--- a/Engine/Source/Runtime/Interchange/Core/Private/Nodes/InterchangeBaseNodeContainer.cpp
+++ b/Engine/Source/Runtime/Interchange/Core/Private/Nodes/InterchangeBaseNodeContainer.cpp
@@ -44,7 +44,7 @@ bool UInterchangeBaseNodeContainer::IsNodeUIDValid(FName NodeUniqueID) const
 	return Nodes.Contains(NodeUniqueID);
 }
 
-void UInterchangeBaseNodeContainer::IterateNodes(TFunctionRef<void(FName, const UInterchangeBaseNode*)> IterationLambda)
+void UInterchangeBaseNodeContainer::IterateNodes(TFunctionRef<void(FName, UInterchangeBaseNode*)> IterationLambda)
 {
 	for (TPair<FName, UInterchangeBaseNode*>& NodeKeyValue : Nodes)
 	{
diff --git a/Engine/Source/Runtime/Interchange/Core/Private/Types/AttributeStorage.cpp b/Engine/Source/Runtime/Interchange/Core/Private/Types/AttributeStorage.cpp
--- a/Engine/Source/Runtime/Interchange/Core/Private/Types/AttributeStorage.cpp
+++ b/Engine/Source/Runtime/Interchange/Core/Private/Types/AttributeStorage.cpp
@@ -113,7 +113,7 @@ namespace UE
 			return EAttributeStorageResult::Operation_Success;
 		}
 
-		EAttributeTypes FAttributeStorage::GetAttributeType(const FAttributeKey& ElementAttributeKey)
+		EAttributeTypes FAttributeStorage::GetAttributeType(const FAttributeKey& ElementAttributeKey) const
 		{
 			//Lock the storage
 			FScopeLock ScopeLock(&StorageMutex);
@@ -226,6 +226,65 @@ namespace UE
 			}
 		}
 
+		void FAttributeStorage::CopyStorageAttributes(const FAttributeStorage& SourceStorage, FAttributeStorage& DestinationStorage, const TArray<FAttributeKey>& AttributeKeys)
+		{
+			//Lock both storages
+			FScopeLock SourceScopeLock(&SourceStorage.StorageMutex);
+			FScopeLock DestinationScopeLock(&DestinationStorage.StorageMutex);
+
+			for (const FAttributeKey& AttributeKey : AttributeKeys)
+			{
+				if (const FAttributeAllocationInfo* SourceAttributeInfo = SourceStorage.AttributeAllocationTable.Find(AttributeKey))
+				{
+					FAttributeAllocationInfo* DestinationAttributeInfo = DestinationStorage.AttributeAllocationTable.Find(AttributeKey);
+					
+					if (!SourceStorage.AttributeStorage.IsValidIndex(SourceAttributeInfo->Offset))
+					{
+						//EAttributeStorageResult::Operation_Error_AttributeAllocationCorrupted;
+						continue;
+					}
+
+ 					const EAttributeTypes SourceValueType = SourceAttributeInfo->Type;
+ 					const uint64 SourceValueSize = SourceAttributeInfo->Size;
+					if (DestinationAttributeInfo)
+					{
+						if (DestinationAttributeInfo->Type != SourceValueType)
+						{
+							//EAttributeStorageResult::Operation_Error_WrongType;
+							continue;
+						}
+						if (DestinationAttributeInfo->Size != SourceValueSize)
+						{
+							//EAttributeStorageResult::Operation_Error_WrongSize;
+							continue;
+						}
+						if (!DestinationStorage.AttributeStorage.IsValidIndex(DestinationAttributeInfo->Offset))
+						{
+							//EAttributeStorageResult::Operation_Error_AttributeAllocationCorrupted;
+							continue;
+						}
+					}
+					else
+					{
+						DestinationAttributeInfo = &DestinationStorage.AttributeAllocationTable.Add(AttributeKey);
+						DestinationAttributeInfo->Type = SourceValueType;
+						DestinationAttributeInfo->Size = SourceValueSize;
+						DestinationAttributeInfo->Offset = DestinationStorage.AttributeStorage.AddZeroed(SourceValueSize);
+					}
+
+					//Force the specified attribute property
+					DestinationAttributeInfo->Property = SourceAttributeInfo->Property;
+					//Save the hash no need to compute it
+					DestinationAttributeInfo->Hash = SourceAttributeInfo->Hash;
+
+					//Memcpy source to destination storage
+					const uint8* SourceStorageData = SourceStorage.AttributeStorage.GetData();
+					uint8* DestinationStorageData = DestinationStorage.AttributeStorage.GetData();
+					FMemory::Memcpy(&DestinationStorageData[DestinationAttributeInfo->Offset], &SourceStorageData[SourceAttributeInfo->Offset], SourceValueSize);
+				}
+			}
+		}
+
 		void FAttributeStorage::SetDefragRatio(const float InDefragRatio)
 		{
 			FScopeLock ScopeLock(&StorageMutex);
diff --git a/Engine/Source/Runtime/Interchange/Core/Public/InterchangeFactoryBase.h b/Engine/Source/Runtime/Interchange/Core/Public/InterchangeFactoryBase.h
--- a/Engine/Source/Runtime/Interchange/Core/Public/InterchangeFactoryBase.h
+++ b/Engine/Source/Runtime/Interchange/Core/Public/InterchangeFactoryBase.h
@@ -3,14 +3,26 @@
 #pragma once
 
 #include "CoreMinimal.h"
-#include "Nodes/InterchangeBaseNode.h"
-#include "InterchangeTranslatorBase.h"
 #include "UObject/Class.h"
 #include "UObject/Object.h"
 #include "UObject/ObjectMacros.h"
 
 #include "InterchangeFactoryBase.generated.h"
 
+class UInterchangeBaseNode;
+class UInterchangeBaseNodeContainer;
+class UInterchangeSourceData;
+class UInterchangeTranslatorBase;
+
+
+UENUM()
+enum class EReimportStrategyFlags : uint8
+{
+	ApplyNoProperties, //Do not apply any property when re-importing, simply change the source data
+	ApplyPipelineProperties, //Always apply all pipeline specified properties
+	ApplyEditorChangedProperties //Always apply all pipeline properties, but leave the properties modified in editor since the last import
+};
+
 UCLASS(BlueprintType, Blueprintable, Abstract)
 class INTERCHANGECORE_API UInterchangeFactoryBase : public UObject
 {
@@ -38,7 +50,7 @@ public:
 		FString AssetName = FString();
 
 		/** The base node that describe how to create the asset */
-		const UInterchangeBaseNode* AssetNode = nullptr;
+		UInterchangeBaseNode* AssetNode = nullptr;
 
 		/** The translator is use to retrieve the PayLoad data in case the factory need it */
 		const UInterchangeTranslatorBase* Translator = nullptr;
@@ -54,6 +66,8 @@ public:
 		 * the asset we want to re-import. The re-import should just change the source data and not any asset settings.
 		 */
 		UObject* ReimportObject = nullptr;
+
+		EReimportStrategyFlags ReimportStrategyFlags;
 	};
 
 	/**
@@ -97,7 +111,10 @@ public:
 		const UInterchangeSourceData* SourceData = nullptr;
 
 		/** The UObject  we want to execute code on*/
-		UObject* ReimportObject = nullptr;
+		UObject* ImportedObject = nullptr;
+		FString NodeUniqueID;
+		UInterchangeBaseNodeContainer* NodeContainer = nullptr;
+;
 	};
 	/* This function is call in the completion task on the main thread, use it to call main thread post creation step for your assets*/
 	virtual void PostImportGameThreadCallback(const FPostImportGameThreadCallbackParams& Arguments) const
diff --git a/Engine/Source/Runtime/Interchange/Core/Public/InterchangePipelineBase.h b/Engine/Source/Runtime/Interchange/Core/Public/InterchangePipelineBase.h
--- a/Engine/Source/Runtime/Interchange/Core/Public/InterchangePipelineBase.h
+++ b/Engine/Source/Runtime/Interchange/Core/Public/InterchangePipelineBase.h
@@ -22,12 +22,12 @@ public:
 	 * the Interchange manager is calling this function not the virtual one that is call by the default implementation.
 	 */
 	UFUNCTION(BlueprintNativeEvent, BlueprintCallable, Category = "Interchange | Translator")
-	bool ScriptedExecuteImportPipeline(UInterchangeBaseNodeContainer* BaseNodeContainerAdapter);
+	bool ScriptedExecuteImportPipeline(UInterchangeBaseNodeContainer* BaseNodeContainer);
 	/** The default implementation (call if the blueprint do not have any implementation) will call the virtual ExecuteImportPipeline */
-	bool ScriptedExecuteImportPipeline_Implementation(UInterchangeBaseNodeContainer* BaseNodeContainerAdapter)
+	bool ScriptedExecuteImportPipeline_Implementation(UInterchangeBaseNodeContainer* BaseNodeContainer)
 	{
 		//By default we call the virtual import pipeline execution
-		return ExecuteImportPipeline(BaseNodeContainerAdapter);
+		return ExecuteImportPipeline(BaseNodeContainer);
 	}
 
 	/**
@@ -35,12 +35,12 @@ public:
 	 * the Interchange manager is calling this function not the virtual one that is call by the default implementation.
 	 */
 	UFUNCTION(BlueprintCallable, BlueprintNativeEvent, Category = "Interchange | Translator")
-	bool ScriptedExecuteExportPipeline(UInterchangeBaseNodeContainer* BaseNodeContainerAdapter);
+	bool ScriptedExecuteExportPipeline(UInterchangeBaseNodeContainer* BaseNodeContainer);
 
 	/** The default implementation (call if the blueprint do not have any implementation) will call the virtual ExecuteExportPipeline */
-	bool ScriptedExecuteExportPipeline_Implementation(UInterchangeBaseNodeContainer* BaseNodeContainerAdapter)
+	bool ScriptedExecuteExportPipeline_Implementation(UInterchangeBaseNodeContainer* BaseNodeContainer)
 	{
-		return ExecuteExportPipeline(BaseNodeContainerAdapter);
+		return ExecuteExportPipeline(BaseNodeContainer);
 	}
 
 	/**
@@ -62,7 +62,7 @@ protected:
 	 * The interchange manager is not calling this function directly. It is calling the blueprint native event in case this object is a blueprint derive object.
 	 * By default the scripted implementation is calling this virtual pipeline.
 	 */
-	virtual bool ExecuteImportPipeline(UInterchangeBaseNodeContainer* BaseNodeContainerAdapter)
+	virtual bool ExecuteImportPipeline(UInterchangeBaseNodeContainer* BaseNodeContainer)
 	{
 		return false;
 	}
@@ -79,7 +79,7 @@ protected:
 	}
 
 	/** This function can modify the BaseNodeContainer to create a pipeline that will set/validate the graph nodes hierarchy and options.*/
-	virtual bool ExecuteExportPipeline(UInterchangeBaseNodeContainer* BaseNodeContainerAdapter)
+	virtual bool ExecuteExportPipeline(UInterchangeBaseNodeContainer* BaseNodeContainer)
 	{
 		return false;
 	}
diff --git a/Engine/Source/Runtime/Interchange/Core/Public/InterchangeTranslatorBase.h b/Engine/Source/Runtime/Interchange/Core/Public/InterchangeTranslatorBase.h
--- a/Engine/Source/Runtime/Interchange/Core/Public/InterchangeTranslatorBase.h
+++ b/Engine/Source/Runtime/Interchange/Core/Public/InterchangeTranslatorBase.h
@@ -34,6 +34,17 @@ public:
 		return false;
 	}
 
+	/**
+	 * Release source is called when we do not need anymore the translator source and also early in the cancel process.
+	 * This is where out of process translator will send the stop command.
+	 * A good example of why its useful to release the source is when the dispatcher delete the worker, the join on the
+	 * thread will be very short and will not block the main thread. If the worker stop command was send before the completion task.
+	 */
+	virtual void ReleaseSource()
+	{
+		return;
+	}
+
 	/**
 	 * This function is call when the import is done (FImportAsyncHelper::CleanUp) and we are cleaning the data.
 	 * Use it to free resource that need to be release before the next garbage collector pass.
diff --git a/Engine/Source/Runtime/Interchange/Core/Public/Nodes/InterchangeBaseNode.h b/Engine/Source/Runtime/Interchange/Core/Public/Nodes/InterchangeBaseNode.h
--- a/Engine/Source/Runtime/Interchange/Core/Public/Nodes/InterchangeBaseNode.h
+++ b/Engine/Source/Runtime/Interchange/Core/Public/Nodes/InterchangeBaseNode.h
@@ -105,6 +105,24 @@ bool ApplyCustom##AttributeName##ToAsset(UObject* Asset) const										\
 		return true;																				\
 	}																								\
 	return false;																					\
+}																									\
+																									\
+bool FillCustom##AttributeName##FromAsset(UObject* Asset)											\
+{																									\
+	if (!Asset)																						\
+	{																								\
+		return false;																				\
+	}																								\
+	AssetType* TypedObject = Cast<AssetType>(Asset);												\
+	if (!TypedObject)																				\
+	{																								\
+		return false;																				\
+	}																								\
+	if (SetCustom##AttributeName((AttributeType)TypedObject->AttributeName, false))					\
+	{																								\
+		return true;																				\
+	}																								\
+	return false;																					\
 }
 
 #else //#if WITH_ENGINE
@@ -129,8 +147,10 @@ bool ApplyCustom##AttributeName##ToAsset(UObject* Asset) const										\
 	{																																				\
 		if(bAddApplyDelegate)																														\
 		{																																			\
-			TArray<UE::Interchange::FApplyAttributeToAsset>& Delegates = ApplyCustomAttributeDelegates.FindOrAdd(AssetType::StaticClass());				\
-			Delegates.Add(UE::Interchange::FApplyAttributeToAsset::CreateUObject(this, &NodeClassName::ApplyCustom##AttributeName##ToAsset));			\
+			TArray<UE::Interchange::FApplyAttributeToAsset>& Delegates = ApplyCustomAttributeDelegates.FindOrAdd(AssetType::StaticClass());			\
+			Delegates.Add(UE::Interchange::FApplyAttributeToAsset::CreateUObject(this, &NodeClassName::ApplyCustom##AttributeName##ToAsset));		\
+			TArray<UE::Interchange::FFillAttributeToAsset>& FillDelegates = FillCustomAttributeDelegates.FindOrAdd(AssetType::StaticClass());		\
+			FillDelegates.Add(UE::Interchange::FFillAttributeToAsset::CreateUObject(this, &NodeClassName::FillCustom##AttributeName##FromAsset));	\
 		}																																			\
 		return true;																																\
 	}																																				\
@@ -152,6 +172,7 @@ namespace UE
 	{
 
 		DECLARE_DELEGATE_RetVal_OneParam(bool, FApplyAttributeToAsset, UObject*);
+		DECLARE_DELEGATE_RetVal_OneParam(bool, FFillAttributeToAsset, UObject*);
 
 		/**
 		 * Helper struct use to declare static const data we use in the UInterchangeBaseNode
@@ -256,9 +277,18 @@ public:
 	 */
 	virtual bool HasAttribute(const UE::Interchange::FAttributeKey& NodeAttributeKey) const;
 
+	/**
+	 * This function return an attribute type for the specified Key. Return type None if the key is invalid
+	 *
+	 * @param NodeAttributeKey - The key of the attribute
+	 */
+	virtual UE::Interchange::EAttributeTypes GetAttributeType(const UE::Interchange::FAttributeKey& NodeAttributeKey) const
+	{
+		return Attributes.GetAttributeType(NodeAttributeKey);
+	}
 
 	/**
-	 * This function returnan  attribute handle for the specified Key.
+	 * This function return an  attribute handle for the specified Key.
 	 * If there is an issue with the KEY or storage the method will trip a check, always make sure you have a valid key before calling this
 	 * 
 	 * @param NodeAttributeKey - The key of the attribute
@@ -273,7 +303,8 @@ public:
 	 * Return the unique id pass in the constructor.
 	 *
 	 */
-	virtual FName GetUniqueID() const;
+	UFUNCTION(BlueprintCallable, Category = "Interchange | Node")
+	FName GetUniqueID() const;
 
 	/**
 	 * Return the display label pass in the constructor.
@@ -347,8 +378,25 @@ public:
 	 */
 	void ApplyAllCustomAttributeToAsset(UObject* Object) const;
 
+	void FillAllCustomAttributeFromAsset(UObject* Object) const;
+
 	virtual void Serialize(FArchive& Ar) override;
 
+	static void CompareNodeStorage(UInterchangeBaseNode* NodeA, const UInterchangeBaseNode* NodeB, TArray<UE::Interchange::FAttributeKey>& RemovedAttributes, TArray<UE::Interchange::FAttributeKey>& AddedAttributes, TArray<UE::Interchange::FAttributeKey>& ModifiedAttributes)
+	{
+		UE::Interchange::FAttributeStorage::CompareStorage(NodeA->Attributes, NodeB->Attributes, RemovedAttributes, AddedAttributes, ModifiedAttributes);
+	}
+	
+	static void CopyStorageAttributes(const UInterchangeBaseNode* SourceNode, UInterchangeBaseNode* DestinationNode, TArray<UE::Interchange::FAttributeKey>& AttributeKeys)
+	{
+		UE::Interchange::FAttributeStorage::CopyStorageAttributes(SourceNode->Attributes, DestinationNode->Attributes, AttributeKeys);
+	}
+
+	static void CopyStorage(const UInterchangeBaseNode* SourceNode, UInterchangeBaseNode* DestinationNode)
+	{
+		DestinationNode->Attributes = SourceNode->Attributes;
+	}
+	
 	UPROPERTY()
 	mutable FSoftObjectPath ReferenceObject;
 protected:
@@ -358,5 +406,7 @@ protected:
 	/* This array hold the delegate to apply the attribute that has to be set on an UObject */
 	TMap<UClass*, TArray<UE::Interchange::FApplyAttributeToAsset>> ApplyCustomAttributeDelegates;
 
+	TMap<UClass*, TArray<UE::Interchange::FFillAttributeToAsset>> FillCustomAttributeDelegates;
+
 	bool bIsInitialized = false;
 };
diff --git a/Engine/Source/Runtime/Interchange/Core/Public/Nodes/InterchangeBaseNodeContainer.h b/Engine/Source/Runtime/Interchange/Core/Public/Nodes/InterchangeBaseNodeContainer.h
--- a/Engine/Source/Runtime/Interchange/Core/Public/Nodes/InterchangeBaseNodeContainer.h
+++ b/Engine/Source/Runtime/Interchange/Core/Public/Nodes/InterchangeBaseNodeContainer.h
@@ -40,7 +40,7 @@ public:
 	bool IsNodeUIDValid(FName NodeUniqueID) const;
 
 	/** Unordered iteration of the all nodes */
-	void IterateNodes(TFunctionRef<void(FName, const UInterchangeBaseNode*)> IterationLambda);
+	void IterateNodes(TFunctionRef<void(FName, UInterchangeBaseNode*)> IterationLambda);
 
 	/** Return all nodes that do not have any parent */
 	UFUNCTION(BlueprintCallable, Category = "Interchange | Node Container")
diff --git a/Engine/Source/Runtime/Interchange/Core/Public/Types/AttributeStorage.h b/Engine/Source/Runtime/Interchange/Core/Public/Types/AttributeStorage.h
--- a/Engine/Source/Runtime/Interchange/Core/Public/Types/AttributeStorage.h
+++ b/Engine/Source/Runtime/Interchange/Core/Public/Types/AttributeStorage.h
@@ -528,7 +528,7 @@ namespace UE
 			 * @param ElementAttributeKey - is the storage key (the path) of the attribute.
 			 *
 			 */
-			EAttributeTypes GetAttributeType(const FAttributeKey& ElementAttributeKey);
+			EAttributeTypes GetAttributeType(const FAttributeKey& ElementAttributeKey) const;
 
 			/**
 			 * Return true if the attribute key point on an existing attribute in the storage. Return false otherwise
@@ -582,6 +582,17 @@ namespace UE
 			 */
 			static void CompareStorage(const FAttributeStorage& BaseStorage, const FAttributeStorage& VersionStorage, TArray<FAttributeKey>& RemovedAttributes, TArray<FAttributeKey>& AddedAttributes, TArray<FAttributeKey>& ModifiedAttributes);
 
+			/**
+			 * Copy an array of attributes from the source storage to the destination storage. If the attribute already exist in the destination, the value will be updated.
+			 * If a key do not exist in the source it will not be copy/created in the destination
+			 *
+			 * @param SourceStorage - The storage source
+			 * @param DestinationStorage - The storage destination
+			 * @param AttributeKeys - All attributes that must be copy from the source to the destination.
+			 *
+			 */
+			static void CopyStorageAttributes(const FAttributeStorage& SourceStorage, FAttributeStorage& DestinationStorage, const TArray<FAttributeKey>& AttributeKeys);
+
 			/**
 			 * Return the defrag ratio. This ratio is use to know when we need to defrag the storage
 			 * @example - a ratio of 0.1f will defrag the storage if the memory lost is bigger then 10% of the storage allocation.
diff --git a/Engine/Source/Runtime/Interchange/Engine/InterchangeEngine.Build.cs b/Engine/Source/Runtime/Interchange/Engine/InterchangeEngine.Build.cs
--- a/Engine/Source/Runtime/Interchange/Engine/InterchangeEngine.Build.cs
+++ b/Engine/Source/Runtime/Interchange/Engine/InterchangeEngine.Build.cs
@@ -6,10 +6,12 @@ public class InterchangeEngine : ModuleRules
 {
 	public InterchangeEngine(ReadOnlyTargetRules Target) : base(Target)
 	{
+		bTreatAsEngineModule = true;
 		PublicDependencyModuleNames.AddRange(
 				new string[] {
 					"Core",
 					"CoreUObject",
+					"Engine",
 					"InterchangeCore"
 				}
 			);
@@ -17,7 +19,8 @@ public class InterchangeEngine : ModuleRules
 		PrivateDependencyModuleNames.AddRange(
 			new string[] {
 				"AssetRegistry",
-				"Engine",
+				"SlateCore",
+				"Slate"
 			}
 		);
 	}
diff --git a/Engine/Source/Runtime/Interchange/Engine/Private/InterchangeManager.cpp b/Engine/Source/Runtime/Interchange/Engine/Private/InterchangeManager.cpp
--- a/Engine/Source/Runtime/Interchange/Engine/Private/InterchangeManager.cpp
+++ b/Engine/Source/Runtime/Interchange/Engine/Private/InterchangeManager.cpp
@@ -6,6 +6,7 @@
 #include "AssetRegistryModule.h"
 #include "CoreMinimal.h"
 #include "Engine/Blueprint.h"
+#include "Framework/Notifications/NotificationManager.h"
 #include "InterchangeFactoryBase.h"
 #include "InterchangeEngineLogPrivate.h"
 #include "InterchangeSourceData.h"
@@ -22,11 +23,13 @@
 #include "Tasks/InterchangeTaskPipeline.h"
 #include "Tasks/InterchangeTaskTranslator.h"
 #include "UObject/Class.h"
+#include "UObject/GarbageCollection.h"
 #include "UObject/NameTypes.h"
 #include "UObject/Object.h"
 #include "UObject/ObjectMacros.h"
 #include "UObject/UObjectIterator.h"
 #include "UObject/WeakObjectPtrTemplates.h"
+#include "Widgets/Notifications/SNotificationList.h"
 
 namespace InternalInterchangePrivate
 {
@@ -66,6 +69,7 @@ UInterchangeTranslatorBase* UE::Interchange::FScopedTranslator::GetTranslator()
 UE::Interchange::FImportAsyncHelper::FImportAsyncHelper()
 {
 	RootObjectCompletionEvent = FGraphEvent::CreateGraphEvent();
+	bCancel = false;
 }
 
 void UE::Interchange::FImportAsyncHelper::AddReferencedObjects(FReferenceCollector& Collector)
@@ -88,6 +92,61 @@ void UE::Interchange::FImportAsyncHelper::AddReferencedObjects(FReferenceCollect
 	}
 }
 
+void UE::Interchange::FImportAsyncHelper::ReleaseTranslatorsSource()
+{
+	for (UInterchangeTranslatorBase* BaseTranslator : Translators)
+	{
+		BaseTranslator->ReleaseSource();
+	}
+}
+
+void UE::Interchange::FImportAsyncHelper::InitCancel()
+{
+	bCancel = true;
+	ReleaseTranslatorsSource();
+}
+
+void UE::Interchange::FImportAsyncHelper::CancelAndWaitUntilDoneSynchronously()
+{
+	bCancel = true;
+
+	FGraphEventArray TasksToComplete;
+
+	if (TranslatorTasks.Num())
+	{
+		TasksToComplete.Append(TranslatorTasks);
+	}
+	if (PipelineTasks.Num())
+	{
+		TasksToComplete.Append(PipelineTasks);
+	}
+	if (ParsingTask.GetReference())
+	{
+		TasksToComplete.Add(ParsingTask);
+	}
+	if (CreateAssetTasks.Num())
+	{
+		TasksToComplete.Append(CreateAssetTasks);
+	}
+	if (CompletionTask.GetReference())
+	{
+		//Completion task will make sure any created asset before canceling will be mark for delete
+		TasksToComplete.Add(CompletionTask);
+	}
+
+	//Block until all task are completed, it should be fast since bCancel is true
+	if (TasksToComplete.Num())
+	{
+		FTaskGraphInterface::Get().WaitUntilTasksComplete(TasksToComplete, ENamedThreads::GameThread);
+	}
+	//Async import result in a null object when we cancel
+	if (!RootObjectCompletionEvent->IsComplete())
+	{
+		RootObject.SetValue(nullptr);
+		RootObjectCompletionEvent->DispatchSubsequents();
+	}
+}
+
 void UE::Interchange::FImportAsyncHelper::CleanUp()
 {
 	//Release the graph
@@ -168,6 +227,56 @@ void UE::Interchange::SanitizeInvalidChar(FString& String)
 	}
 }
 
+UInterchangeManager& UInterchangeManager::GetInterchangeManager()
+{
+	static TStrongObjectPtr<UInterchangeManager> InterchangeManager = nullptr;
+	
+	//This boolean will be true after we delete the singleton
+	static bool InterchangeManagerScopeOfLifeEnded = false;
+
+	if (!InterchangeManager.IsValid())
+	{
+		//We cannot create a TStrongObjectPtr outside of the main thread, we also need a valid Transient package
+		check(IsInGameThread() && GetTransientPackage());
+
+		//Avoid hard crash if someone call the manager after we delete it, but send a callstack to the crash manager
+		ensure(!InterchangeManagerScopeOfLifeEnded);
+
+		InterchangeManager = TStrongObjectPtr<UInterchangeManager>(NewObject<UInterchangeManager>(GetTransientPackage(), NAME_None, EObjectFlags::RF_NoFlags));
+		
+		//We cancel any running task when we pre exit the engine
+		FCoreDelegates::OnEnginePreExit.AddLambda([]()
+		{
+			//In editor the user cannot exit the editor if the interchange manager has active task.
+			//But if we are not running the editor its possible to get here, so block the main thread until all
+			//cancel tasks are done.
+			if(GIsEditor)
+			{
+				ensure(InterchangeManager->ImportTasks.Num() == 0);
+			}
+			else
+			{
+				InterchangeManager->CancelAllTasksSynchronously();
+			}
+		});
+
+		//We release the singleton here so all module was able to unhook there delegates
+		FCoreDelegates::OnExit.AddLambda([]()
+		{
+			//Task should have been cancel in the Engine pre exit callback
+			ensure(InterchangeManager->ImportTasks.Num() == 0);
+			//Release the InterchangeManager object
+			InterchangeManager.Reset();
+			InterchangeManagerScopeOfLifeEnded = true;
+		});
+	}
+
+	//When we get here we should be valid
+	check(InterchangeManager.IsValid());
+
+	return *(InterchangeManager.Get());
+}
+
 bool UInterchangeManager::RegisterTranslator(const UClass* TranslatorClass)
 {
 	if(!TranslatorClass)
@@ -273,7 +382,7 @@ UE::Interchange::FAsyncImportResult UInterchangeManager::ImportAssetAsync(const
 		NotificationConfig.bKeepOpenOnFailure = true;
 		NotificationConfig.TitleText = TitleText;
 		NotificationConfig.LogCategory = InternalInterchangePrivate::GetLogInterchangePtr();
-		NotificationConfig.bCanCancel.Set(false);
+		NotificationConfig.bCanCancel.Set(true);
 		NotificationConfig.bKeepOpenOnFailure.Set(true);
 
 		Notification = MakeShared<FAsyncTaskNotification>(NotificationConfig);
@@ -302,7 +411,7 @@ UE::Interchange::FAsyncImportResult UInterchangeManager::ImportAssetAsync(const
 	if ( ImportAssetParameters.OverridePipeline == nullptr )
 	{
 		//Get all pipeline candidate we want for this import
-		/*TArray<UClass*> PipelineCandidates;
+		TArray<UClass*> PipelineCandidates;
 		FindPipelineCandidate(PipelineCandidates);
 
 		// Stack all pipelines, for this import proto. TODO: We need to be able to control which pipeline we use for the import
@@ -312,7 +421,7 @@ UE::Interchange::FAsyncImportResult UInterchangeManager::ImportAssetAsync(const
 		{
 			UInterchangePipelineBase* GeneratedPipeline = NewObject<UInterchangePipelineBase>(GetTransientPackage(), PipelineCandidates[GraphPipelineNumber], NAME_None, RF_NoFlags);
 			AsyncHelper->Pipelines.Add(GeneratedPipeline);
-		}*/
+		}
 	}
 	else
 	{
@@ -389,6 +498,7 @@ TWeakPtr<UE::Interchange::FImportAsyncHelper, ESPMode::ThreadSafe> UInterchangeM
 	//Copy the task data
 	AsyncHelper->TaskData = Data;
 	int32 AsyncHelperIndex = ImportTasks.Add(AsyncHelper);
+	SetActiveMode(true);
 	//Update the asynchronous notification
 	if (Notification.IsValid())
 	{
@@ -404,17 +514,26 @@ void UInterchangeManager::ReleaseAsyncHelper(TWeakPtr<UE::Interchange::FImportAs
 {
 	check(AsyncHelper.IsValid());
 	ImportTasks.RemoveSingle(AsyncHelper.Pin());
-	check(!AsyncHelper.IsValid());
+	//Make sure the async helper is destroy, if not destroy its because we are canceling the import and we still have a shared ptr on it
+	{
+		TSharedPtr<UE::Interchange::FImportAsyncHelper, ESPMode::ThreadSafe> AsyncHelperSharedPtr = AsyncHelper.Pin();
+		check(!AsyncHelperSharedPtr.IsValid() || AsyncHelperSharedPtr->bCancel);
+	}
 
 	int32 ImportTaskNumber = ImportTasks.Num();
 	FString ImportTaskNumberStr = TEXT(" (") + FString::FromInt(ImportTaskNumber) + TEXT(")");
-	if (ImportTasks.Num() == 0 && Notification.IsValid())
+	if (ImportTaskNumber == 0)
 	{
-		FText TitleText = NSLOCTEXT("Interchange", "Asynchronous_import_end", "Import Done");
-		//TODO make sure any error are reported so we can control success or not
-		const bool bSuccess = true;
-		Notification->SetComplete(TitleText, FText::GetEmpty(), bSuccess);
-		Notification = nullptr; //This should delete the notification
+		SetActiveMode(false);
+
+		if (Notification.IsValid())
+		{
+			FText TitleText = NSLOCTEXT("Interchange", "Asynchronous_import_end", "Import Done");
+			//TODO make sure any error are reported so we can control success or not
+			const bool bSuccess = true;
+			Notification->SetComplete(TitleText, FText::GetEmpty(), bSuccess);
+			Notification = nullptr; //This should delete the notification
+		}
 	}
 	else if(Notification.IsValid())
 	{
@@ -440,6 +559,23 @@ UInterchangeTranslatorBase* UInterchangeManager::GetTranslatorForSourceData(cons
 	return nullptr;
 }
 
+bool UInterchangeManager::WarnIfInterchangeIsActive()
+{
+	if (!bIsActive)
+	{
+		return false;
+	}
+	//Tell user he have to cancel the import before closing the editor
+	FNotificationInfo Info(NSLOCTEXT("InterchangeManager", "WarnCannotProceed", "An import process is currently underway! Please cancel it to proceed!"));
+	Info.ExpireDuration = 5.0f;
+	TSharedPtr<SNotificationItem> WarnNotification = FSlateNotificationManager::Get().AddNotification(Info);
+	if (WarnNotification.IsValid())
+	{
+		WarnNotification->SetCompletionState(SNotificationItem::CS_Fail);
+	}
+	return true;
+}
+
 bool UInterchangeManager::IsAttended()
 {
 	if (FApp::IsGame())
@@ -460,10 +596,10 @@ void UInterchangeManager::FindPipelineCandidate(TArray<UClass*>& PipelineCandida
 	{
 		UClass* Class = *ClassIt;
 		// Only interested in native C++ classes
-		if (!Class->IsNative())
-		{
-			continue;
-		}
+// 		if (!Class->IsNative())
+// 		{
+// 			continue;
+// 		}
 		// Ignore deprecated
 		if (Class->HasAnyClassFlags(CLASS_Deprecated | CLASS_NewerVersionExists))
 		{
@@ -534,3 +670,89 @@ void UInterchangeManager::FindPipelineCandidate(TArray<UClass*>& PipelineCandida
 		}
 	}
 }
+
+void UInterchangeManager::CancelAllTasks()
+{
+	check(IsInGameThread());
+
+	//Set the cancel state on all task
+	int32 ImportTaskCount = ImportTasks.Num();
+	for (int32 TaskIndex = 0; TaskIndex < ImportTaskCount; ++TaskIndex)
+	{
+		TSharedPtr<UE::Interchange::FImportAsyncHelper, ESPMode::ThreadSafe> AsyncHelper = ImportTasks[TaskIndex];
+		if (AsyncHelper.IsValid())
+		{
+			AsyncHelper->InitCancel();
+		}
+	}
+
+	//Tasks should all finish quite fast now
+};
+
+void UInterchangeManager::CancelAllTasksSynchronously()
+{
+	//Start the cancel process by cancelling all current task
+	CancelAllTasks();
+
+	//Now wait for each task to be completed on the main thread
+	while (ImportTasks.Num() > 0)
+	{
+		int32 ImportTaskCount = ImportTasks.Num();
+		TSharedPtr<UE::Interchange::FImportAsyncHelper, ESPMode::ThreadSafe> AsyncHelper = ImportTasks[0];
+		if (AsyncHelper.IsValid())
+		{
+			//Cancel any on going interchange activity this is blocking but necessary.
+			AsyncHelper->CancelAndWaitUntilDoneSynchronously();
+			ensure(ImportTaskCount > ImportTasks.Num());
+			TWeakPtr<UE::Interchange::FImportAsyncHelper, ESPMode::ThreadSafe> WeakAsyncHelper = AsyncHelper;
+			//free the async helper
+			AsyncHelper = nullptr;
+			//We verify that the weak pointer is invalid after releasing the async helper
+			ensure(!WeakAsyncHelper.IsValid());
+		}
+	}
+}
+
+void UInterchangeManager::SetActiveMode(bool IsActive)
+{
+	if (bIsActive == IsActive)
+	{
+		return;
+	}
+
+	bIsActive = IsActive;
+	if (bIsActive)
+	{
+		ensure(!NotificationTickHandle.IsValid());
+		NotificationTickHandle = FTicker::GetCoreTicker().AddTicker(TEXT("InterchangeManagerTickHandle"), 0.1f, [this](float)
+		{
+			if (Notification.IsValid() && Notification->GetPromptAction() == EAsyncTaskNotificationPromptAction::Cancel)
+			{
+				CancelAllTasks();
+			}
+			return true;
+		});
+
+		//Block GC in a different thread then game thread
+		FString ThreadName = FString(TEXT("InterchangeGCGuard"));
+		GcGuardThread = FThread(*ThreadName, [this]()
+		{
+			FGCScopeGuard GCScopeGuard;
+			while (bIsActive && ImportTasks.Num() > 0)
+			{
+				FPlatformProcess::Sleep(0.01f);
+			}
+		});
+	}
+	else
+	{
+		FTicker::GetCoreTicker().RemoveTicker(NotificationTickHandle);
+		NotificationTickHandle.Reset();
+
+		if (GcGuardThread.IsJoinable())
+		{
+			//Finish the thread
+			GcGuardThread.Join();
+		}
+	}
+}
\ No newline at end of file
diff --git a/Engine/Source/Runtime/Interchange/Engine/Private/Tasks/InterchangeTaskCompletion.cpp b/Engine/Source/Runtime/Interchange/Engine/Private/Tasks/InterchangeTaskCompletion.cpp
--- a/Engine/Source/Runtime/Interchange/Engine/Private/Tasks/InterchangeTaskCompletion.cpp
+++ b/Engine/Source/Runtime/Interchange/Engine/Private/Tasks/InterchangeTaskCompletion.cpp
@@ -16,9 +16,18 @@ void UE::Interchange::FTaskCompletion::DoTask(ENamedThreads::Type CurrentThread,
 {
 	TSharedPtr<FImportAsyncHelper, ESPMode::ThreadSafe> AsyncHelper = WeakAsyncHelper.Pin();
 	check(AsyncHelper.IsValid());
+
+	//No need anymore of the translators sources
+	AsyncHelper->ReleaseTranslatorsSource();
+
 	bool bIsFutureRootObjectSet = false;
 	for(TPair<int32, TArray<FImportAsyncHelper::FImportedAssetInfo>>& AssetInfosPerSourceIndexPair : AsyncHelper->ImportedAssetsPerSourceIndex)
 	{
+		//Verify if the task was cancel
+		if (AsyncHelper->bCancel)
+		{
+			break;
+		}
 		const int32 SourceIndex = AssetInfosPerSourceIndexPair.Key;
 		const bool bCallPostImportGameThreadCallback = ensure(AsyncHelper->SourceDatas.IsValidIndex(SourceIndex));
 		for (const FImportAsyncHelper::FImportedAssetInfo& AssetInfo : AssetInfosPerSourceIndexPair.Value)
@@ -28,8 +37,10 @@ void UE::Interchange::FTaskCompletion::DoTask(ENamedThreads::Type CurrentThread,
 			if(bCallPostImportGameThreadCallback && AssetInfo.Factory)
 			{
 				UInterchangeFactoryBase::FPostImportGameThreadCallbackParams Arguments;
-				Arguments.ReimportObject = Asset;
+				Arguments.ImportedObject = Asset;
 				Arguments.SourceData = AsyncHelper->SourceDatas[SourceIndex];
+				Arguments.NodeUniqueID = AssetInfo.NodeUniqueId;
+				Arguments.NodeContainer = AsyncHelper->BaseNodeContainers[SourceIndex].Get();
 				AssetInfo.Factory->PostImportGameThreadCallback(Arguments);
 			}
 
@@ -49,9 +60,9 @@ void UE::Interchange::FTaskCompletion::DoTask(ENamedThreads::Type CurrentThread,
 			else
 			{
 				InterchangeManager->OnAssetPostImport.Broadcast(Asset);
+				//Notify the asset registry, only when we have created the asset
+				FAssetRegistryModule::AssetCreated(Asset);
 			}
-			//Notify the asset registry
-			FAssetRegistryModule::AssetCreated(Asset);
 
 			if (!bIsFutureRootObjectSet && SourceIndex == 0)
 			{
@@ -62,6 +73,27 @@ void UE::Interchange::FTaskCompletion::DoTask(ENamedThreads::Type CurrentThread,
 		}
 	}
 
+	//If task is cancel, delete all created assets by this task
+	if (AsyncHelper->bCancel)
+	{
+		for (TPair<int32, TArray<FImportAsyncHelper::FImportedAssetInfo>>& AssetInfosPerSourceIndexPair : AsyncHelper->ImportedAssetsPerSourceIndex)
+		{
+			const int32 SourceIndex = AssetInfosPerSourceIndexPair.Key;
+			for (const FImportAsyncHelper::FImportedAssetInfo& AssetInfo : AssetInfosPerSourceIndexPair.Value)
+			{
+				UObject* Asset = AssetInfo.ImportAsset;
+				if (Asset)
+				{
+					//Make any created asset go away
+					Asset->ClearFlags(RF_Standalone | RF_Public | RF_Transactional);
+					Asset->ClearInternalFlags(EInternalObjectFlags::Async);
+					Asset->SetFlags(RF_Transient);
+					Asset->MarkPendingKill();
+				}
+			}
+		}
+	}
+
 	if (!AsyncHelper->RootObjectCompletionEvent->IsComplete())
 	{
 		AsyncHelper->RootObject.SetValue(nullptr);
diff --git a/Engine/Source/Runtime/Interchange/Engine/Private/Tasks/InterchangeTaskCreateAsset.cpp b/Engine/Source/Runtime/Interchange/Engine/Private/Tasks/InterchangeTaskCreateAsset.cpp
--- a/Engine/Source/Runtime/Interchange/Engine/Private/Tasks/InterchangeTaskCreateAsset.cpp
+++ b/Engine/Source/Runtime/Interchange/Engine/Private/Tasks/InterchangeTaskCreateAsset.cpp
@@ -54,8 +54,11 @@ void UE::Interchange::FTaskCreatePackage::DoTask(ENamedThreads::Type CurrentThre
 	TSharedPtr<UE::Interchange::FImportAsyncHelper, ESPMode::ThreadSafe> AsyncHelper = WeakAsyncHelper.Pin();
 	check(AsyncHelper.IsValid());
 
-	//The create package thread must always execute on the game thread
-	check(IsInGameThread());
+	//Verify if the task was cancel
+	if (AsyncHelper->bCancel)
+	{
+		return;
+	}
 
 	UPackage* Pkg = nullptr;
 	FString PackageName;
@@ -68,6 +71,9 @@ void UE::Interchange::FTaskCreatePackage::DoTask(ENamedThreads::Type CurrentThre
 	}
 	else
 	{
+		//The create package thread must always execute on the game thread
+		check(IsInGameThread());
+
 		Private::InternalGetPackageName(*AsyncHelper, SourceIndex, PackageBasePath, Node, PackageName, AssetName);
 		// We can not create assets that share the name of a map file in the same location
 		if (UE::Interchange::FPackageUtils::IsMapPackageAsset(PackageName))
@@ -100,13 +106,25 @@ void UE::Interchange::FTaskCreatePackage::DoTask(ENamedThreads::Type CurrentThre
 		}
 		CreateAssetParams.ReimportObject = AsyncHelper->TaskData.ReimportObject;
 		//Make sure the asset UObject is created with the correct type on the main thread
-		Factory->CreateEmptyAsset(CreateAssetParams);
+		UObject* NodeAsset = Factory->CreateEmptyAsset(CreateAssetParams);
+		if (NodeAsset)
+		{
+			FScopeLock Lock(&AsyncHelper->ImportedAssetsPerSourceIndexLock);
+			TArray<UE::Interchange::FImportAsyncHelper::FImportedAssetInfo>& ImportedInfos = AsyncHelper->ImportedAssetsPerSourceIndex.FindOrAdd(SourceIndex);
+			UE::Interchange::FImportAsyncHelper::FImportedAssetInfo& AssetInfo = ImportedInfos.AddDefaulted_GetRef();
+			AssetInfo.ImportAsset = NodeAsset;
+			AssetInfo.Factory = Factory;
+			AssetInfo.NodeUniqueId = Node->GetUniqueID().ToString();
+			Node->ReferenceObject = FSoftObjectPath(NodeAsset);
+		}
 	}
 	// Make sure the destination package is loaded
 	Pkg->FullyLoad();
-
-	FScopeLock Lock(&AsyncHelper->CreatedPackagesLock);
-	AsyncHelper->CreatedPackages.Add(PackageName, Pkg);
+	
+	{
+		FScopeLock Lock(&AsyncHelper->CreatedPackagesLock);
+		AsyncHelper->CreatedPackages.Add(PackageName, Pkg);
+	}
 }
 
 void UE::Interchange::FTaskCreateAsset::DoTask(ENamedThreads::Type CurrentThread, const FGraphEventRef& MyCompletionGraphEvent)
@@ -114,16 +132,22 @@ void UE::Interchange::FTaskCreateAsset::DoTask(ENamedThreads::Type CurrentThread
 	TSharedPtr<UE::Interchange::FImportAsyncHelper, ESPMode::ThreadSafe> AsyncHelper = WeakAsyncHelper.Pin();
 	check(WeakAsyncHelper.IsValid());
 
+	//Verify if the task was cancel
+	if (AsyncHelper->bCancel)
+	{
+		return;
+	}
+
+	UPackage* Pkg = nullptr;
 	FString PackageName;
 	FString AssetName;
 	Private::InternalGetPackageName(*AsyncHelper, SourceIndex, PackageBasePath, Node, PackageName, AssetName);
 	if (AsyncHelper->TaskData.ReimportObject)
 	{
-		UPackage* Pkg = AsyncHelper->TaskData.ReimportObject->GetPackage();
+		Pkg = AsyncHelper->TaskData.ReimportObject->GetPackage();
 		PackageName = Pkg->GetPathName();
 	}
-
-	UPackage* Pkg = nullptr;
+	else
 	{
 		FScopeLock Lock(&AsyncHelper->CreatedPackagesLock);
 		UPackage** PkgPtr = AsyncHelper->CreatedPackages.Find(PackageName);
@@ -159,14 +183,28 @@ void UE::Interchange::FTaskCreateAsset::DoTask(ENamedThreads::Type CurrentThread
 	}
 	CreateAssetParams.ReimportObject = AsyncHelper->TaskData.ReimportObject;
 
+	CreateAssetParams.ReimportStrategyFlags = EReimportStrategyFlags::ApplyNoProperties;
+	//CreateAssetParams.ReimportStrategyFlags = EReimportStrategyFlags::ApplyPipelineProperties;
+	//CreateAssetParams.ReimportStrategyFlags = EReimportStrategyFlags::ApplyEditorChangedProperties;
+
 	UObject* NodeAsset = Factory->CreateAsset(CreateAssetParams);
 	if (NodeAsset)
 	{
 		FScopeLock Lock(&AsyncHelper->ImportedAssetsPerSourceIndexLock);
 		TArray<UE::Interchange::FImportAsyncHelper::FImportedAssetInfo>& ImportedInfos = AsyncHelper->ImportedAssetsPerSourceIndex.FindOrAdd(SourceIndex);
-		UE::Interchange::FImportAsyncHelper::FImportedAssetInfo& AssetInfo = ImportedInfos.AddZeroed_GetRef();
-		AssetInfo.ImportAsset = NodeAsset;
+		UE::Interchange::FImportAsyncHelper::FImportedAssetInfo* AssetInfoPtr = ImportedInfos.FindByPredicate([NodeAsset](const UE::Interchange::FImportAsyncHelper::FImportedAssetInfo& CurInfo)
+		{
+			return CurInfo.ImportAsset == NodeAsset;
+		});
+
+		if (!AssetInfoPtr)
+		{
+			UE::Interchange::FImportAsyncHelper::FImportedAssetInfo& AssetInfo = ImportedInfos.AddDefaulted_GetRef();
+			AssetInfo.ImportAsset = NodeAsset;
+			AssetInfo.Factory = Factory;
+			AssetInfo.NodeUniqueId = Node->GetUniqueID().ToString();
+		}
+
 		Node->ReferenceObject = FSoftObjectPath(NodeAsset);
-		AssetInfo.Factory = Factory;
 	}
 }
\ No newline at end of file
diff --git a/Engine/Source/Runtime/Interchange/Engine/Private/Tasks/InterchangeTaskCreateAsset.h b/Engine/Source/Runtime/Interchange/Engine/Private/Tasks/InterchangeTaskCreateAsset.h
--- a/Engine/Source/Runtime/Interchange/Engine/Private/Tasks/InterchangeTaskCreateAsset.h
+++ b/Engine/Source/Runtime/Interchange/Engine/Private/Tasks/InterchangeTaskCreateAsset.h
@@ -24,11 +24,11 @@ namespace UE
 			FString PackageBasePath;
 			int32 SourceIndex;
 			TWeakPtr<FImportAsyncHelper, ESPMode::ThreadSafe> WeakAsyncHelper;
-			const UInterchangeBaseNode* Node;
+			UInterchangeBaseNode* Node;
 			UInterchangeFactoryBase* Factory;
 
 		public:
-			FTaskCreatePackage(const FString& InPackageBasePath, const int32 InSourceIndex, TWeakPtr<FImportAsyncHelper, ESPMode::ThreadSafe> InAsyncHelper, const UInterchangeBaseNode* InNode, UInterchangeFactoryBase* InFactory)
+			FTaskCreatePackage(const FString& InPackageBasePath, const int32 InSourceIndex, TWeakPtr<FImportAsyncHelper, ESPMode::ThreadSafe> InAsyncHelper, UInterchangeBaseNode* InNode, UInterchangeFactoryBase* InFactory)
 				: PackageBasePath(InPackageBasePath)
 				, SourceIndex(InSourceIndex)
 				, WeakAsyncHelper(InAsyncHelper)
@@ -68,11 +68,11 @@ namespace UE
 			FString PackageBasePath;
 			int32 SourceIndex;
 			TWeakPtr<FImportAsyncHelper, ESPMode::ThreadSafe> WeakAsyncHelper;
-			const UInterchangeBaseNode* Node;
+			UInterchangeBaseNode* Node;
 			UInterchangeFactoryBase* Factory;
 
 		public:
-			FTaskCreateAsset(const FString& InPackageBasePath, const int32 InSourceIndex, TWeakPtr<FImportAsyncHelper, ESPMode::ThreadSafe> InAsyncHelper, const UInterchangeBaseNode* InNode, UInterchangeFactoryBase* InFactory)
+			FTaskCreateAsset(const FString& InPackageBasePath, const int32 InSourceIndex, TWeakPtr<FImportAsyncHelper, ESPMode::ThreadSafe> InAsyncHelper, UInterchangeBaseNode* InNode, UInterchangeFactoryBase* InFactory)
 				: PackageBasePath(InPackageBasePath)
 				, SourceIndex(InSourceIndex)
 				, WeakAsyncHelper(InAsyncHelper)
diff --git a/Engine/Source/Runtime/Interchange/Engine/Private/Tasks/InterchangeTaskParsing.cpp b/Engine/Source/Runtime/Interchange/Engine/Private/Tasks/InterchangeTaskParsing.cpp
--- a/Engine/Source/Runtime/Interchange/Engine/Private/Tasks/InterchangeTaskParsing.cpp
+++ b/Engine/Source/Runtime/Interchange/Engine/Private/Tasks/InterchangeTaskParsing.cpp
@@ -30,7 +30,7 @@ void UE::Interchange::FTaskParsing::DoTask(ENamedThreads::Type CurrentThread, co
 	{
 		FName UniqueID;
 		int32 SourceIndex;
-		const UInterchangeBaseNode* Node;
+		UInterchangeBaseNode* Node;
 		TArray<FName> Dependencies;
 		FGraphEventRef GraphEventRef;
 		FGraphEventArray Prerequistes;
@@ -38,40 +38,43 @@ void UE::Interchange::FTaskParsing::DoTask(ENamedThreads::Type CurrentThread, co
 	};
 
 
-
 	FGraphEventArray CompletionPrerequistes;
 	TArray<FTaskData> TaskDatas;
-	for (int32 SourceIndex = 0; SourceIndex < AsyncHelper->SourceDatas.Num(); ++SourceIndex)
+
+	//Avoid creating asset if the asynchronous import is cancel, just create the completion task
+	if (!AsyncHelper->bCancel)
 	{
-		if (!AsyncHelper->BaseNodeContainers.IsValidIndex(SourceIndex))
-		{
-			continue;
-		}
-		UInterchangeBaseNodeContainer* BaseNodeContainer = AsyncHelper->BaseNodeContainers[SourceIndex].Get();
-		if (!BaseNodeContainer)
+		for (int32 SourceIndex = 0; SourceIndex < AsyncHelper->SourceDatas.Num(); ++SourceIndex)
 		{
-			continue;
-		}
-		BaseNodeContainer->IterateNodes([&](const FName NodeUID, const UInterchangeBaseNode* Node)
-		{
-			if (Node->GetAssetClass() != nullptr)
+			if (!AsyncHelper->BaseNodeContainers.IsValidIndex(SourceIndex))
 			{
-				UInterchangeFactoryBase* NodeFactory = InterchangeManager->GetRegisterFactory(Node->GetAssetClass());
-				if (!NodeFactory)
+				continue;
+			}
+			UInterchangeBaseNodeContainer* BaseNodeContainer = AsyncHelper->BaseNodeContainers[SourceIndex].Get();
+			if (!BaseNodeContainer)
+			{
+				continue;
+			}
+			BaseNodeContainer->IterateNodes([&](const FName NodeUID, UInterchangeBaseNode* Node)
+			{
+				if (Node->GetAssetClass() != nullptr)
 				{
-					//nothing we can import from this element
-					return;
+					UInterchangeFactoryBase* NodeFactory = InterchangeManager->GetRegisterFactory(Node->GetAssetClass());
+					if (!NodeFactory)
+					{
+						//nothing we can import from this element
+						return;
+					}
+					FTaskData& NodeTaskData = TaskDatas.AddDefaulted_GetRef();
+					NodeTaskData.UniqueID = Node->GetUniqueID();
+					NodeTaskData.SourceIndex = SourceIndex;
+					NodeTaskData.Node = Node;
+					Node->GetDependecies(NodeTaskData.Dependencies);
+					NodeTaskData.Factory = NodeFactory;
 				}
-				FTaskData& NodeTaskData = TaskDatas.AddDefaulted_GetRef();
-				NodeTaskData.UniqueID = Node->GetUniqueID();
-				NodeTaskData.SourceIndex = SourceIndex;
-				NodeTaskData.Node = Node;
-				Node->GetDependecies(NodeTaskData.Dependencies);
-				NodeTaskData.Factory = NodeFactory;
-			}
-		});
+			});
+		}
 	}
-
 	//Sort per dependencies
 	TaskDatas.Sort([](const FTaskData& A, const FTaskData& B)
 	{
diff --git a/Engine/Source/Runtime/Interchange/Engine/Private/Tasks/InterchangeTaskPipeline.cpp b/Engine/Source/Runtime/Interchange/Engine/Private/Tasks/InterchangeTaskPipeline.cpp
--- a/Engine/Source/Runtime/Interchange/Engine/Private/Tasks/InterchangeTaskPipeline.cpp
+++ b/Engine/Source/Runtime/Interchange/Engine/Private/Tasks/InterchangeTaskPipeline.cpp
@@ -19,8 +19,16 @@ void UE::Interchange::FTaskPipeline::DoTask(ENamedThreads::Type CurrentThread, c
 	{
 		for (int32 GraphIndex = 0; GraphIndex < AsyncHelper->BaseNodeContainers.Num(); ++GraphIndex)
 		{
-			check(AsyncHelper->BaseNodeContainers[GraphIndex].IsValid());
-			Pipeline->ScriptedExecuteImportPipeline(AsyncHelper->BaseNodeContainers[GraphIndex].Get());
+			//Verify if the task was cancel
+			if (AsyncHelper->bCancel)
+			{
+				return;
+			}
+
+			if (ensure(AsyncHelper->BaseNodeContainers[GraphIndex].IsValid()))
+			{
+				Pipeline->ScriptedExecuteImportPipeline(AsyncHelper->BaseNodeContainers[GraphIndex].Get());
+			}
 		}
 	}
 }
diff --git a/Engine/Source/Runtime/Interchange/Engine/Private/Tasks/InterchangeTaskTranslator.cpp b/Engine/Source/Runtime/Interchange/Engine/Private/Tasks/InterchangeTaskTranslator.cpp
--- a/Engine/Source/Runtime/Interchange/Engine/Private/Tasks/InterchangeTaskTranslator.cpp
+++ b/Engine/Source/Runtime/Interchange/Engine/Private/Tasks/InterchangeTaskTranslator.cpp
@@ -32,6 +32,12 @@ void UE::Interchange::FTaskTranslator::DoTask(ENamedThreads::Type CurrentThread,
 		return;
 	}
 
+	//Verify if the task was cancel
+	if (AsyncHelper->bCancel)
+	{
+		return;
+	}
+
 	//Translate the source data
 	UInterchangeBaseNodeContainer& BaseNodeContainer = *(AsyncHelper->BaseNodeContainers[SourceIndex].Get());
 	Translator->Translate(SourceData, BaseNodeContainer);
diff --git a/Engine/Source/Runtime/Interchange/Engine/Public/InterchangeAssetImportData.h b/Engine/Source/Runtime/Interchange/Engine/Public/InterchangeAssetImportData.h
new file mode 100644
--- /dev/null
+++ b/Engine/Source/Runtime/Interchange/Engine/Public/InterchangeAssetImportData.h
@@ -0,0 +1,58 @@
+// Copyright Epic Games, Inc. All Rights Reserved.
+
+#pragma once
+
+#include "CoreMinimal.h"
+#include "EditorFramework/AssetImportData.h"
+#include "Nodes/InterchangeBaseNode.h"
+#include "Nodes/InterchangeBaseNodeContainer.h"
+#include "UObject/Class.h"
+#include "UObject/Object.h"
+#include "UObject/ObjectMacros.h"
+
+
+#include "InterchangeAssetImportData.generated.h"
+
+UCLASS(BlueprintType)
+class INTERCHANGEENGINE_API UInterchangeAssetImportData : public UAssetImportData
+{
+	GENERATED_BODY()
+public:
+
+	/**
+	 * Return the first filename stored in this data. The resulting filename will be absolute (ie, not relative to the asset).
+	 */
+	UFUNCTION(BlueprintCallable, Category = "Interchange | AssetImportData")
+	FString ScriptGetFirstFilename() const
+	{
+		return GetFirstFilename();
+	}
+
+	/**
+	 * Extract all the (resolved) filenames.
+	 */
+	UFUNCTION(BlueprintCallable, Category = "Interchange | AssetImportData")
+	TArray<FString> ScriptExtractFilenames() const
+	{
+		return ExtractFilenames();
+	}
+
+	/**
+	 * Extract all the filename labels.
+	 */
+	UFUNCTION(BlueprintCallable, Category = "Interchange | AssetImportData")
+	TArray<FString> ScriptExtractDisplayLabels() const
+	{
+		TArray<FString> TempDisplayLabels;
+		ExtractDisplayLabels(TempDisplayLabels);
+		return TempDisplayLabels;
+	}
+	
+	/** The graph that was use to create this asset */
+	UPROPERTY(VisibleAnywhere, Category = "Interchange | AssetImportData")
+	UInterchangeBaseNodeContainer* NodeContainer;
+
+	/** The Node UID pass to the factory that exist in the graph that was use to create this asset */
+	UPROPERTY(VisibleAnywhere, Category = "Interchange | AssetImportData")
+	FString NodeUniqueID;
+};
diff --git a/Engine/Source/Runtime/Interchange/Engine/Public/InterchangeManager.h b/Engine/Source/Runtime/Interchange/Engine/Public/InterchangeManager.h
--- a/Engine/Source/Runtime/Interchange/Engine/Public/InterchangeManager.h
+++ b/Engine/Source/Runtime/Interchange/Engine/Public/InterchangeManager.h
@@ -2,9 +2,12 @@
 
 #pragma once
 
+#include <atomic>
 #include "Async/TaskGraphInterfaces.h"
 #include "CoreMinimal.h"
 #include "Delegates/DelegateCombinations.h"
+#include "HAL/Thread.h"
+#include "HAL/ThreadSafeBool.h"
 #include "InterchangeFactoryBase.h"
 #include "InterchangePipelineBase.h"
 #include "InterchangeSourceData.h"
@@ -96,7 +99,7 @@ namespace UE
 			{
 				UObject* ImportAsset;
 				UInterchangeFactoryBase* Factory;
-				UInterchangeBaseNode* AssetNode;
+				FString NodeUniqueId;
 			};
 
 			FCriticalSection ImportedAssetsPerSourceIndexLock;
@@ -106,6 +109,13 @@ namespace UE
 
 			TPromise< UObject* > RootObject;
 			FGraphEventRef RootObjectCompletionEvent;
+			
+			//If we cancel the tasks, we set this boolean to true
+			std::atomic<bool> bCancel;
+
+			void ReleaseTranslatorsSource();
+			void InitCancel();
+			void CancelAndWaitUntilDoneSynchronously();
 
 			void CleanUp();
 		};
@@ -163,21 +173,7 @@ public:
 	}
 
 	/** Return the interchange manager singleton.*/
-	static UInterchangeManager& GetInterchangeManager()
-	{
-		static TStrongObjectPtr<UInterchangeManager> InterchangeManager = nullptr;
-		if (!InterchangeManager.IsValid())
-		{
-			//We cannot create a TStrongObjectPtr outside of the main thread, we also need a valid Transient package
-			check(IsInGameThread() && GetTransientPackage());
-			InterchangeManager = TStrongObjectPtr<UInterchangeManager>(NewObject<UInterchangeManager>(GetTransientPackage(), NAME_None, EObjectFlags::RF_NoFlags));
-		}
-		//When we get here we should be valid
-		check(InterchangeManager.IsValid());
-
-		return *(InterchangeManager.Get());
-	}
-
+	static UInterchangeManager& GetInterchangeManager();
 
 	/** delegate type fired when new assets have been imported. Note: InCreatedObject can be NULL if import failed. Params: UFactory* InFactory, UObject* InCreatedObject */
 	DECLARE_MULTICAST_DELEGATE_OneParam(FInterchangeOnAssetPostImport, UObject*);
@@ -312,6 +308,13 @@ public:
 	 */
 	UInterchangeTranslatorBase* GetTranslatorForSourceData(const UInterchangeSourceData* SourceData) const;
 
+	/**
+	 * Return false if the Interchange is not active (importing or exporting).
+	 * If the interchange is active it will display a notification to let the user know he can cancel the asynchronous import/export
+	 * To be able to complete the operation he want to do. (The exit editor operation is calling this)
+	 */
+	bool WarnIfInterchangeIsActive();
+
 protected:
 
 	/** Return true if we can show some UI */
@@ -324,11 +327,30 @@ protected:
 	 */
 	void FindPipelineCandidate(TArray<UClass*>& PipelineCandidates);
 
+	/**
+	 * This function cancel all task and finish them has fast as possible.
+	 * We use this if the user cancel the work or if the editor is exiting.
+	 * @note - This is a blocking call until the tasks are completed.
+	 */
+	void CancelAllTasks();
+
+	/**
+	 * Same has CancelAllTasks, but will wait all task are done before exiting the function
+	 */
+	void CancelAllTasksSynchronously();
+
+	/**
+	 * If we set the mode to active we will setup the timer and add the thread that will block the GC.
+	 * If the we set the mode to inactive we will remove the timer and finish the thread that block the GC.
+	 */
+	void SetActiveMode(bool IsActive);
+
 private:
 	//By using pointer, there is no issue if the array get resize
 	TArray<TSharedPtr<UE::Interchange::FImportAsyncHelper, ESPMode::ThreadSafe> > ImportTasks;
 
 	TSharedPtr<FAsyncTaskNotification> Notification = nullptr;
+	FDelegateHandle NotificationTickHandle;
 
 	//The manager will create translator at every import, translator must be able to retrieve payload information when the factory ask for it.
 	//The translator stored has value is only use to know if we can use this type of translator.
@@ -347,5 +369,8 @@ private:
 	UPROPERTY()
 	TMap<const UClass*, UInterchangeWriterBase* > RegisteredWriters;
 
+	//If interchange is currently importing we have a timer to watch the cancel and we block GC 
+	FThreadSafeBool bIsActive = false;
+	FThread GcGuardThread;
 	friend class UE::Interchange::FScopedTranslator;
 };
