Fix build

diff --git a/Engine/Plugins/Runtime/Interchange/Source/Import/Private/InterchangeImportCommon.h b/Engine/Plugins/Runtime/Interchange/Source/Import/Private/InterchangeImportCommon.h
--- a/Engine/Plugins/Runtime/Interchange/Source/Import/Private/InterchangeImportCommon.h
+++ b/Engine/Plugins/Runtime/Interchange/Source/Import/Private/InterchangeImportCommon.h
@@ -1,4 +1,5 @@
 // Copyright Epic Games, Inc. All Rights Reserved.
+#pragma once
 
 #include "CoreMinimal.h"
 #include "InterchangeFactoryBase.h"
diff --git a/Engine/Plugins/Runtime/Interchange/Source/Import/Private/Material/InterchangeMaterialFactory.cpp b/Engine/Plugins/Runtime/Interchange/Source/Import/Private/Material/InterchangeMaterialFactory.cpp
--- a/Engine/Plugins/Runtime/Interchange/Source/Import/Private/Material/InterchangeMaterialFactory.cpp
+++ b/Engine/Plugins/Runtime/Interchange/Source/Import/Private/Material/InterchangeMaterialFactory.cpp
@@ -285,6 +285,9 @@ void UInterchangeMaterialFactory::PostImportGameThreadCallback(const FPostImport
 {
 	check(IsInGameThread());
 	Super::PostImportGameThreadCallback(Arguments);
+
+	//TODO make sure this work at runtime
+#if WITH_EDITORONLY_DATA
 	if (ensure(Arguments.ImportedObject && Arguments.SourceData))
 	{
 		//We must call the Update of the asset source file in the main thread because UAssetImportData::Update execute some delegate we do not control
@@ -297,5 +300,5 @@ void UInterchangeMaterialFactory::PostImportGameThreadCallback(const FPostImport
 																										  , Arguments.NodeContainer);
 		UE::Interchange::FFactoryCommon::UpdateImportAssetData(UpdateImportAssetDataParameters);
 	}
-	return;
+#endif
 }
\ No newline at end of file
diff --git a/Engine/Plugins/Runtime/Interchange/Source/Import/Private/Texture/InterchangeTextureFactory.cpp b/Engine/Plugins/Runtime/Interchange/Source/Import/Private/Texture/InterchangeTextureFactory.cpp
--- a/Engine/Plugins/Runtime/Interchange/Source/Import/Private/Texture/InterchangeTextureFactory.cpp
+++ b/Engine/Plugins/Runtime/Interchange/Source/Import/Private/Texture/InterchangeTextureFactory.cpp
@@ -207,6 +207,8 @@ void UInterchangeTextureFactory::PostImportGameThreadCallback(const FPostImportG
 {
 	check(IsInGameThread());
 	Super::PostImportGameThreadCallback(Arguments);
+	//TODO make sure this work at runtime
+#if WITH_EDITORONLY_DATA
 	if (ensure(Arguments.ImportedObject && Arguments.SourceData))
 	{
 		//We must call the Update of the asset source file in the main thread because UAssetImportData::Update execute some delegate we do not control
@@ -219,5 +221,5 @@ void UInterchangeTextureFactory::PostImportGameThreadCallback(const FPostImportG
 																						 , Arguments.NodeContainer);
 		UE::Interchange::FFactoryCommon::UpdateImportAssetData(UpdateImportAssetDataParameters);
 	}
-	return;
+#endif
 }
\ No newline at end of file
diff --git a/Engine/Source/Runtime/Interchange/Engine/Public/InterchangeAssetImportData.h b/Engine/Source/Runtime/Interchange/Engine/Public/InterchangeAssetImportData.h
--- a/Engine/Source/Runtime/Interchange/Engine/Public/InterchangeAssetImportData.h
+++ b/Engine/Source/Runtime/Interchange/Engine/Public/InterchangeAssetImportData.h
@@ -25,7 +25,12 @@ public:
 	UFUNCTION(BlueprintCallable, Category = "Interchange | AssetImportData")
 	FString ScriptGetFirstFilename() const
 	{
+		//TODO make sure this work at runtime
+#if WITH_EDITORONLY_DATA
 		return GetFirstFilename();
+#else
+		return FString();
+#endif
 	}
 
 	/**
@@ -34,7 +39,12 @@ public:
 	UFUNCTION(BlueprintCallable, Category = "Interchange | AssetImportData")
 	TArray<FString> ScriptExtractFilenames() const
 	{
+		//TODO make sure this work at runtime
+#if WITH_EDITORONLY_DATA
 		return ExtractFilenames();
+#else
+		return TArray<FString>();
+#endif
 	}
 
 	/**
@@ -44,7 +54,10 @@ public:
 	TArray<FString> ScriptExtractDisplayLabels() const
 	{
 		TArray<FString> TempDisplayLabels;
+		//TODO make sure this work at runtime
+#if WITH_EDITORONLY_DATA
 		ExtractDisplayLabels(TempDisplayLabels);
+#endif
 		return TempDisplayLabels;
 	}
 	
