Guard against crash while importing some datasmith scenes

diff --git a/Engine/Source/Runtime/Interchange/Engine/Private/InterchangeManager.cpp b/Engine/Source/Runtime/Interchange/Engine/Private/InterchangeManager.cpp
--- a/Engine/Source/Runtime/Interchange/Engine/Private/InterchangeManager.cpp
+++ b/Engine/Source/Runtime/Interchange/Engine/Private/InterchangeManager.cpp
@@ -96,7 +96,10 @@ void UE::Interchange::FImportAsyncHelper::ReleaseTranslatorsSource()
 {
 	for (UInterchangeTranslatorBase* BaseTranslator : Translators)
 	{
-		BaseTranslator->ReleaseSource();
+		if (BaseTranslator)
+		{
+			BaseTranslator->ReleaseSource();
+		}
 	}
 }
 
