Create a PostImportPipeline stage
Also add a way to trace the tasks in the logs

diff --git a/Engine/Plugins/Enterprise/DatasmithImporter/Source/DatasmithImporter/Private/DatasmithTextureImporter.cpp b/Engine/Plugins/Enterprise/DatasmithImporter/Source/DatasmithImporter/Private/DatasmithTextureImporter.cpp
--- a/Engine/Plugins/Enterprise/DatasmithImporter/Source/DatasmithImporter/Private/DatasmithTextureImporter.cpp
+++ b/Engine/Plugins/Enterprise/DatasmithImporter/Source/DatasmithImporter/Private/DatasmithTextureImporter.cpp
@@ -382,7 +382,7 @@ UE::Interchange::FAsyncImportResult FDatasmithTextureImporter::CreateTextureAsyn
 	return FutureTexture;
 }
 
-bool UDatasmithTexturePipeline::ExecuteImportPipeline(UInterchangeBaseNodeContainer* BaseNodeContainer)
+bool UDatasmithTexturePipeline::ExecutePreImportPipeline(UInterchangeBaseNodeContainer* BaseNodeContainer)
 {
 	if ( !TextureElement.IsValid() )
 	{
diff --git a/Engine/Plugins/Enterprise/DatasmithImporter/Source/DatasmithImporter/Public/DatasmithTextureImporter.h b/Engine/Plugins/Enterprise/DatasmithImporter/Source/DatasmithImporter/Public/DatasmithTextureImporter.h
--- a/Engine/Plugins/Enterprise/DatasmithImporter/Source/DatasmithImporter/Public/DatasmithTextureImporter.h
+++ b/Engine/Plugins/Enterprise/DatasmithImporter/Source/DatasmithImporter/Public/DatasmithTextureImporter.h
@@ -52,7 +52,7 @@ class UDatasmithTexturePipeline : public UInterchangePipelineBase
 	GENERATED_BODY()
 
 public:
-	virtual bool ExecuteImportPipeline(UInterchangeBaseNodeContainer* BaseNodeContainer) override;
+	virtual bool ExecutePreImportPipeline(UInterchangeBaseNodeContainer* BaseNodeContainer) override;
 
 	TSharedPtr< IDatasmithTextureElement > TextureElement;
 };
\ No newline at end of file
diff --git a/Engine/Source/Runtime/Interchange/Core/Public/InterchangePipelineBase.h b/Engine/Source/Runtime/Interchange/Core/Public/InterchangePipelineBase.h
--- a/Engine/Source/Runtime/Interchange/Core/Public/InterchangePipelineBase.h
+++ b/Engine/Source/Runtime/Interchange/Core/Public/InterchangePipelineBase.h
@@ -10,6 +10,14 @@
 
 #include "InterchangePipelineBase.generated.h"
 
+UENUM(BlueprintType)
+enum class EInterchangePipelineTask : uint8
+{
+	PreFactoryImport,
+	PostFactoryImport,
+	Export
+};
+
 UCLASS(BlueprintType, Blueprintable)
 class INTERCHANGECORE_API UInterchangePipelineBase : public UObject
 {
@@ -18,16 +26,29 @@ class INTERCHANGECORE_API UInterchangePipelineBase : public UObject
 public:
 
 	/**
-	 * Non virtual helper to allow blueprint to implement event base function to implement a pipeline,
+	 * Non virtual helper to allow blueprint to implement event base function to implement a pre import pipeline,
+	 * the Interchange manager is calling this function not the virtual one that is call by the default implementation.
+	 */
+	UFUNCTION(BlueprintNativeEvent, BlueprintCallable, Category = "Interchange | Translator")
+	bool ScriptedExecutePreImportPipeline(UInterchangeBaseNodeContainer* BaseNodeContainer);
+	/** The default implementation (call if the blueprint do not have any implementation) will call the virtual ExecuteImportPipeline */
+	bool ScriptedExecutePreImportPipeline_Implementation(UInterchangeBaseNodeContainer* BaseNodeContainer)
+	{
+		//By default we call the virtual import pipeline execution
+		return ExecutePreImportPipeline(BaseNodeContainer);
+	}
+
+	/**
+	 * Non virtual helper to allow blueprint to implement event base function to implement a post import pipeline,
 	 * the Interchange manager is calling this function not the virtual one that is call by the default implementation.
 	 */
 	UFUNCTION(BlueprintNativeEvent, BlueprintCallable, Category = "Interchange | Translator")
-	bool ScriptedExecuteImportPipeline(UInterchangeBaseNodeContainer* BaseNodeContainer);
+	bool ScriptedExecutePostImportPipeline(const UInterchangeBaseNodeContainer* BaseNodeContainer, const FName& NodeKey, UObject* CreatedAsset);
 	/** The default implementation (call if the blueprint do not have any implementation) will call the virtual ExecuteImportPipeline */
-	bool ScriptedExecuteImportPipeline_Implementation(UInterchangeBaseNodeContainer* BaseNodeContainer)
+	bool ScriptedExecutePostImportPipeline_Implementation(const UInterchangeBaseNodeContainer* BaseNodeContainer, const FName& NodeKey, UObject* CreatedAsset)
 	{
 		//By default we call the virtual import pipeline execution
-		return ExecuteImportPipeline(BaseNodeContainer);
+		return ExecutePostImportPipeline(BaseNodeContainer, NodeKey, CreatedAsset);
 	}
 
 	/**
@@ -48,12 +69,12 @@ public:
 	 * the Interchange manager is calling this function not the virtual one that is call by the default implementation.
 	 */
 	UFUNCTION(BlueprintNativeEvent, BlueprintCallable, Category = "Interchange | Translator")
-	bool ScriptedCanExecuteOnAnyThread();
+	bool ScriptedCanExecuteOnAnyThread(EInterchangePipelineTask PipelineTask);
 
 	/** The default implementation (call if the blueprint do not have any implementation) will call the virtual CanExecuteAsync */
-	bool ScriptedCanExecuteOnAnyThread_Implementation()
+	bool ScriptedCanExecuteOnAnyThread_Implementation(EInterchangePipelineTask PipelineTask)
 	{
-		return CanExecuteOnAnyThread();
+		return CanExecuteOnAnyThread(PipelineTask);
 	}
 protected:
 
@@ -62,7 +83,17 @@ protected:
 	 * The interchange manager is not calling this function directly. It is calling the blueprint native event in case this object is a blueprint derive object.
 	 * By default the scripted implementation is calling this virtual pipeline.
 	 */
-	virtual bool ExecuteImportPipeline(UInterchangeBaseNodeContainer* BaseNodeContainer)
+	virtual bool ExecutePreImportPipeline(UInterchangeBaseNodeContainer* BaseNodeContainer)
+	{
+		return false;
+	}
+
+	/**
+	 * This function can read the node data and apply some change to the imported asset. This is call after the factory create the asset and configure the asset properties.
+	 * The interchange manager is not calling this function directly. It is calling the blueprint native event in case this object is a blueprint derive object.
+	 * By default the scripted implementation is calling this virtual pipeline.
+	 */
+	virtual bool ExecutePostImportPipeline(const UInterchangeBaseNodeContainer* BaseNodeContainer, const FName& NodeKey, UObject* CreatedAsset)
 	{
 		return false;
 	}
@@ -73,7 +104,7 @@ protected:
 	 * import process in same time.
 	 *
 	 */
-	virtual bool CanExecuteOnAnyThread()
+	virtual bool CanExecuteOnAnyThread(EInterchangePipelineTask PipelineTask)
 	{
 		return true;
 	}
diff --git a/Engine/Source/Runtime/Interchange/Core/Public/Nodes/InterchangeBaseNode.h b/Engine/Source/Runtime/Interchange/Core/Public/Nodes/InterchangeBaseNode.h
--- a/Engine/Source/Runtime/Interchange/Core/Public/Nodes/InterchangeBaseNode.h
+++ b/Engine/Source/Runtime/Interchange/Core/Public/Nodes/InterchangeBaseNode.h
@@ -132,17 +132,17 @@ bool FillCustom##AttributeName##FromAsset(UObject* Asset)											\
 #endif //#if WITH_ENGINE
 
 #define IMPLEMENT_NODE_ATTRIBUTE_GETTER(AttributeName, AttributeType)																					\
-	FString OperationName = GetTypeName() + TEXT(".Get##AttributeName");																				\
+	FString OperationName = GetTypeName() + TEXT(".Get" #AttributeName);																				\
 	return InterchangePrivateNodeBase::GetCustomAttribute<AttributeType>(Attributes, Macro_Custom##AttributeName##Key, OperationName, AttributeValue);
 
 #define IMPLEMENT_NODE_ATTRIBUTE_SETTER_NODELEGATE(AttributeName, AttributeType)															\
-	FString OperationName = GetTypeName() + TEXT(".Set##AttributeName");																				\
+	FString OperationName = GetTypeName() + TEXT(".Set" #AttributeName);																				\
 	return InterchangePrivateNodeBase::SetCustomAttribute<AttributeType>(Attributes, Macro_Custom##AttributeName##Key, OperationName, AttributeValue);
 
 #if WITH_ENGINE
 
 #define IMPLEMENT_NODE_ATTRIBUTE_SETTER(NodeClassName, AttributeName, AttributeType, AssetType)														\
-	FString OperationName = GetTypeName() + TEXT(".Set##AttributeName");																			\
+	FString OperationName = GetTypeName() + TEXT(".Set" #AttributeName);																			\
 	if(InterchangePrivateNodeBase::SetCustomAttribute<AttributeType>(Attributes, Macro_Custom##AttributeName##Key, OperationName, AttributeValue))	\
 	{																																				\
 		if(bAddApplyDelegate)																														\
@@ -159,7 +159,7 @@ bool FillCustom##AttributeName##FromAsset(UObject* Asset)											\
 #else //#if WITH_ENGINE
 
 #define IMPLEMENT_NODE_ATTRIBUTE_SETTER(NodeClassName, AttributeName, AttributeType, AssetType)															\
-	FString OperationName = GetTypeName() + TEXT(".Set##AttributeName");																				\
+	FString OperationName = GetTypeName() + TEXT(".Set" #AttributeName);																				\
 	return InterchangePrivateNodeBase::SetCustomAttribute<AttributeType>(Attributes, Macro_Custom##AttributeName##Key, OperationName, AttributeValue);
 
 #endif //#if WITH_ENGINE
diff --git a/Engine/Source/Runtime/Interchange/Engine/Private/InterchangeEngineLogPrivate.h b/Engine/Source/Runtime/Interchange/Engine/Private/InterchangeEngineLogPrivate.h
--- a/Engine/Source/Runtime/Interchange/Engine/Private/InterchangeEngineLogPrivate.h
+++ b/Engine/Source/Runtime/Interchange/Engine/Private/InterchangeEngineLogPrivate.h
@@ -5,3 +5,53 @@
 #include "CoreMinimal.h"
 
 DECLARE_LOG_CATEGORY_EXTERN(LogInterchangeEngine, Log, All);
+
+#define INTERCHANGE_TRACE_ASYNCHRONOUS_TASK_ENABLED 0
+
+#if INTERCHANGE_TRACE_ASYNCHRONOUS_TASK_ENABLED
+
+class FScopeInterchangeTraceAsyncTask
+{
+public:
+	FScopeInterchangeTraceAsyncTask(FString InTaskName)
+		: TaskName(InTaskName)
+	{
+		StartTime = FPlatformTime::Seconds();
+		double IntegerPart;
+		double FractionalPart = FMath::Modf(StartTime, &IntegerPart);
+		TaskID = FractionalPart * 1000000000;
+		FString MessageStr = FString::FromInt(TaskID)
+			+ TEXT(" ")
+			+ TaskName
+			+ TEXT(" Execute Start on ")
+			+ (IsInGameThread() ? TraceAsyncGameThreadStr : TraceAsyncBackgroundThreadStr);
+
+		UE_LOG(LogInterchangeEngine, Warning, TEXT("%s"), *MessageStr);
+	}
+
+	~FScopeInterchangeTraceAsyncTask()
+	{
+		const int32 DeltaTime = (FPlatformTime::Seconds() - StartTime)*1000;
+		FString MessageStr = FString::FromInt(TaskID)
+			+ TEXT(" ")
+			+ TaskName 
+			+ TEXT(" Execute End on ")
+			+ (IsInGameThread() ? TraceAsyncGameThreadStr : TraceAsyncBackgroundThreadStr)
+			+ TEXT(" Time: ")
+			+ FString::FromInt(DeltaTime)
+			+ TEXT("ms");
+
+		UE_LOG(LogInterchangeEngine, Warning, TEXT("%s"), *MessageStr);
+	}
+private:
+	double StartTime = 0;
+	int32 TaskID = 0;
+	FString TaskName;
+	const FString TraceAsyncGameThreadStr = TEXT("Game Thread");
+	const FString TraceAsyncBackgroundThreadStr = TEXT("Background Thread");
+};
+
+#define INTERCHANGE_TRACE_ASYNCHRONOUS_TASK(TaskName) \
+FScopeInterchangeTraceAsyncTask ScopeTraceAsynTask(TEXT(#TaskName));
+
+#endif //INTERCHANGE_TRACE_ASYNCHRONOUS_TASK_ENABLED
\ No newline at end of file
diff --git a/Engine/Source/Runtime/Interchange/Engine/Private/InterchangeManager.cpp b/Engine/Source/Runtime/Interchange/Engine/Private/InterchangeManager.cpp
--- a/Engine/Source/Runtime/Interchange/Engine/Private/InterchangeManager.cpp
+++ b/Engine/Source/Runtime/Interchange/Engine/Private/InterchangeManager.cpp
@@ -119,18 +119,26 @@ void UE::Interchange::FImportAsyncHelper::CancelAndWaitUntilDoneSynchronously()
 	{
 		TasksToComplete.Append(TranslatorTasks);
 	}
-	if (PipelineTasks.Num())
+	if (PipelinePreImportTasks.Num())
 	{
-		TasksToComplete.Append(PipelineTasks);
+		TasksToComplete.Append(PipelinePreImportTasks);
 	}
 	if (ParsingTask.GetReference())
 	{
 		TasksToComplete.Add(ParsingTask);
 	}
+	if (CreatePackageTasks.Num())
+	{
+		TasksToComplete.Append(CreatePackageTasks);
+	}
 	if (CreateAssetTasks.Num())
 	{
 		TasksToComplete.Append(CreateAssetTasks);
 	}
+	if (PipelinePostImportTasks.Num())
+	{
+		TasksToComplete.Append(PipelinePostImportTasks);
+	}
 	if (CompletionTask.GetReference())
 	{
 		//Completion task will make sure any created asset before canceling will be mark for delete
@@ -420,9 +428,9 @@ UE::Interchange::FAsyncImportResult UInterchangeManager::ImportAssetAsync(const
 		// Stack all pipelines, for this import proto. TODO: We need to be able to control which pipeline we use for the import
 		// It can be set in the project settings and can also be set by a UI where the user create the pipeline stack he want.
 		// This should be a list of available pipelines that can be drop into a stack where you can control the order.
-		for (int32 GraphPipelineNumber = 0; GraphPipelineNumber < PipelineCandidates.Num(); ++GraphPipelineNumber)
+		for (int32 GraphPipelineIndex = 0; GraphPipelineIndex < PipelineCandidates.Num(); ++GraphPipelineIndex)
 		{
-			UInterchangePipelineBase* GeneratedPipeline = NewObject<UInterchangePipelineBase>(GetTransientPackage(), PipelineCandidates[GraphPipelineNumber], NAME_None, RF_NoFlags);
+			UInterchangePipelineBase* GeneratedPipeline = NewObject<UInterchangePipelineBase>(GetTransientPackage(), PipelineCandidates[GraphPipelineIndex], NAME_None, RF_NoFlags);
 			AsyncHelper->Pipelines.Add(GeneratedPipeline);
 		}
 	}
@@ -446,13 +454,13 @@ UE::Interchange::FAsyncImportResult UInterchangeManager::ImportAssetAsync(const
 		UInterchangePipelineBase* GraphPipeline = AsyncHelper->Pipelines[GraphPipelineIndex];
 		TWeakObjectPtr<UInterchangePipelineBase> WeakPipelinePtr = GraphPipeline;
 		int32 GraphPipelineTaskIndex = INDEX_NONE;
-		GraphPipelineTaskIndex = AsyncHelper->PipelineTasks.Add(TGraphTask<UE::Interchange::FTaskPipeline>::CreateTask(&PipelinePrerequistes).ConstructAndDispatchWhenReady(WeakPipelinePtr, WeakAsyncHelper));
+		GraphPipelineTaskIndex = AsyncHelper->PipelinePreImportTasks.Add(TGraphTask<UE::Interchange::FTaskPipelinePreImport>::CreateTask(&PipelinePrerequistes).ConstructAndDispatchWhenReady(WeakPipelinePtr, WeakAsyncHelper));
 		//Ensure we run the pipeline in the same order we create the task, since pipeline modify the node container, its important that its not process in parallel, Adding the one we start to the prerequisites
 		//is the way to go here
-		PipelinePrerequistes.Add(AsyncHelper->PipelineTasks[GraphPipelineTaskIndex]);
+		PipelinePrerequistes.Add(AsyncHelper->PipelinePreImportTasks[GraphPipelineTaskIndex]);
 
 		//Add pipeline to the graph parsing prerequisites
-		GraphParsingPrerequistes.Add(AsyncHelper->PipelineTasks[GraphPipelineTaskIndex]);
+		GraphParsingPrerequistes.Add(AsyncHelper->PipelinePreImportTasks[GraphPipelineTaskIndex]);
 	}
 
 	if (GraphParsingPrerequistes.Num() > 0)
diff --git a/Engine/Source/Runtime/Interchange/Engine/Private/Tasks/InterchangeTaskCompletion.cpp b/Engine/Source/Runtime/Interchange/Engine/Private/Tasks/InterchangeTaskCompletion.cpp
--- a/Engine/Source/Runtime/Interchange/Engine/Private/Tasks/InterchangeTaskCompletion.cpp
+++ b/Engine/Source/Runtime/Interchange/Engine/Private/Tasks/InterchangeTaskCompletion.cpp
@@ -4,6 +4,7 @@
 #include "AssetRegistryModule.h"
 #include "Async/TaskGraphInterfaces.h"
 #include "CoreMinimal.h"
+#include "InterchangeEngineLogPrivate.h"
 #include "InterchangeFactoryBase.h"
 #include "InterchangeManager.h"
 #include "Stats/Stats.h"
@@ -14,6 +15,9 @@
 
 void UE::Interchange::FTaskCompletion::DoTask(ENamedThreads::Type CurrentThread, const FGraphEventRef& MyCompletionGraphEvent)
 {
+#if INTERCHANGE_TRACE_ASYNCHRONOUS_TASK_ENABLED
+	INTERCHANGE_TRACE_ASYNCHRONOUS_TASK(Completion)
+#endif
 	TSharedPtr<FImportAsyncHelper, ESPMode::ThreadSafe> AsyncHelper = WeakAsyncHelper.Pin();
 	check(AsyncHelper.IsValid());
 
diff --git a/Engine/Source/Runtime/Interchange/Engine/Private/Tasks/InterchangeTaskCreateAsset.cpp b/Engine/Source/Runtime/Interchange/Engine/Private/Tasks/InterchangeTaskCreateAsset.cpp
--- a/Engine/Source/Runtime/Interchange/Engine/Private/Tasks/InterchangeTaskCreateAsset.cpp
+++ b/Engine/Source/Runtime/Interchange/Engine/Private/Tasks/InterchangeTaskCreateAsset.cpp
@@ -51,6 +51,9 @@ namespace UE
 
 void UE::Interchange::FTaskCreatePackage::DoTask(ENamedThreads::Type CurrentThread, const FGraphEventRef& MyCompletionGraphEvent)
 {
+#if INTERCHANGE_TRACE_ASYNCHRONOUS_TASK_ENABLED
+	INTERCHANGE_TRACE_ASYNCHRONOUS_TASK(CreatePackage)
+#endif
 	TSharedPtr<UE::Interchange::FImportAsyncHelper, ESPMode::ThreadSafe> AsyncHelper = WeakAsyncHelper.Pin();
 	check(AsyncHelper.IsValid());
 
@@ -129,6 +132,9 @@ void UE::Interchange::FTaskCreatePackage::DoTask(ENamedThreads::Type CurrentThre
 
 void UE::Interchange::FTaskCreateAsset::DoTask(ENamedThreads::Type CurrentThread, const FGraphEventRef& MyCompletionGraphEvent)
 {
+#if INTERCHANGE_TRACE_ASYNCHRONOUS_TASK_ENABLED
+	INTERCHANGE_TRACE_ASYNCHRONOUS_TASK(CreateAsset)
+#endif
 	TSharedPtr<UE::Interchange::FImportAsyncHelper, ESPMode::ThreadSafe> AsyncHelper = WeakAsyncHelper.Pin();
 	check(WeakAsyncHelper.IsValid());
 
diff --git a/Engine/Source/Runtime/Interchange/Engine/Private/Tasks/InterchangeTaskParsing.cpp b/Engine/Source/Runtime/Interchange/Engine/Private/Tasks/InterchangeTaskParsing.cpp
--- a/Engine/Source/Runtime/Interchange/Engine/Private/Tasks/InterchangeTaskParsing.cpp
+++ b/Engine/Source/Runtime/Interchange/Engine/Private/Tasks/InterchangeTaskParsing.cpp
@@ -3,11 +3,13 @@
 
 #include "Async/TaskGraphInterfaces.h"
 #include "CoreMinimal.h"
+#include "InterchangeEngineLogPrivate.h"
 #include "InterchangeFactoryBase.h"
 #include "InterchangeManager.h"
 #include "InterchangeSourceData.h"
 #include "InterchangeTaskCompletion.h"
 #include "InterchangeTaskCreateAsset.h"
+#include "InterchangeTaskPipeline.h"
 #include "InterchangeTranslatorBase.h"
 #include "Misc/Paths.h"
 #include "PackageUtils/PackageUtils.h"
@@ -22,6 +24,9 @@
 
 void UE::Interchange::FTaskParsing::DoTask(ENamedThreads::Type CurrentThread, const FGraphEventRef& MyCompletionGraphEvent)
 {
+#if INTERCHANGE_TRACE_ASYNCHRONOUS_TASK_ENABLED
+	INTERCHANGE_TRACE_ASYNCHRONOUS_TASK(ParsingGraph)
+#endif
 	TSharedPtr<FImportAsyncHelper, ESPMode::ThreadSafe> AsyncHelper = WeakAsyncHelper.Pin();
 	check(AsyncHelper.IsValid());
 
@@ -107,10 +112,27 @@ void UE::Interchange::FTaskParsing::DoTask(ENamedThreads::Type CurrentThread, co
 
 		//Add create package task has a prerequisite of FTaskCreateAsset. Create package task is a game thread task
 		FGraphEventArray CreatePackagePrerequistes;
-		CreatePackagePrerequistes.Add(TGraphTask<FTaskCreatePackage>::CreateTask(&(TaskDatas[TaskIndex].Prerequistes)).ConstructAndDispatchWhenReady(PackageBasePath, TaskDatas[TaskIndex].SourceIndex, WeakAsyncHelper, TaskDatas[TaskIndex].Node, TaskDatas[TaskIndex].Factory));
+		int32 CreatePackageTaskIndex = AsyncHelper->CreatePackageTasks.Add(TGraphTask<FTaskCreatePackage>::CreateTask(&(TaskDatas[TaskIndex].Prerequistes)).ConstructAndDispatchWhenReady(PackageBasePath, TaskDatas[TaskIndex].SourceIndex, WeakAsyncHelper, TaskDatas[TaskIndex].Node, TaskDatas[TaskIndex].Factory));
+		CreatePackagePrerequistes.Add(AsyncHelper->CreatePackageTasks[CreatePackageTaskIndex]);
 
+		FGraphEventArray PostPipelinePrerequistes;
 		int32 CreateTaskIndex = AsyncHelper->CreateAssetTasks.Add(TGraphTask<FTaskCreateAsset>::CreateTask(&(CreatePackagePrerequistes)).ConstructAndDispatchWhenReady(PackageBasePath, TaskDatas[TaskIndex].SourceIndex, WeakAsyncHelper, TaskDatas[TaskIndex].Node, TaskDatas[TaskIndex].Factory));
+		PostPipelinePrerequistes.Add(AsyncHelper->CreateAssetTasks[CreateTaskIndex]);
+		
 		TaskDatas[TaskIndex].GraphEventRef = AsyncHelper->CreateAssetTasks[CreateTaskIndex];
+
+		for (int32 GraphPipelineIndex = 0; GraphPipelineIndex < AsyncHelper->Pipelines.Num(); ++GraphPipelineIndex)
+		{
+			int32 GraphPipelineTaskIndex = INDEX_NONE;
+			GraphPipelineTaskIndex = AsyncHelper->PipelinePostImportTasks.Add(TGraphTask<FTaskPipelinePostImport>::CreateTask(&(PostPipelinePrerequistes)).ConstructAndDispatchWhenReady(TaskDatas[TaskIndex].SourceIndex, GraphPipelineIndex, WeakAsyncHelper));
+			//Ensure we run the pipeline in the same order we create the task, since pipeline modify the node container, its important that its not process in parallel, Adding the one we start to the prerequisites
+			//is the way to go here
+			PostPipelinePrerequistes.Add(AsyncHelper->PipelinePostImportTasks[GraphPipelineTaskIndex]);
+
+			//Override the completion prerequisite with the latest post import pipeline task
+			TaskDatas[TaskIndex].GraphEventRef = AsyncHelper->PipelinePostImportTasks[GraphPipelineTaskIndex];
+		}
+
 		CompletionPrerequistes.Add(TaskDatas[TaskIndex].GraphEventRef);
 	}
 
diff --git a/Engine/Source/Runtime/Interchange/Engine/Private/Tasks/InterchangeTaskPipeline.cpp b/Engine/Source/Runtime/Interchange/Engine/Private/Tasks/InterchangeTaskPipeline.cpp
--- a/Engine/Source/Runtime/Interchange/Engine/Private/Tasks/InterchangeTaskPipeline.cpp
+++ b/Engine/Source/Runtime/Interchange/Engine/Private/Tasks/InterchangeTaskPipeline.cpp
@@ -3,6 +3,7 @@
 
 #include "Async/TaskGraphInterfaces.h"
 #include "CoreMinimal.h"
+#include "InterchangeEngineLogPrivate.h"
 #include "InterchangeManager.h"
 #include "InterchangePipelineBase.h"
 #include "Stats/Stats.h"
@@ -10,8 +11,13 @@
 #include "UObject/WeakObjectPtrTemplates.h"
 #include "Nodes/InterchangeBaseNodeContainer.h"
 
-void UE::Interchange::FTaskPipeline::DoTask(ENamedThreads::Type CurrentThread, const FGraphEventRef& MyCompletionGraphEvent)
+
+
+void UE::Interchange::FTaskPipelinePreImport::DoTask(ENamedThreads::Type CurrentThread, const FGraphEventRef& MyCompletionGraphEvent)
 {
+#if INTERCHANGE_TRACE_ASYNCHRONOUS_TASK_ENABLED
+	INTERCHANGE_TRACE_ASYNCHRONOUS_TASK(PipelinePreImport)
+#endif
 	TSharedPtr<FImportAsyncHelper, ESPMode::ThreadSafe> AsyncHelper = WeakAsyncHelper.Pin();
 	check(AsyncHelper.IsValid());
 
@@ -27,8 +33,62 @@ void UE::Interchange::FTaskPipeline::DoTask(ENamedThreads::Type CurrentThread, c
 
 			if (ensure(AsyncHelper->BaseNodeContainers[GraphIndex].IsValid()))
 			{
-				Pipeline->ScriptedExecuteImportPipeline(AsyncHelper->BaseNodeContainers[GraphIndex].Get());
+				Pipeline->ScriptedExecutePreImportPipeline(AsyncHelper->BaseNodeContainers[GraphIndex].Get());
 			}
 		}
 	}
 }
+
+void UE::Interchange::FTaskPipelinePostImport::DoTask(ENamedThreads::Type CurrentThread, const FGraphEventRef& MyCompletionGraphEvent)
+{
+#if INTERCHANGE_TRACE_ASYNCHRONOUS_TASK_ENABLED
+	INTERCHANGE_TRACE_ASYNCHRONOUS_TASK(PipelinePostImport)
+#endif
+	TSharedPtr<FImportAsyncHelper, ESPMode::ThreadSafe> AsyncHelper = WeakAsyncHelper.Pin();
+	if (!ensure(AsyncHelper.IsValid()) || AsyncHelper->bCancel)
+	{
+		return;
+	}
+
+	if (!ensure(AsyncHelper->Pipelines.IsValidIndex(PipelineIndex)) || !ensure(AsyncHelper->BaseNodeContainers.IsValidIndex(SourceIndex)))
+	{
+		return;
+	}
+	UInterchangePipelineBase* PipelineBase = AsyncHelper->Pipelines[PipelineIndex];
+	TArray<FName> NodeUniqueIDs;
+	TArray<UObject*> ImportAssets;
+	//Create a lock scope to read the Imported asset infos map
+	{
+		FScopeLock Lock(&AsyncHelper->ImportedAssetsPerSourceIndexLock);
+		if (AsyncHelper->ImportedAssetsPerSourceIndex.Contains(SourceIndex))
+		{
+			TArray<UE::Interchange::FImportAsyncHelper::FImportedAssetInfo>& ImportedInfos = AsyncHelper->ImportedAssetsPerSourceIndex.FindChecked(SourceIndex);
+			NodeUniqueIDs.Reserve(ImportedInfos.Num());
+			ImportAssets.Reserve(ImportedInfos.Num());
+			for (UE::Interchange::FImportAsyncHelper::FImportedAssetInfo& ImportedInfo : ImportedInfos)
+			{
+				NodeUniqueIDs.Add(FName(*ImportedInfo.NodeUniqueId));
+				ImportAssets.Add(ImportedInfo.ImportAsset);
+			}
+		}
+	}
+
+	if (!ensure(NodeUniqueIDs.Num() == ImportAssets.Num()))
+	{
+		//We do not execute the script if we cannot give proper parameter
+		return;
+	}
+
+	//Get the Container from the async helper
+	UInterchangeBaseNodeContainer* NodeContainer = AsyncHelper->BaseNodeContainers[SourceIndex].Get();
+	if (!ensure(NodeContainer))
+	{
+		return;
+	}
+	UInterchangePipelineBase* Pipeline = AsyncHelper->Pipelines[PipelineIndex];
+	//Call the pipeline outside of the lock, we do this in case the pipeline take a long time. We call it for each asset created by this import
+	for (int32 AssetIndex = 0; AssetIndex < ImportAssets.Num(); ++AssetIndex)
+	{
+		Pipeline->ScriptedExecutePostImportPipeline(NodeContainer, NodeUniqueIDs[AssetIndex], ImportAssets[AssetIndex]);
+	}
+}
diff --git a/Engine/Source/Runtime/Interchange/Engine/Private/Tasks/InterchangeTaskPipeline.h b/Engine/Source/Runtime/Interchange/Engine/Private/Tasks/InterchangeTaskPipeline.h
--- a/Engine/Source/Runtime/Interchange/Engine/Private/Tasks/InterchangeTaskPipeline.h
+++ b/Engine/Source/Runtime/Interchange/Engine/Private/Tasks/InterchangeTaskPipeline.h
@@ -14,13 +14,13 @@ namespace UE
 	namespace Interchange
 	{
 
-		class FTaskPipeline
+		class FTaskPipelinePreImport
 		{
 		private:
 			TWeakObjectPtr<UInterchangePipelineBase> PipelineBase;
 			TWeakPtr<FImportAsyncHelper, ESPMode::ThreadSafe> WeakAsyncHelper;
 		public:
-			FTaskPipeline(TWeakObjectPtr<UInterchangePipelineBase> InPipelineBase, TWeakPtr<FImportAsyncHelper, ESPMode::ThreadSafe> InAsyncHelper)
+			FTaskPipelinePreImport(TWeakObjectPtr<UInterchangePipelineBase> InPipelineBase, TWeakPtr<FImportAsyncHelper, ESPMode::ThreadSafe> InAsyncHelper)
 				: PipelineBase(InPipelineBase)
 				, WeakAsyncHelper(InAsyncHelper)
 			{
@@ -28,7 +28,7 @@ namespace UE
 
 			ENamedThreads::Type GetDesiredThread()
 			{
-				return PipelineBase.Get()->ScriptedCanExecuteOnAnyThread() ? ENamedThreads::AnyBackgroundThreadNormalTask : ENamedThreads::GameThread;
+				return PipelineBase.Get()->ScriptedCanExecuteOnAnyThread(EInterchangePipelineTask::PreFactoryImport) ? ENamedThreads::AnyBackgroundThreadNormalTask : ENamedThreads::GameThread;
 			}
 
 			static ESubsequentsMode::Type GetSubsequentsMode()
@@ -38,7 +38,56 @@ namespace UE
 
 			FORCEINLINE TStatId GetStatId() const
 			{
-				RETURN_QUICK_DECLARE_CYCLE_STAT(FTaskPipeline, STATGROUP_TaskGraphTasks);
+				RETURN_QUICK_DECLARE_CYCLE_STAT(FTaskPipelinePreImport, STATGROUP_TaskGraphTasks);
+			}
+
+			void DoTask(ENamedThreads::Type CurrentThread, const FGraphEventRef& MyCompletionGraphEvent);
+		};
+
+		class FTaskPipelinePostImport
+		{
+		private:
+			int32 SourceIndex;
+			int32 PipelineIndex;
+			TWeakPtr<FImportAsyncHelper, ESPMode::ThreadSafe> WeakAsyncHelper;
+
+		public:
+			FTaskPipelinePostImport(int32 InSourceIndex, int32 InPipelineIndex, TWeakPtr<FImportAsyncHelper, ESPMode::ThreadSafe> InAsyncHelper)
+				: SourceIndex(InSourceIndex)
+				, PipelineIndex(InPipelineIndex)
+				, WeakAsyncHelper(InAsyncHelper)
+			{
+			}
+
+			ENamedThreads::Type GetDesiredThread()
+			{
+				TSharedPtr<FImportAsyncHelper, ESPMode::ThreadSafe> AsyncHelper = WeakAsyncHelper.Pin();
+				//Always use the game thread if there is some error
+				if (!ensure(AsyncHelper.IsValid()))
+				{
+					return ENamedThreads::GameThread;
+				}
+				if (!AsyncHelper->Pipelines.IsValidIndex(PipelineIndex))
+				{
+					return ENamedThreads::GameThread;
+				}
+				//Ask the pipeline implementation
+				if (AsyncHelper->Pipelines[PipelineIndex]->ScriptedCanExecuteOnAnyThread(EInterchangePipelineTask::PostFactoryImport))
+				{
+					return ENamedThreads::AnyBackgroundThreadNormalTask;
+				}
+
+				return ENamedThreads::GameThread;
+			}
+
+			static ESubsequentsMode::Type GetSubsequentsMode()
+			{
+				return ESubsequentsMode::TrackSubsequents;
+			}
+
+			FORCEINLINE TStatId GetStatId() const
+			{
+				RETURN_QUICK_DECLARE_CYCLE_STAT(FTaskPipelinePostImport, STATGROUP_TaskGraphTasks);
 			}
 
 			void DoTask(ENamedThreads::Type CurrentThread, const FGraphEventRef& MyCompletionGraphEvent);
diff --git a/Engine/Source/Runtime/Interchange/Engine/Private/Tasks/InterchangeTaskTranslator.cpp b/Engine/Source/Runtime/Interchange/Engine/Private/Tasks/InterchangeTaskTranslator.cpp
--- a/Engine/Source/Runtime/Interchange/Engine/Private/Tasks/InterchangeTaskTranslator.cpp
+++ b/Engine/Source/Runtime/Interchange/Engine/Private/Tasks/InterchangeTaskTranslator.cpp
@@ -3,6 +3,7 @@
 
 #include "Async/TaskGraphInterfaces.h"
 #include "CoreMinimal.h"
+#include "InterchangeEngineLogPrivate.h"
 #include "InterchangeManager.h"
 #include "InterchangeSourceData.h"
 #include "InterchangeTranslatorBase.h"
@@ -13,6 +14,9 @@
 
 void UE::Interchange::FTaskTranslator::DoTask(ENamedThreads::Type CurrentThread, const FGraphEventRef& MyCompletionGraphEvent)
 {
+#if INTERCHANGE_TRACE_ASYNCHRONOUS_TASK_ENABLED
+	INTERCHANGE_TRACE_ASYNCHRONOUS_TASK(Translator)
+#endif
 	TSharedPtr<FImportAsyncHelper, ESPMode::ThreadSafe> AsyncHelper = WeakAsyncHelper.Pin();
 	check(AsyncHelper.IsValid());
 
diff --git a/Engine/Source/Runtime/Interchange/Engine/Public/InterchangeManager.h b/Engine/Source/Runtime/Interchange/Engine/Public/InterchangeManager.h
--- a/Engine/Source/Runtime/Interchange/Engine/Public/InterchangeManager.h
+++ b/Engine/Source/Runtime/Interchange/Engine/Public/InterchangeManager.h
@@ -77,16 +77,21 @@ namespace UE
 			/* FGCObject interface */
 			virtual void AddReferencedObjects(FReferenceCollector& Collector) override;
 
+			//The following Arrays are per source data
 			TArray <TStrongObjectPtr<UInterchangeBaseNodeContainer>> BaseNodeContainers;
-
 			TArray<UInterchangeSourceData* > SourceDatas;
 			TArray<UInterchangeTranslatorBase* > Translators;
-			TArray<UInterchangePipelineBase* > Pipelines;
 			TArray<UInterchangeFactoryBase* > Factories;
 
+			//Pipelines array is not per source data 
+			TArray<UInterchangePipelineBase* > Pipelines;
+			
+
 			TArray<FGraphEventRef> TranslatorTasks;
-			TArray<FGraphEventRef> PipelineTasks;
+			TArray<FGraphEventRef> PipelinePreImportTasks;
+			TArray<FGraphEventRef> PipelinePostImportTasks;
 			FGraphEventRef ParsingTask;
+			TArray<FGraphEventRef> CreatePackageTasks;
 			TArray<FGraphEventRef> CreateAssetTasks;
 
 			FGraphEventRef CompletionTask;
