Add a Queue of tasks to prevent starting too much import task in same time.

diff --git a/Engine/Source/Runtime/Interchange/Engine/Private/InterchangeManager.cpp b/Engine/Source/Runtime/Interchange/Engine/Private/InterchangeManager.cpp
--- a/Engine/Source/Runtime/Interchange/Engine/Private/InterchangeManager.cpp
+++ b/Engine/Source/Runtime/Interchange/Engine/Private/InterchangeManager.cpp
@@ -361,6 +361,103 @@ bool UInterchangeManager::CanTranslateSourceData(const UInterchangeSourceData* S
 	return false;
 }
 
+void UInterchangeManager::StartQueuedTasks(bool bCancelAllTasks /*= false*/)
+{
+	if (!ensure(IsInGameThread()))
+	{
+		//Do not crash but we will not start any queued tasks if we are not in the game thread
+		return;
+	}
+
+	auto UpdateNotification = [this]()
+	{
+		if (Notification.IsValid())
+		{
+			int32 ImportTaskNumber = ImportTasks.Num() + QueueTaskCount;
+			FString ImportTaskNumberStr = TEXT(" (") + FString::FromInt(ImportTaskNumber) + TEXT(")");
+			Notification->SetProgressText(FText::FromString(ImportTaskNumberStr));
+		}
+		else
+		{
+			FText TitleText = NSLOCTEXT("Interchange", "Asynchronous_import_start", "Importing");
+			FAsyncTaskNotificationConfig NotificationConfig;
+			NotificationConfig.bIsHeadless = false;
+			NotificationConfig.bKeepOpenOnFailure = true;
+			NotificationConfig.TitleText = TitleText;
+			NotificationConfig.LogCategory = InternalInterchangePrivate::GetLogInterchangePtr();
+			NotificationConfig.bCanCancel.Set(true);
+			NotificationConfig.bKeepOpenOnFailure.Set(true);
+
+			Notification = MakeShared<FAsyncTaskNotification>(NotificationConfig);
+			Notification->SetNotificationState(FAsyncNotificationStateData(TitleText, FText::GetEmpty(), EAsyncTaskNotificationState::Pending));
+		}
+	};
+
+	while (!QueuedTasks.IsEmpty() && (ImportTasks.Num() < FTaskGraphInterface::Get().GetNumWorkerThreads() || bCancelAllTasks))
+	{
+		FQueuedTaskData QueuedTaskData;
+		if (QueuedTasks.Dequeue(QueuedTaskData))
+		{
+			QueueTaskCount = FMath::Clamp(QueueTaskCount-1, 0, MAX_int32);
+			check(QueuedTaskData.AsyncHelper.IsValid());
+
+			int32 AsyncHelperIndex = ImportTasks.Add(QueuedTaskData.AsyncHelper);
+			SetActiveMode(true);
+			//Update the asynchronous notification
+			UpdateNotification();
+
+			TWeakPtr<UE::Interchange::FImportAsyncHelper, ESPMode::ThreadSafe> WeakAsyncHelper = QueuedTaskData.AsyncHelper;
+			
+			if (bCancelAllTasks)
+			{
+				QueuedTaskData.AsyncHelper->InitCancel();
+			}
+
+			//Create/Start import tasks
+			FGraphEventArray PipelinePrerequistes;
+			check(QueuedTaskData.AsyncHelper->Translators.Num() == QueuedTaskData.AsyncHelper->SourceDatas.Num());
+			for (int32 SourceDataIndex = 0; SourceDataIndex < QueuedTaskData.AsyncHelper->SourceDatas.Num(); ++SourceDataIndex)
+			{
+				int32 TranslatorTaskIndex = QueuedTaskData.AsyncHelper->TranslatorTasks.Add(TGraphTask<UE::Interchange::FTaskTranslator>::CreateTask().ConstructAndDispatchWhenReady(SourceDataIndex, WeakAsyncHelper));
+				PipelinePrerequistes.Add(QueuedTaskData.AsyncHelper->TranslatorTasks[TranslatorTaskIndex]);
+			}
+
+			FGraphEventArray GraphParsingPrerequistes;
+			for (int32 GraphPipelineIndex = 0; GraphPipelineIndex < QueuedTaskData.AsyncHelper->Pipelines.Num(); ++GraphPipelineIndex)
+			{
+				UInterchangePipelineBase* GraphPipeline = QueuedTaskData.AsyncHelper->Pipelines[GraphPipelineIndex];
+				TWeakObjectPtr<UInterchangePipelineBase> WeakPipelinePtr = GraphPipeline;
+				int32 GraphPipelineTaskIndex = INDEX_NONE;
+				GraphPipelineTaskIndex = QueuedTaskData.AsyncHelper->PipelinePreImportTasks.Add(TGraphTask<UE::Interchange::FTaskPipelinePreImport>::CreateTask(&PipelinePrerequistes).ConstructAndDispatchWhenReady(WeakPipelinePtr, WeakAsyncHelper));
+				//Ensure we run the pipeline in the same order we create the task, since pipeline modify the node container, its important that its not process in parallel, Adding the one we start to the prerequisites
+				//is the way to go here
+				PipelinePrerequistes.Add(QueuedTaskData.AsyncHelper->PipelinePreImportTasks[GraphPipelineTaskIndex]);
+
+				//Add pipeline to the graph parsing prerequisites
+				GraphParsingPrerequistes.Add(QueuedTaskData.AsyncHelper->PipelinePreImportTasks[GraphPipelineTaskIndex]);
+			}
+
+			if (GraphParsingPrerequistes.Num() > 0)
+			{
+				QueuedTaskData.AsyncHelper->ParsingTask = TGraphTask<UE::Interchange::FTaskParsing>::CreateTask(&GraphParsingPrerequistes).ConstructAndDispatchWhenReady(this, QueuedTaskData.PackageBasePath, WeakAsyncHelper);
+			}
+			else
+			{
+				//Fallback on the translator pipeline prerequisites (translator must be done if there is no pipeline)
+				QueuedTaskData.AsyncHelper->ParsingTask = TGraphTask<UE::Interchange::FTaskParsing>::CreateTask(&PipelinePrerequistes).ConstructAndDispatchWhenReady(this, QueuedTaskData.PackageBasePath, WeakAsyncHelper);
+			}
+
+			//The graph parsing task will create the FCreateAssetTask that will run after them, the FAssetImportTask will call the appropriate Post asset import pipeline when the asset is completed
+		}
+	}
+
+	if (!QueuedTasks.IsEmpty())
+	{
+		//Make sure any task we add is count in the task to do, even if we cannot start it
+		UpdateNotification();
+	}
+}
+
 bool UInterchangeManager::ImportAsset(const FString& ContentPath, const UInterchangeSourceData* SourceData, const FImportAssetParameters& ImportAssetParameters)
 {
 	return ImportAssetAsync( ContentPath, SourceData, ImportAssetParameters ).IsValid();
@@ -368,6 +465,12 @@ bool UInterchangeManager::ImportAsset(const FString& ContentPath, const UInterch
 
 UE::Interchange::FAsyncImportResult UInterchangeManager::ImportAssetAsync(const FString& ContentPath, const UInterchangeSourceData* SourceData, const FImportAssetParameters& ImportAssetParameters)
 {
+	if (!ensure(IsInGameThread()))
+	{
+		//Import process can be start only in game thread
+		return UE::Interchange::FAsyncImportResult{};
+	}
+
 	FString PackageBasePath = ContentPath;
 	if(!ImportAssetParameters.ReimportAsset)
 	{
@@ -381,25 +484,9 @@ UE::Interchange::FAsyncImportResult UInterchangeManager::ImportAssetAsync(const
 	TaskData.bIsAutomated = ImportAssetParameters.bIsAutomated;
 	TaskData.ImportType = UE::Interchange::EImportType::ImportType_Asset;
 	TaskData.ReimportObject = ImportAssetParameters.ReimportAsset;
-	TWeakPtr<UE::Interchange::FImportAsyncHelper, ESPMode::ThreadSafe> WeakAsyncHelper = CreateAsyncHelper(TaskData);
-	TSharedPtr<UE::Interchange::FImportAsyncHelper, ESPMode::ThreadSafe> AsyncHelper = WeakAsyncHelper.Pin();
+	TSharedPtr<UE::Interchange::FImportAsyncHelper, ESPMode::ThreadSafe> AsyncHelper = CreateAsyncHelper(TaskData);
 	check(AsyncHelper.IsValid());
 
-	FText TitleText = NSLOCTEXT("Interchange", "Asynchronous_import_start", "Importing");
-	if(!Notification.IsValid())
-	{
-		FAsyncTaskNotificationConfig NotificationConfig;
-		NotificationConfig.bIsHeadless = false;
-		NotificationConfig.bKeepOpenOnFailure = true;
-		NotificationConfig.TitleText = TitleText;
-		NotificationConfig.LogCategory = InternalInterchangePrivate::GetLogInterchangePtr();
-		NotificationConfig.bCanCancel.Set(true);
-		NotificationConfig.bKeepOpenOnFailure.Set(true);
-
-		Notification = MakeShared<FAsyncTaskNotification>(NotificationConfig);
-		Notification->SetNotificationState(FAsyncNotificationStateData(TitleText, FText::GetEmpty(), EAsyncTaskNotificationState::Pending));
-	}
-	
 	//Create a duplicate of the source data, we need to be multithread safe so we copy it to control the life cycle. The async helper will hold it and delete it when the import task will be completed.
 	UInterchangeSourceData* DuplicateSourceData = Cast<UInterchangeSourceData>(StaticDuplicateObject(SourceData, GetTransientPackage()));
 	//Array of source data to build one graph per source
@@ -439,41 +526,13 @@ UE::Interchange::FAsyncImportResult UInterchangeManager::ImportAssetAsync(const
 		AsyncHelper->Pipelines.Add(ImportAssetParameters.OverridePipeline);
 	}
 
-	//Create/Start import tasks
-	FGraphEventArray PipelinePrerequistes;
-	check(AsyncHelper->Translators.Num() == AsyncHelper->SourceDatas.Num());
-	for (int32 SourceDataIndex = 0; SourceDataIndex < AsyncHelper->SourceDatas.Num(); ++SourceDataIndex)
-	{
-		int32 TranslatorTaskIndex = AsyncHelper->TranslatorTasks.Add(TGraphTask<UE::Interchange::FTaskTranslator>::CreateTask().ConstructAndDispatchWhenReady(SourceDataIndex, WeakAsyncHelper));
-		PipelinePrerequistes.Add(AsyncHelper->TranslatorTasks[TranslatorTaskIndex]);
-	}
-		
-	FGraphEventArray GraphParsingPrerequistes;
-	for(int32 GraphPipelineIndex = 0; GraphPipelineIndex < AsyncHelper->Pipelines.Num(); ++GraphPipelineIndex)
-	{
-		UInterchangePipelineBase* GraphPipeline = AsyncHelper->Pipelines[GraphPipelineIndex];
-		TWeakObjectPtr<UInterchangePipelineBase> WeakPipelinePtr = GraphPipeline;
-		int32 GraphPipelineTaskIndex = INDEX_NONE;
-		GraphPipelineTaskIndex = AsyncHelper->PipelinePreImportTasks.Add(TGraphTask<UE::Interchange::FTaskPipelinePreImport>::CreateTask(&PipelinePrerequistes).ConstructAndDispatchWhenReady(WeakPipelinePtr, WeakAsyncHelper));
-		//Ensure we run the pipeline in the same order we create the task, since pipeline modify the node container, its important that its not process in parallel, Adding the one we start to the prerequisites
-		//is the way to go here
-		PipelinePrerequistes.Add(AsyncHelper->PipelinePreImportTasks[GraphPipelineTaskIndex]);
-
-		//Add pipeline to the graph parsing prerequisites
-		GraphParsingPrerequistes.Add(AsyncHelper->PipelinePreImportTasks[GraphPipelineTaskIndex]);
-	}
-
-	if (GraphParsingPrerequistes.Num() > 0)
-	{
-		AsyncHelper->ParsingTask = TGraphTask<UE::Interchange::FTaskParsing>::CreateTask(&GraphParsingPrerequistes).ConstructAndDispatchWhenReady(this, PackageBasePath, WeakAsyncHelper);
-	}
-	else
-	{
-		//Fallback on the translator pipeline prerequisites (translator must be done if there is no pipeline)
-		AsyncHelper->ParsingTask = TGraphTask<UE::Interchange::FTaskParsing>::CreateTask(&PipelinePrerequistes).ConstructAndDispatchWhenReady(this, PackageBasePath, WeakAsyncHelper);
-	}
+	FQueuedTaskData QueuedTaskData;
+	QueuedTaskData.AsyncHelper = AsyncHelper;
+	QueuedTaskData.PackageBasePath = PackageBasePath;
+	QueuedTasks.Enqueue(QueuedTaskData);
+	QueueTaskCount = FMath::Clamp(QueueTaskCount + 1, 0, MAX_int32);
 
-	//The graph parsing task will create the FCreateAssetTask that will run after them, the FAssetImportTask will call the appropriate Post asset import pipeline when the asset is completed
+	StartQueuedTasks();
 
 	return UE::Interchange::FAsyncImportResult{ AsyncHelper->RootObject.GetFuture(), AsyncHelper->RootObjectCompletionEvent };
 }
@@ -503,22 +562,13 @@ UInterchangeSourceData* UInterchangeManager::CreateSourceData(const FString& InF
 	return SourceDataAsset;
 }
 
-TWeakPtr<UE::Interchange::FImportAsyncHelper, ESPMode::ThreadSafe> UInterchangeManager::CreateAsyncHelper(const UE::Interchange::FImportAsyncHelperData& Data)
+TSharedPtr<UE::Interchange::FImportAsyncHelper, ESPMode::ThreadSafe> UInterchangeManager::CreateAsyncHelper(const UE::Interchange::FImportAsyncHelperData& Data)
 {
 	TSharedPtr<UE::Interchange::FImportAsyncHelper, ESPMode::ThreadSafe> AsyncHelper = MakeShared<UE::Interchange::FImportAsyncHelper, ESPMode::ThreadSafe>();
 	//Copy the task data
 	AsyncHelper->TaskData = Data;
-	int32 AsyncHelperIndex = ImportTasks.Add(AsyncHelper);
-	SetActiveMode(true);
-	//Update the asynchronous notification
-	if (Notification.IsValid())
-	{
-		int32 ImportTaskNumber = ImportTasks.Num();
-		FString ImportTaskNumberStr = TEXT(" (") + FString::FromInt(ImportTaskNumber) + TEXT(")");
-		Notification->SetProgressText(FText::FromString(ImportTaskNumberStr));
-	}
 	
-	return ImportTasks[AsyncHelperIndex];
+	return AsyncHelper;
 }
 
 void UInterchangeManager::ReleaseAsyncHelper(TWeakPtr<UE::Interchange::FImportAsyncHelper, ESPMode::ThreadSafe> AsyncHelper)
@@ -550,6 +600,9 @@ void UInterchangeManager::ReleaseAsyncHelper(TWeakPtr<UE::Interchange::FImportAs
 	{
 		Notification->SetProgressText(FText::FromString(ImportTaskNumberStr));
 	}
+
+	//Start some task if there is some waiting
+	StartQueuedTasks();
 }
 
 UInterchangeTranslatorBase* UInterchangeManager::GetTranslatorForSourceData(const UInterchangeSourceData* SourceData) const
@@ -686,7 +739,12 @@ void UInterchangeManager::CancelAllTasks()
 {
 	check(IsInGameThread());
 
-	//Set the cancel state on all task
+	//Cancel the queued tasks, we cannot simply not do them since, there is some promise objects
+	//to setup in the completion task
+	const bool bCancelAllTasks = true;
+	StartQueuedTasks(bCancelAllTasks);
+
+	//Set the cancel state on all running tasks
 	int32 ImportTaskCount = ImportTasks.Num();
 	for (int32 TaskIndex = 0; TaskIndex < ImportTaskCount; ++TaskIndex)
 	{
@@ -696,7 +754,6 @@ void UInterchangeManager::CancelAllTasks()
 			AsyncHelper->InitCancel();
 		}
 	}
-
 	//Tasks should all finish quite fast now
 };
 
diff --git a/Engine/Source/Runtime/Interchange/Engine/Public/InterchangeManager.h b/Engine/Source/Runtime/Interchange/Engine/Public/InterchangeManager.h
--- a/Engine/Source/Runtime/Interchange/Engine/Public/InterchangeManager.h
+++ b/Engine/Source/Runtime/Interchange/Engine/Public/InterchangeManager.h
@@ -301,7 +301,7 @@ public:
 	 * Return an FImportAsynHelper pointer. The pointer is deleted when ReleaseAsyncHelper is call.
 	 * @param Data - The data we want to pass to the different import tasks
 	 */
-	TWeakPtr<UE::Interchange::FImportAsyncHelper, ESPMode::ThreadSafe> CreateAsyncHelper(const UE::Interchange::FImportAsyncHelperData& Data);
+	TSharedPtr<UE::Interchange::FImportAsyncHelper, ESPMode::ThreadSafe> CreateAsyncHelper(const UE::Interchange::FImportAsyncHelperData& Data);
 
 	/** Delete the specified AsyncHelper and remove it from the array that was holding it. */
 	void ReleaseAsyncHelper(TWeakPtr<UE::Interchange::FImportAsyncHelper, ESPMode::ThreadSafe> AsyncHelper);
@@ -350,7 +350,23 @@ protected:
 	 */
 	void SetActiveMode(bool IsActive);
 
+	/**
+	 * Start task until we reach the taskgraph worker number.
+	 * @param bCancelAllTasks - If true we will start all task but with the cancel state set, so task will complete fast and call the completion task
+	 */
+	void StartQueuedTasks(bool bCancelAllTasks = false);
+
 private:
+	struct FQueuedTaskData
+	{
+		FString PackageBasePath;
+		TSharedPtr<UE::Interchange::FImportAsyncHelper, ESPMode::ThreadSafe> AsyncHelper;
+	};
+	
+	//Queue all incomming tasks if there is more started task then we have cores
+	TQueue<FQueuedTaskData> QueuedTasks;
+	int32 QueueTaskCount = 0;
+
 	//By using pointer, there is no issue if the array get resize
 	TArray<TSharedPtr<UE::Interchange::FImportAsyncHelper, ESPMode::ThreadSafe> > ImportTasks;
 
