Allow to specify the interchange pipeline stack in the project settings

diff --git a/Engine/Plugins/Enterprise/DatasmithImporter/Source/DatasmithImporter/Private/DatasmithTextureImporter.cpp b/Engine/Plugins/Enterprise/DatasmithImporter/Source/DatasmithImporter/Private/DatasmithTextureImporter.cpp
--- a/Engine/Plugins/Enterprise/DatasmithImporter/Source/DatasmithImporter/Private/DatasmithTextureImporter.cpp
+++ b/Engine/Plugins/Enterprise/DatasmithImporter/Source/DatasmithImporter/Private/DatasmithTextureImporter.cpp
@@ -375,7 +375,7 @@ UE::Interchange::FAsyncImportResult FDatasmithTextureImporter::CreateTextureAsyn
 	UDatasmithTexturePipeline* TexturePipeline = NewObject< UDatasmithTexturePipeline >();
 	TexturePipeline->TextureElement = TextureElement;
 
-	ImportAssetParameters.OverridePipeline = TexturePipeline;
+	ImportAssetParameters.OverridePipelines.Add(TexturePipeline);
 
 	UE::Interchange::FAsyncImportResult FutureTexture = UInterchangeManager::GetInterchangeManager().ImportAssetAsync( ContentPath, ScopedSourceData.GetSourceData(), ImportAssetParameters );
 
diff --git a/Engine/Plugins/Enterprise/DatasmithImporter/Source/DatasmithImporter/Public/DatasmithTextureImporter.h b/Engine/Plugins/Enterprise/DatasmithImporter/Source/DatasmithImporter/Public/DatasmithTextureImporter.h
--- a/Engine/Plugins/Enterprise/DatasmithImporter/Source/DatasmithImporter/Public/DatasmithTextureImporter.h
+++ b/Engine/Plugins/Enterprise/DatasmithImporter/Source/DatasmithImporter/Public/DatasmithTextureImporter.h
@@ -54,5 +54,9 @@ class UDatasmithTexturePipeline : public UInterchangePipelineBase
 public:
 	virtual bool ExecutePreImportPipeline(UInterchangeBaseNodeContainer* BaseNodeContainer) override;
 
+	bool IsScripted() override
+	{
+		return false;
+	}
 	TSharedPtr< IDatasmithTextureElement > TextureElement;
 };
\ No newline at end of file
diff --git a/Engine/Source/Editor/UnrealEd/Classes/Settings/EditorExperimentalSettings.h b/Engine/Source/Editor/UnrealEd/Classes/Settings/EditorExperimentalSettings.h
--- a/Engine/Source/Editor/UnrealEd/Classes/Settings/EditorExperimentalSettings.h
+++ b/Engine/Source/Editor/UnrealEd/Classes/Settings/EditorExperimentalSettings.h
@@ -7,8 +7,10 @@
 #include "UObject/ObjectMacros.h"
 #include "UObject/Object.h"
 #include "InputCoreTypes.h"
+#include "InterchangePipelineBase.h"
 #include "EditorExperimentalSettings.generated.h"
 
+
 /**
  * Implements Editor settings for experimental features.
  */
@@ -23,8 +25,8 @@ public:
 	UPROPERTY(EditAnywhere, config, Category = Performance, meta = (DisplayName = "Enable async texture compilation to improve PIE and map load time performance when texture compilation is required"))
 	bool bEnableAsyncTextureCompilation;
 
-	/** Enable interchange framework, the interchange framework is a new import system which can import asynchronously and in parallel. */
-	UPROPERTY(EditAnywhere, config, Category = Performance, meta = (DisplayName = "Enable interchange framework import"))
+	/** Enable interchange framework, the interchange framework is a new import system which can import asynchronously and in parallel. See the interchange project settings to configure the import pipeline*/
+	UPROPERTY(EditAnywhere, config, Category = Interchange, meta = (DisplayName = "Enable interchange framework import"))
 	bool bEnableInterchangeFramework;
 
 	/** Allows the editor to run on HDR monitors on Windows 10 */
diff --git a/Engine/Source/Runtime/Interchange/Core/Public/InterchangePipelineBase.h b/Engine/Source/Runtime/Interchange/Core/Public/InterchangePipelineBase.h
--- a/Engine/Source/Runtime/Interchange/Core/Public/InterchangePipelineBase.h
+++ b/Engine/Source/Runtime/Interchange/Core/Public/InterchangePipelineBase.h
@@ -76,6 +76,19 @@ public:
 	{
 		return CanExecuteOnAnyThread(PipelineTask);
 	}
+
+	/**
+	 * Non scripted class should return false here, we have the default to true because scripted class cannot override
+	 * this function since it can be call in a asynchronous thread, which python cannot be executed.
+	 *
+	 * We cannot call ScriptedCanExecuteOnAnyThread for a scripted python pipeline from the task parsing async thread.
+	 * This function allow us to not call it and force the ScriptedExecutePostImportPipeline to execute on the game thread.
+	 */
+	virtual bool IsScripted()
+	{
+		return true;
+	}
+
 protected:
 
 	/**
diff --git a/Engine/Source/Runtime/Interchange/Engine/InterchangeEngine.Build.cs b/Engine/Source/Runtime/Interchange/Engine/InterchangeEngine.Build.cs
--- a/Engine/Source/Runtime/Interchange/Engine/InterchangeEngine.Build.cs
+++ b/Engine/Source/Runtime/Interchange/Engine/InterchangeEngine.Build.cs
@@ -20,7 +20,8 @@ public class InterchangeEngine : ModuleRules
 			new string[] {
 				"AssetRegistry",
 				"SlateCore",
-				"Slate"
+				"Slate",
+				"DeveloperSettings"
 			}
 		);
 	}
diff --git a/Engine/Source/Runtime/Interchange/Engine/Private/InterchangeManager.cpp b/Engine/Source/Runtime/Interchange/Engine/Private/InterchangeManager.cpp
--- a/Engine/Source/Runtime/Interchange/Engine/Private/InterchangeManager.cpp
+++ b/Engine/Source/Runtime/Interchange/Engine/Private/InterchangeManager.cpp
@@ -9,6 +9,7 @@
 #include "Framework/Notifications/NotificationManager.h"
 #include "InterchangeFactoryBase.h"
 #include "InterchangeEngineLogPrivate.h"
+#include "InterchangeProjectSettings.h"
 #include "InterchangeSourceData.h"
 #include "InterchangeTranslatorBase.h"
 #include "InterchangeWriterBase.h"
@@ -31,6 +32,8 @@
 #include "UObject/WeakObjectPtrTemplates.h"
 #include "Widgets/Notifications/SNotificationList.h"
 
+
+
 namespace InternalInterchangePrivate
 {
 	const FLogCategoryBase* GetLogInterchangePtr()
@@ -506,24 +509,24 @@ UE::Interchange::FAsyncImportResult UInterchangeManager::ImportAssetAsync(const
 		check(AsyncHelper->BaseNodeContainers[SourceDataIndex].IsValid());
 	}
 
-	if ( ImportAssetParameters.OverridePipeline == nullptr )
+	if ( ImportAssetParameters.OverridePipelines.Num() == 0 )
 	{
-		//Get all pipeline candidate we want for this import
-		TArray<UClass*> PipelineCandidates;
-		FindPipelineCandidate(PipelineCandidates);
+		const TArray<TSoftClassPtr<UInterchangePipelineBase>>& PipelineStack = GetDefault<UInterchangeProjectSettings>()->PipelineStack;
 
-		// Stack all pipelines, for this import proto. TODO: We need to be able to control which pipeline we use for the import
-		// It can be set in the project settings and can also be set by a UI where the user create the pipeline stack he want.
-		// This should be a list of available pipelines that can be drop into a stack where you can control the order.
-		for (int32 GraphPipelineIndex = 0; GraphPipelineIndex < PipelineCandidates.Num(); ++GraphPipelineIndex)
+		for (int32 GraphPipelineIndex = 0; GraphPipelineIndex < PipelineStack.Num(); ++GraphPipelineIndex)
 		{
-			UInterchangePipelineBase* GeneratedPipeline = NewObject<UInterchangePipelineBase>(GetTransientPackage(), PipelineCandidates[GraphPipelineIndex], NAME_None, RF_NoFlags);
+			UClass* PipelineClass = PipelineStack[GraphPipelineIndex].LoadSynchronous();
+			UInterchangePipelineBase* GeneratedPipeline = NewObject<UInterchangePipelineBase>(GetTransientPackage(), PipelineClass, NAME_None, RF_NoFlags);
 			AsyncHelper->Pipelines.Add(GeneratedPipeline);
 		}
 	}
 	else
 	{
-		AsyncHelper->Pipelines.Add(ImportAssetParameters.OverridePipeline);
+		for (int32 GraphPipelineIndex = 0; GraphPipelineIndex < ImportAssetParameters.OverridePipelines.Num(); ++GraphPipelineIndex)
+		{
+			//Use directly the override pipeline no duplicate.
+			AsyncHelper->Pipelines.Add(ImportAssetParameters.OverridePipelines[GraphPipelineIndex]);
+		}
 	}
 
 	FQueuedTaskData QueuedTaskData;
diff --git a/Engine/Source/Runtime/Interchange/Engine/Private/Tasks/InterchangeTaskPipeline.h b/Engine/Source/Runtime/Interchange/Engine/Private/Tasks/InterchangeTaskPipeline.h
--- a/Engine/Source/Runtime/Interchange/Engine/Private/Tasks/InterchangeTaskPipeline.h
+++ b/Engine/Source/Runtime/Interchange/Engine/Private/Tasks/InterchangeTaskPipeline.h
@@ -28,6 +28,17 @@ namespace UE
 
 			ENamedThreads::Type GetDesiredThread()
 			{
+				if (!ensure(PipelineBase.IsValid()))
+				{
+					return ENamedThreads::GameThread;
+				}
+				
+				//Scripted (python) cannot run outside of the game thread, it will lock forever if we do this
+				if (PipelineBase.Get()->IsScripted())
+				{
+					return ENamedThreads::GameThread;
+				}
+
 				return PipelineBase.Get()->ScriptedCanExecuteOnAnyThread(EInterchangePipelineTask::PreFactoryImport) ? ENamedThreads::AnyBackgroundThreadNormalTask : ENamedThreads::GameThread;
 			}
 
@@ -71,6 +82,13 @@ namespace UE
 				{
 					return ENamedThreads::GameThread;
 				}
+
+				//Scripted (python) cannot run outside of the game thread, it will lock forever if we do this
+				if (AsyncHelper->Pipelines[PipelineIndex]->IsScripted())
+				{
+					return ENamedThreads::GameThread;
+				}
+
 				//Ask the pipeline implementation
 				if (AsyncHelper->Pipelines[PipelineIndex]->ScriptedCanExecuteOnAnyThread(EInterchangePipelineTask::PostFactoryImport))
 				{
diff --git a/Engine/Source/Runtime/Interchange/Engine/Public/InterchangeManager.h b/Engine/Source/Runtime/Interchange/Engine/Public/InterchangeManager.h
--- a/Engine/Source/Runtime/Interchange/Engine/Public/InterchangeManager.h
+++ b/Engine/Source/Runtime/Interchange/Engine/Public/InterchangeManager.h
@@ -157,7 +157,7 @@ struct FImportAssetParameters
 	bool bIsAutomated = false;
 
 	UPROPERTY(EditAnywhere, BlueprintReadWrite, Category = "Interchange|ImportAsset")
-	UInterchangePipelineBase* OverridePipeline = nullptr;
+	TArray<UInterchangePipelineBase*> OverridePipelines;
 };
 
 UCLASS(Transient, BlueprintType)
diff --git a/Engine/Source/Runtime/Interchange/Engine/Public/InterchangeProjectSettings.h b/Engine/Source/Runtime/Interchange/Engine/Public/InterchangeProjectSettings.h
new file mode 100644
--- /dev/null
+++ b/Engine/Source/Runtime/Interchange/Engine/Public/InterchangeProjectSettings.h
@@ -0,0 +1,20 @@
+// Copyright Epic Games, Inc. All Rights Reserved.
+
+#pragma once
+
+#include "CoreMinimal.h"
+#include "Engine/DeveloperSettings.h"
+#include "Engine/EngineTypes.h"
+
+#include "InterchangeProjectSettings.generated.h"
+
+UCLASS(config=Engine, meta=(DisplayName=Interchange), MinimalAPI)
+class UInterchangeProjectSettings : public UDeveloperSettings
+{
+	GENERATED_BODY()
+
+public:
+	/** Stack the pipeline you want to use to import asset with interchange. The pipelines are execute from top to bottom order. You can order them by using the grip on the left of any pipelines.*/
+	UPROPERTY(EditAnywhere, config, Category = Interchange, meta = (DisplayName = "Import pipeline stack"))
+	TArray<TSoftClassPtr<UInterchangePipelineBase>> PipelineStack;
+};
\ No newline at end of file
