Fix CIS build

diff --git a/Engine/Source/Runtime/Interchange/Engine/Public/InterchangeProjectSettings.h b/Engine/Source/Runtime/Interchange/Engine/Public/InterchangeProjectSettings.h
--- a/Engine/Source/Runtime/Interchange/Engine/Public/InterchangeProjectSettings.h
+++ b/Engine/Source/Runtime/Interchange/Engine/Public/InterchangeProjectSettings.h
@@ -5,6 +5,7 @@
 #include "CoreMinimal.h"
 #include "Engine/DeveloperSettings.h"
 #include "Engine/EngineTypes.h"
+#include "InterchangePipelineBase.h"
 
 #include "InterchangeProjectSettings.generated.h"
 
