Fix an assert when importing texture named "None"

diff --git a/Engine/Source/Runtime/Interchange/Core/Private/Nodes/InterchangeBaseNodeContainer.cpp b/Engine/Source/Runtime/Interchange/Core/Private/Nodes/InterchangeBaseNodeContainer.cpp
--- a/Engine/Source/Runtime/Interchange/Core/Private/Nodes/InterchangeBaseNodeContainer.cpp
+++ b/Engine/Source/Runtime/Interchange/Core/Private/Nodes/InterchangeBaseNodeContainer.cpp
@@ -30,6 +30,13 @@ FName UInterchangeBaseNodeContainer::AddNode(UInterchangeBaseNode* Node)
 		return NodeUniqueID;
 	}
 
+	if (Node->GetDisplayLabel() == NAME_None)
+	{
+		//Replace None by Null, since None name will be interpret like NAME_None which will not work with UObject creation
+		//UObject Creation will name it ClassName_X instead of None
+		//TODO Log an warning to the user
+		Node->SetDisplayLabel(FName("Null"));
+	}
 	//Copy the node
 	Nodes.Add(NodeUniqueID, Node);
 	return NodeUniqueID;
