Fix datasmith crash importing texture name none

diff --git a/Engine/Source/Runtime/Interchange/Core/Private/Nodes/InterchangeBaseNode.cpp b/Engine/Source/Runtime/Interchange/Core/Private/Nodes/InterchangeBaseNode.cpp
--- a/Engine/Source/Runtime/Interchange/Core/Private/Nodes/InterchangeBaseNode.cpp
+++ b/Engine/Source/Runtime/Interchange/Core/Private/Nodes/InterchangeBaseNode.cpp
@@ -55,6 +55,17 @@ FName UInterchangeBaseNode::GetDisplayLabel() const
 	return DisplayLabel;
 }
 
+bool UInterchangeBaseNode::SetDisplayLabel(FName DisplayLabel)
+{
+	UE::Interchange::EAttributeStorageResult Result = Attributes.RegisterAttribute(UE::Interchange::FBaseNodeStaticData::DisplayLabelKey(), DisplayLabel);
+	if (IsAttributeStorageResultSuccess(Result))
+	{
+		UE::Interchange::FAttributeStorage::TAttributeHandle<FName> Handle = Attributes.GetAttributeHandle<FName>(UE::Interchange::FBaseNodeStaticData::DisplayLabelKey());
+		return Handle.IsValid();
+	}
+	return false;
+}
+
 FName UInterchangeBaseNode::GetParentUID() const
 {
 	if (!Attributes.ContainAttribute(UE::Interchange::FBaseNodeStaticData::ParentIDKey()))
diff --git a/Engine/Source/Runtime/Interchange/Core/Private/Nodes/InterchangeBaseNodeContainer.cpp b/Engine/Source/Runtime/Interchange/Core/Private/Nodes/InterchangeBaseNodeContainer.cpp
--- a/Engine/Source/Runtime/Interchange/Core/Private/Nodes/InterchangeBaseNodeContainer.cpp
+++ b/Engine/Source/Runtime/Interchange/Core/Private/Nodes/InterchangeBaseNodeContainer.cpp
@@ -30,6 +30,14 @@ FName UInterchangeBaseNodeContainer::AddNode(UInterchangeBaseNode* Node)
 		return NodeUniqueID;
 	}
 
+	if (Node->GetDisplayLabel() == NAME_None)
+	{
+		//Replace None by Null, since None name will be interpret like NAME_None which will not work with UObject creation
+		//UObject Creation will name it ClassName_X instead of None
+		//TODO Log an warning to the user
+		Node->SetDisplayLabel(FName("Null"));
+	}
+
 	//Copy the node
 	Nodes.Add(NodeUniqueID, Node);
 	return NodeUniqueID;
diff --git a/Engine/Source/Runtime/Interchange/Core/Public/Nodes/InterchangeBaseNode.h b/Engine/Source/Runtime/Interchange/Core/Public/Nodes/InterchangeBaseNode.h
--- a/Engine/Source/Runtime/Interchange/Core/Public/Nodes/InterchangeBaseNode.h
+++ b/Engine/Source/Runtime/Interchange/Core/Public/Nodes/InterchangeBaseNode.h
@@ -307,12 +307,19 @@ public:
 	FName GetUniqueID() const;
 
 	/**
-	 * Return the display label pass in the constructor.
+	 * Return the display label.
 	 *
 	 */
 	UFUNCTION(BlueprintCallable, Category = "Interchange | Node")
 	FName GetDisplayLabel() const;
 
+	/**
+	 * Change the display label.
+	 *
+	 */
+	UFUNCTION(BlueprintCallable, Category = "Interchange | Node")
+	bool SetDisplayLabel(FName DisplayName);
+
 	/**
 	 * Return the parent unique id. In case the attribute does not exist it will return InvalidNodeUID()
 	 */
