Interchange: Changed AssetImportResult to be able to hold multiple assets.
- FAssetImportResult is now stateful, going from InProgress to Done when all assets have been imported.
- It's possible to request the list of assets while the import is InProgress. The list will be partial.
- A callback is available to be called once the import is complete and all the results have been made available.
- For non async code paths in Datasmith. a FAssetImportResult is created, its asset is immediately added and its status is set to done.
- Added a new InterchangeEngineFwd to put forward declarations for the Interchange Engine module so that we don't have to write the namespaces all the time.

diff --git a/Engine/Plugins/Enterprise/DatasmithImporter/Source/DatasmithImporter/Private/DatasmithConsumer.cpp b/Engine/Plugins/Enterprise/DatasmithImporter/Source/DatasmithImporter/Private/DatasmithConsumer.cpp
--- a/Engine/Plugins/Enterprise/DatasmithImporter/Source/DatasmithImporter/Private/DatasmithConsumer.cpp
+++ b/Engine/Plugins/Enterprise/DatasmithImporter/Source/DatasmithImporter/Private/DatasmithConsumer.cpp
@@ -1528,8 +1528,10 @@ namespace DatasmithConsumerUtils
 					TSharedRef< IDatasmithTextureElement > TextureElement = FDatasmithSceneFactory::CreateTexture( *AssetTag );
 					TextureElement->SetLabel( *Texture->GetName() );
 
-					TPromise< UObject* > TexturePromise = MakeFulfilledPromise< UObject* >( Texture );
-					ImportContext.ImportedTextures.FindOrAdd( TextureElement ) = UE::Interchange::FAsyncImportResult{ TexturePromise.GetFuture(), FGraphEventRef() };
+					UE::Interchange::FAssetImportResultRef& AssetImportResults = ImportContext.ImportedTextures.Add( TextureElement, MakeShared< UE::Interchange::FAssetImportResult, ESPMode::ThreadSafe >() );
+					AssetImportResults->AddImportedAsset( Texture );
+					AssetImportResults->SetDone();
+
 					ImportContext.Scene->AddTexture( TextureElement );
 				}
 				else if(UMaterialInstance* MaterialInstance = Cast<UMaterialInstance>(Asset))
diff --git a/Engine/Plugins/Enterprise/DatasmithImporter/Source/DatasmithImporter/Private/DatasmithFileProducer.cpp b/Engine/Plugins/Enterprise/DatasmithImporter/Source/DatasmithImporter/Private/DatasmithFileProducer.cpp
--- a/Engine/Plugins/Enterprise/DatasmithImporter/Source/DatasmithImporter/Private/DatasmithFileProducer.cpp
+++ b/Engine/Plugins/Enterprise/DatasmithImporter/Source/DatasmithImporter/Private/DatasmithFileProducer.cpp
@@ -301,11 +301,13 @@ void UDatasmithFileProducer::SceneElementToWorld()
 		}
 	}
 
-	for ( TPair< TSharedRef< IDatasmithTextureElement >, UE::Interchange::FAsyncImportResult >& AssetPair : ImportContextPtr->ImportedTextures )
+	for ( TPair< TSharedRef< IDatasmithTextureElement >, UE::Interchange::FAssetImportResultRef >& AssetPair : ImportContextPtr->ImportedTextures )
 	{
-		if(AssetPair.Value.IsValid())
+		AssetPair.Value->WaitUntilDone();
+
+		if ( UTexture* ImportedTexture = Cast< UTexture >( AssetPair.Value->GetFirstAssetOfClass( UTexture::StaticClass() ) ) )
 		{
-			Assets.Emplace( AssetPair.Value.Get() );
+			Assets.Emplace( ImportedTexture );
 		}
 	}
 
diff --git a/Engine/Plugins/Enterprise/DatasmithImporter/Source/DatasmithImporter/Private/DatasmithImporter.cpp b/Engine/Plugins/Enterprise/DatasmithImporter/Source/DatasmithImporter/Private/DatasmithImporter.cpp
--- a/Engine/Plugins/Enterprise/DatasmithImporter/Source/DatasmithImporter/Private/DatasmithImporter.cpp
+++ b/Engine/Plugins/Enterprise/DatasmithImporter/Source/DatasmithImporter/Private/DatasmithImporter.cpp
@@ -342,21 +342,18 @@ void FDatasmithImporter::ImportTextures( FDatasmithImportContext& ImportContext
 				break;
 			}
 
-			UE::Interchange::FAsyncImportResult FutureTexture = DatasmithTextureImporter.CreateTextureAsync( FilteredTextureElements[TextureIndex] );
+			UE::Interchange::FAssetImportResultRef FutureTexture = DatasmithTextureImporter.CreateTextureAsync( FilteredTextureElements[TextureIndex] );
 
-			if ( FutureTexture.IsValid() )
+			if ( FutureTexture->IsValid() )
 			{
-				FutureTexture = FutureTexture.Next(
-					[ TextureElement = FilteredTextureElements[TextureIndex].ToSharedRef(), &ImportContext ]( UObject* Future )
+				FutureTexture->OnDone(
+					[ TextureElement = FilteredTextureElements[TextureIndex].ToSharedRef(), &ImportContext ]( UE::Interchange::FAssetImportResult& AssetImportResults )
 					{
-						ImportMetaDataForObject( ImportContext, TextureElement, Future );
-
-						return Future;
+						ImportMetaDataForObject( ImportContext, TextureElement, AssetImportResults.GetFirstAssetOfClass( UTexture::StaticClass() ) );
 					}
 				);
 
-				UE::Interchange::FAsyncImportResult& ImportedTexture = ImportContext.ImportedTextures.FindOrAdd( FilteredTextureElements[TextureIndex].ToSharedRef() );
-				ImportedTexture = MoveTemp( FutureTexture );
+				ImportContext.ImportedTextures.Add( FilteredTextureElements[TextureIndex].ToSharedRef(), MoveTemp( FutureTexture ) );
 			}
 		}
 
@@ -449,19 +446,20 @@ UTexture* FDatasmithImporter::ImportTexture( FDatasmithImportContext& ImportCont
 {
 	TRACE_CPUPROFILER_EVENT_SCOPE(FDatasmithImporter::ImportTexture);
 
-	UE::Interchange::FAsyncImportResult& ImportedTexture = ImportContext.ImportedTextures.FindOrAdd( TextureElement );
-	TPromise< UObject* > TexturePromise = MakeFulfilledPromise< UObject* >( DatasmithTextureImporter.CreateTexture( TextureElement, TextureData, Extension ) );
-	ImportedTexture = UE::Interchange::FAsyncImportResult{ TexturePromise.GetFuture(), FGraphEventRef() };
+	UTexture* ImportedTexture = DatasmithTextureImporter.CreateTexture( TextureElement, TextureData, Extension );
 
-	if ( !ImportedTexture.Get() )
+	if ( !ImportedTexture )
 	{
-		ImportContext.ImportedTextures.Remove( TextureElement );
 		return nullptr;
 	}
 
-	ImportMetaDataForObject( ImportContext, TextureElement, ImportedTexture.Get() );
+	UE::Interchange::FAssetImportResultRef& AssetImportResults = ImportContext.ImportedTextures.Add( TextureElement, MakeShared< UE::Interchange::FAssetImportResult, ESPMode::ThreadSafe >() );
+	AssetImportResults->AddImportedAsset( ImportedTexture );
+	AssetImportResults->SetDone();
+
+	ImportMetaDataForObject( ImportContext, TextureElement, ImportedTexture );
 
-	return Cast< UTexture >( ImportedTexture.Get() );
+	return Cast< UTexture >( ImportedTexture );
 }
 
 UTexture* FDatasmithImporter::FinalizeTexture( UTexture* SourceTexture, const TCHAR* TexturesFolderPath, UTexture* ExistingTexture, TMap< UObject*, UObject* >* ReferencesToRemap )
@@ -1521,14 +1519,15 @@ void FDatasmithImporter::FinalizeImport(FDatasmithImportContext& ImportContext,
 	FScopedSlowTask* Progress = ProgressPtr.Get();
 
 	// Needs to be done in dependencies order (textures -> materials -> static meshes)
-	for (const TPair< TSharedRef< IDatasmithTextureElement >, UE::Interchange::FAsyncImportResult >& ImportedTexturePair : ImportContext.ImportedTextures)
+	for (const TPair< TSharedRef< IDatasmithTextureElement >, UE::Interchange::FAssetImportResultRef >& ImportedTexturePair : ImportContext.ImportedTextures)
 	{
 		if (ImportContext.bUserCancelled)
 		{
 			break;
 		}
 
-		UTexture* SourceTexture = Cast< UTexture >( ImportedTexturePair.Value.Get() );
+		ImportedTexturePair.Value->WaitUntilDone();
+		UTexture* SourceTexture = Cast<UTexture>(ImportedTexturePair.Value->GetFirstAssetOfClass(UTexture::StaticClass()));
 
 		if (!SourceTexture || (ValidAssets.Num() > 0 && !ValidAssets.Contains(Cast<UObject>(SourceTexture))))
 		{
diff --git a/Engine/Plugins/Enterprise/DatasmithImporter/Source/DatasmithImporter/Private/DatasmithTextureImporter.cpp b/Engine/Plugins/Enterprise/DatasmithImporter/Source/DatasmithImporter/Private/DatasmithTextureImporter.cpp
--- a/Engine/Plugins/Enterprise/DatasmithImporter/Source/DatasmithImporter/Private/DatasmithTextureImporter.cpp
+++ b/Engine/Plugins/Enterprise/DatasmithImporter/Source/DatasmithImporter/Private/DatasmithTextureImporter.cpp
@@ -363,7 +363,7 @@ UTexture* FDatasmithTextureImporter::CreateIESTexture(const TSharedPtr<IDatasmit
 	return IESTexture;
 }
 
-UE::Interchange::FAsyncImportResult FDatasmithTextureImporter::CreateTextureAsync(const TSharedPtr<IDatasmithTextureElement>& TextureElement)
+UE::Interchange::FAssetImportResultRef FDatasmithTextureImporter::CreateTextureAsync(const TSharedPtr<IDatasmithTextureElement>& TextureElement)
 {
 	UE::Interchange::FScopedSourceData ScopedSourceData( TextureElement->GetFile() );
 
@@ -377,9 +377,7 @@ UE::Interchange::FAsyncImportResult FDatasmithTextureImporter::CreateTextureAsyn
 
 	ImportAssetParameters.OverridePipelines.Add(TexturePipeline);
 
-	UE::Interchange::FAsyncImportResult FutureTexture = UInterchangeManager::GetInterchangeManager().ImportAssetAsync( ContentPath, ScopedSourceData.GetSourceData(), ImportAssetParameters );
-
-	return FutureTexture;
+	return UInterchangeManager::GetInterchangeManager().ImportAssetAsync( ContentPath, ScopedSourceData.GetSourceData(), ImportAssetParameters );
 }
 
 bool UDatasmithTexturePipeline::ExecutePreImportPipeline(UInterchangeBaseNodeContainer* BaseNodeContainer)
diff --git a/Engine/Plugins/Enterprise/DatasmithImporter/Source/DatasmithImporter/Public/DatasmithImportContext.h b/Engine/Plugins/Enterprise/DatasmithImporter/Source/DatasmithImporter/Public/DatasmithImportContext.h
--- a/Engine/Plugins/Enterprise/DatasmithImporter/Source/DatasmithImporter/Public/DatasmithImportContext.h
+++ b/Engine/Plugins/Enterprise/DatasmithImporter/Source/DatasmithImporter/Public/DatasmithImportContext.h
@@ -259,7 +259,7 @@ struct DATASMITHIMPORTER_API FDatasmithImportContext
 	TMap< FString, TSharedRef < IDatasmithMeshElement > > ImportedStaticMeshesByName;
 
 	/** Map of imported texture for each texture element */
-	TMap< TSharedRef< IDatasmithTextureElement >, UE::Interchange::FAsyncImportResult > ImportedTextures;
+	TMap< TSharedRef< IDatasmithTextureElement >, UE::Interchange::FAssetImportResultRef > ImportedTextures;
 
 	/** Map of imported material for each material element */
 	TMap< TSharedRef< IDatasmithBaseMaterialElement >, UMaterialInterface* > ImportedMaterials;
diff --git a/Engine/Plugins/Enterprise/DatasmithImporter/Source/DatasmithImporter/Public/DatasmithTextureImporter.h b/Engine/Plugins/Enterprise/DatasmithImporter/Source/DatasmithImporter/Public/DatasmithTextureImporter.h
--- a/Engine/Plugins/Enterprise/DatasmithImporter/Source/DatasmithImporter/Public/DatasmithTextureImporter.h
+++ b/Engine/Plugins/Enterprise/DatasmithImporter/Source/DatasmithImporter/Public/DatasmithTextureImporter.h
@@ -7,6 +7,7 @@
 #include "UObject/ObjectMacros.h"
 #include "UObject/StrongObjectPtr.h"
 
+#include "InterchangeEngineFwd.h"
 #include "InterchangePipelineBase.h"
 #include "Async/TaskGraphInterfaces.h"
 
@@ -17,14 +18,6 @@ class IDatasmithTextureElement;
 class UTexture;
 class UTextureFactory;
 
-namespace UE
-{
-	namespace Interchange
-	{
-		class FAsyncImportResult;
-	}
-}
-
 class FDatasmithTextureImporter : private FNoncopyable
 {
 public:
@@ -35,7 +28,7 @@ public:
 	bool GetTextureData(const TSharedPtr<IDatasmithTextureElement>& TextureElement, TArray<uint8>& TextureData, FString& Extension);
 	UTexture* CreateIESTexture(const TSharedPtr<IDatasmithTextureElement>& TextureElement);
 
-	UE::Interchange::FAsyncImportResult CreateTextureAsync(const TSharedPtr<IDatasmithTextureElement>& TextureElement);
+	UE::Interchange::FAssetImportResultRef CreateTextureAsync(const TSharedPtr<IDatasmithTextureElement>& TextureElement);
 
 private:
 	bool ResizeTextureElement(const TSharedPtr<IDatasmithTextureElement>& TextureElement, FString& ResizedFilename);
diff --git a/Engine/Plugins/Enterprise/DatasmithImporter/Source/DatasmithImporter/Public/Utility/DatasmithImporterUtils.h b/Engine/Plugins/Enterprise/DatasmithImporter/Source/DatasmithImporter/Public/Utility/DatasmithImporterUtils.h
--- a/Engine/Plugins/Enterprise/DatasmithImporter/Source/DatasmithImporter/Public/Utility/DatasmithImporterUtils.h
+++ b/Engine/Plugins/Enterprise/DatasmithImporter/Source/DatasmithImporter/Public/Utility/DatasmithImporterUtils.h
@@ -253,7 +253,7 @@ struct FDatasmithFindAssetTypeHelper< UStaticMesh >
 template<>
 struct FDatasmithFindAssetTypeHelper< UTexture >
 {
-	static const TMap< TSharedRef< IDatasmithTextureElement >, UE::Interchange::FAsyncImportResult >& GetImportedAssetsMap( const FDatasmithAssetsImportContext& AssetsContext )
+	static const TMap< TSharedRef< IDatasmithTextureElement >, UE::Interchange::FAssetImportResultRef >& GetImportedAssetsMap( const FDatasmithAssetsImportContext& AssetsContext )
 	{
 		return AssetsContext.GetParentContext().ImportedTextures;
 	}
@@ -424,25 +424,28 @@ inline UTexture* FDatasmithImporterUtils::FindAsset< UTexture >( const FDatasmit
 	if ( FPaths::IsRelative( ObjectPathName ) )
 	{
 		const TSharedRef< IDatasmithTextureElement >* ImportedElement   = FDatasmithFindAssetTypeHelper< UTexture >::GetImportedElementByName( AssetsContext, ObjectPathName );
-		const TMap< TSharedRef< IDatasmithTextureElement >, UE::Interchange::FAsyncImportResult >& ImportedAssetsMap = FDatasmithFindAssetTypeHelper< UTexture >::GetImportedAssetsMap( AssetsContext );
+		const TMap< TSharedRef< IDatasmithTextureElement >, UE::Interchange::FAssetImportResultRef >& ImportedAssetsMap = FDatasmithFindAssetTypeHelper< UTexture >::GetImportedAssetsMap( AssetsContext );
 
 		if (ImportedElement)
 		{
-			UE::Interchange::FAsyncImportResult* Result =  const_cast<UE::Interchange::FAsyncImportResult* >( ImportedAssetsMap.Find(*ImportedElement) );
+			UE::Interchange::FAssetImportResultRef* Result =  const_cast<UE::Interchange::FAssetImportResultRef* >( ImportedAssetsMap.Find(*ImportedElement) );
 			ensure(Result);
 			if (Result)
 			{
-				return Cast< UTexture >( Result->Get() );
+				(*Result)->WaitUntilDone();
+				return Cast< UTexture >( (*Result)->GetFirstAssetOfClass( UTexture::StaticClass() ) );
 			}
 		}
 		else
 		{
-			for ( const TPair< TSharedRef< IDatasmithTextureElement >, UE::Interchange::FAsyncImportResult >& ImportedAssetPair : ImportedAssetsMap )
+			for ( const TPair< TSharedRef< IDatasmithTextureElement >, UE::Interchange::FAssetImportResultRef >& ImportedAssetPair : ImportedAssetsMap )
 			{
 				if ( FCString::Stricmp( ImportedAssetPair.Key->GetName(), ObjectPathName ) == 0 )
 				{
-					UE::Interchange::FAsyncImportResult& Result = const_cast<UE::Interchange::FAsyncImportResult& >( ImportedAssetPair.Value );
-					return Cast< UTexture >( Result.Get() );
+					UE::Interchange::FAssetImportResultRef& Result = const_cast<UE::Interchange::FAssetImportResultRef& >( ImportedAssetPair.Value );
+					Result->WaitUntilDone();
+
+					return Cast< UTexture >( Result->GetFirstAssetOfClass( UTexture::StaticClass() ) );
 				}
 			}
 		}
diff --git a/Engine/Source/Runtime/Interchange/Engine/Private/InterchangeManager.cpp b/Engine/Source/Runtime/Interchange/Engine/Private/InterchangeManager.cpp
--- a/Engine/Source/Runtime/Interchange/Engine/Private/InterchangeManager.cpp
+++ b/Engine/Source/Runtime/Interchange/Engine/Private/InterchangeManager.cpp
@@ -70,8 +70,8 @@ UInterchangeTranslatorBase* UE::Interchange::FScopedTranslator::GetTranslator()
 }
 
 UE::Interchange::FImportAsyncHelper::FImportAsyncHelper()
+	: AssetImportResult(new FAssetImportResult())
 {
-	RootObjectCompletionEvent = FGraphEvent::CreateGraphEvent();
 	bCancel = false;
 }
 
@@ -153,11 +153,11 @@ void UE::Interchange::FImportAsyncHelper::CancelAndWaitUntilDoneSynchronously()
 	{
 		FTaskGraphInterface::Get().WaitUntilTasksComplete(TasksToComplete, ENamedThreads::GameThread);
 	}
+
 	//Async import result in a null object when we cancel
-	if (!RootObjectCompletionEvent->IsComplete())
+	if (AssetImportResult->GetStatus() != FAssetImportResult::EStatus::Done)
 	{
-		RootObject.SetValue(nullptr);
-		RootObjectCompletionEvent->DispatchSubsequents();
+		AssetImportResult->SetDone();
 	}
 }
 
@@ -204,31 +204,92 @@ void UE::Interchange::FImportAsyncHelper::CleanUp()
 	Factories.Empty();
 }
 
-UE::Interchange::FAsyncImportResult::FAsyncImportResult( TFuture< UObject* >&& InFutureObject, const FGraphEventRef& InGraphEvent )
-	: FutureObject( MoveTemp( InFutureObject ) )
-	, GraphEvent( InGraphEvent )
+UE::Interchange::FAssetImportResult::FAssetImportResult()
+	: ImportStatus(EStatus::Invalid)
+{
+}
+
+UE::Interchange::FAssetImportResult::EStatus UE::Interchange::FAssetImportResult::GetStatus() const
+{
+	return ImportStatus;
+}
+
+bool UE::Interchange::FAssetImportResult::IsValid() const
+{
+	return GetStatus() != EStatus::Invalid;
+}
+
+void UE::Interchange::FAssetImportResult::SetInProgress()
+{
+	EStatus ExpectedStatus = EStatus::Invalid;
+	if (ImportStatus.compare_exchange_strong(ExpectedStatus, EStatus::InProgress))
+	{
+		GraphEvent = FGraphEvent::CreateGraphEvent();
+	}
+}
+
+void UE::Interchange::FAssetImportResult::SetDone()
+{
+	SetInProgress(); // Make sure we always pass through the InProgress state
+
+	EStatus ExpectedStatus = EStatus::InProgress;
+	if (ImportStatus.compare_exchange_strong(ExpectedStatus, EStatus::Done))
+	{
+		if (DoneCallback)
+		{
+			DoneCallback(*this);
+		}
+
+		GraphEvent->DispatchSubsequents();
+	}
+}
+
+void UE::Interchange::FAssetImportResult::WaitUntilDone()
 {
+	if (ImportStatus == EStatus::InProgress)
+	{
+		FTaskGraphInterface::Get().WaitUntilTaskCompletes(GraphEvent);
+	}
 }
 
-bool UE::Interchange::FAsyncImportResult::IsValid() const
+const TArray< UObject* >& UE::Interchange::FAssetImportResult::GetImportedAssets() const
 {
-	return FutureObject.IsValid();
+	FReadScopeLock ReadScopeLock(ImportedAssetsRWLock);
+	return ImportedAssets;
 }
 
-UObject* UE::Interchange::FAsyncImportResult::Get() const
+UObject* UE::Interchange::FAssetImportResult::GetFirstAssetOfClass(UClass* InClass) const
 {
-	if ( !FutureObject.IsReady() )
+	UObject* Asset = nullptr;
+
+	FReadScopeLock ReadScopeLock(ImportedAssetsRWLock);
+	for (UObject* ImportedAsset : ImportedAssets)
 	{
-		// Tick the task graph until our FutureObject is ready
-		FTaskGraphInterface::Get().WaitUntilTaskCompletes( GraphEvent );
+		if (ImportedAsset->IsA(InClass))
+		{
+			Asset = ImportedAsset;
+			break;
+		}
 	}
 
-	return FutureObject.Get();
+	return Asset;
+}
+
+void UE::Interchange::FAssetImportResult::AddImportedAsset(UObject* ImportedAsset)
+{
+	FWriteScopeLock WriteScopeLock(ImportedAssetsRWLock);
+	ImportedAssets.Add(ImportedAsset);
 }
 
-UE::Interchange::FAsyncImportResult UE::Interchange::FAsyncImportResult::Next( TFunction< UObject*( UObject* ) > Continuation )
+void UE::Interchange::FAssetImportResult::OnDone(TFunction< void(FAssetImportResult&) > Callback)
 {
-	return UE::Interchange::FAsyncImportResult{ FutureObject.Next( Continuation ), GraphEvent };
+	DoneCallback = Callback;
+}
+
+void UE::Interchange::FAssetImportResult::AddReferencedObjects(FReferenceCollector& Collector)
+{
+	FReadScopeLock ReadScopeLock(ImportedAssetsRWLock);
+	Collector.AddReferencedObjects(ImportedAssets);
 }
 
 void UE::Interchange::SanitizeInvalidChar(FString& String)
@@ -459,15 +520,15 @@ void UInterchangeManager::StartQueuedTasks(bool bCancelAllTasks /*= false*/)
 
 bool UInterchangeManager::ImportAsset(const FString& ContentPath, const UInterchangeSourceData* SourceData, const FImportAssetParameters& ImportAssetParameters)
 {
-	return ImportAssetAsync( ContentPath, SourceData, ImportAssetParameters ).IsValid();
+	return ImportAssetAsync( ContentPath, SourceData, ImportAssetParameters )->IsValid();
 }
 
-UE::Interchange::FAsyncImportResult UInterchangeManager::ImportAssetAsync(const FString& ContentPath, const UInterchangeSourceData* SourceData, const FImportAssetParameters& ImportAssetParameters)
+UE::Interchange::FAssetImportResultRef UInterchangeManager::ImportAssetAsync(const FString& ContentPath, const UInterchangeSourceData* SourceData, const FImportAssetParameters& ImportAssetParameters)
 {
 	if (!ensure(IsInGameThread()))
 	{
-		//Import process can be start only in game thread
-		return UE::Interchange::FAsyncImportResult{};
+		//Import process can be started only in the game thread
+		return MakeShared< UE::Interchange::FAssetImportResult, ESPMode::ThreadSafe >();
 	}
 
 	FString PackageBasePath = ContentPath;
@@ -533,7 +594,7 @@ UE::Interchange::FAsyncImportResult UInterchangeManager::ImportAssetAsync(const
 
 	StartQueuedTasks();
 
-	return UE::Interchange::FAsyncImportResult{ AsyncHelper->RootObject.GetFuture(), AsyncHelper->RootObjectCompletionEvent };
+	return AsyncHelper->AssetImportResult;
 }
 
 bool UInterchangeManager::ImportScene(const FString& ImportContext, const UInterchangeSourceData* SourceData, bool bIsReimport, bool bIsAutomated)
@@ -567,6 +628,8 @@ TSharedPtr<UE::Interchange::FImportAsyncHelper, ESPMode::ThreadSafe> UInterchang
 	//Copy the task data
 	AsyncHelper->TaskData = Data;
 	
+	AsyncHelper->AssetImportResult->SetInProgress();
+
 	return AsyncHelper;
 }
 
diff --git a/Engine/Source/Runtime/Interchange/Engine/Private/Tasks/InterchangeTaskCompletion.cpp b/Engine/Source/Runtime/Interchange/Engine/Private/Tasks/InterchangeTaskCompletion.cpp
--- a/Engine/Source/Runtime/Interchange/Engine/Private/Tasks/InterchangeTaskCompletion.cpp
+++ b/Engine/Source/Runtime/Interchange/Engine/Private/Tasks/InterchangeTaskCompletion.cpp
@@ -24,7 +24,6 @@ void UE::Interchange::FTaskCompletion::DoTask(ENamedThreads::Type CurrentThread,
 	//No need anymore of the translators sources
 	AsyncHelper->ReleaseTranslatorsSource();
 
-	bool bIsFutureRootObjectSet = false;
 	for(TPair<int32, TArray<FImportAsyncHelper::FImportedAssetInfo>>& AssetInfosPerSourceIndexPair : AsyncHelper->ImportedAssetsPerSourceIndex)
 	{
 		//Verify if the task was cancel
@@ -68,12 +67,7 @@ void UE::Interchange::FTaskCompletion::DoTask(ENamedThreads::Type CurrentThread,
 				FAssetRegistryModule::AssetCreated(Asset);
 			}
 
-			if (!bIsFutureRootObjectSet && SourceIndex == 0)
-			{
-				bIsFutureRootObjectSet = true;
-				AsyncHelper->RootObject.SetValue(Asset);
-				AsyncHelper->RootObjectCompletionEvent->DispatchSubsequents();
-			}
+			AsyncHelper->AssetImportResult->AddImportedAsset(Asset);
 		}
 	}
 
@@ -98,11 +92,7 @@ void UE::Interchange::FTaskCompletion::DoTask(ENamedThreads::Type CurrentThread,
 		}
 	}
 
-	if (!AsyncHelper->RootObjectCompletionEvent->IsComplete())
-	{
-		AsyncHelper->RootObject.SetValue(nullptr);
-		AsyncHelper->RootObjectCompletionEvent->DispatchSubsequents();
-	}
+	AsyncHelper->AssetImportResult->SetDone();
 
 	//Release the async helper
 	AsyncHelper = nullptr;
diff --git a/Engine/Source/Runtime/Interchange/Engine/Public/InterchangeEngineFwd.h b/Engine/Source/Runtime/Interchange/Engine/Public/InterchangeEngineFwd.h
new file mode 100644
--- /dev/null
+++ b/Engine/Source/Runtime/Interchange/Engine/Public/InterchangeEngineFwd.h
@@ -0,0 +1,15 @@
+// Copyright Epic Games, Inc. All Rights Reserved.
+
+#pragma once
+
+#include "Templates/SharedPointer.h"
+
+namespace UE
+{
+	namespace Interchange
+	{
+		class FAssetImportResult;
+
+		using FAssetImportResultRef = TSharedRef< FAssetImportResult, ESPMode::ThreadSafe >;
+	}
+}
diff --git a/Engine/Source/Runtime/Interchange/Engine/Public/InterchangeManager.h b/Engine/Source/Runtime/Interchange/Engine/Public/InterchangeManager.h
--- a/Engine/Source/Runtime/Interchange/Engine/Public/InterchangeManager.h
+++ b/Engine/Source/Runtime/Interchange/Engine/Public/InterchangeManager.h
@@ -4,6 +4,7 @@
 
 #include <atomic>
 #include "Async/TaskGraphInterfaces.h"
+#include "Containers/Queue.h"
 #include "CoreMinimal.h"
 #include "Delegates/DelegateCombinations.h"
 #include "HAL/Thread.h"
@@ -65,7 +66,67 @@ namespace UE
 			UObject* ReimportObject = nullptr;
 		};
 
-		class FImportAsyncHelper : public FGCObject
+		class INTERCHANGEENGINE_API FAssetImportResult : protected FGCObject
+		{
+		public:
+			FAssetImportResult();
+
+			FAssetImportResult(FAssetImportResult&&) = delete;
+			FAssetImportResult& operator=(FAssetImportResult&&) = delete;
+
+			FAssetImportResult(const FAssetImportResult&) = delete;
+			FAssetImportResult& operator=(const FAssetImportResult&) = delete;
+
+			virtual ~FAssetImportResult() = default;
+
+		public:
+			enum class EStatus
+			{
+				Invalid,
+				InProgress,
+				Done
+			};
+
+			EStatus GetStatus() const;
+
+			bool IsValid() const;
+
+			void SetInProgress();
+			void SetDone();
+			void WaitUntilDone();
+
+			// Assets are only made available once they have been completely imported (passed through the entire import pipeline)
+			// While the status isn't EStatus::Done, the list can grow between subsequent calls.
+			// FAssetImportResult holds a reference to the assets so that they aren't garbage collected.
+			const TArray< UObject* >& GetImportedAssets() const;
+
+			// Helper to get the first asset of a certain class. Use when expecting a single asset of that class to be imported since the order isn't deterministic.
+			UObject* GetFirstAssetOfClass(UClass* InClass) const;
+
+			// Adds an asset to the list of imported assets.
+			void AddImportedAsset(UObject* ImportedAsset);
+
+			// Callback when the status switches to done.
+			void OnDone(TFunction< void(FAssetImportResult&) > Callback);
+
+		protected:
+			/* FGCObject interface */
+			virtual void AddReferencedObjects(FReferenceCollector& Collector) override;
+
+		private:
+			std::atomic< EStatus > ImportStatus;
+
+			TArray< UObject* > ImportedAssets;
+			mutable FRWLock ImportedAssetsRWLock;
+
+			FGraphEventRef GraphEvent; // WaitUntilDone waits for this event to be triggered.
+
+			TFunction< void(FAssetImportResult&) > DoneCallback;
+		};
+
+		using FAssetImportResultRef = TSharedRef< FAssetImportResult, ESPMode::ThreadSafe >;
+
+		class FImportAsyncHelper : protected FGCObject
 		{
 		public:
 			FImportAsyncHelper();
@@ -79,7 +140,7 @@ namespace UE
 			virtual void AddReferencedObjects(FReferenceCollector& Collector) override;
 
 			//The following Arrays are per source data
-			TArray <TStrongObjectPtr<UInterchangeBaseNodeContainer>> BaseNodeContainers;
+			TArray<TStrongObjectPtr<UInterchangeBaseNodeContainer>> BaseNodeContainers;
 			TArray<UInterchangeSourceData* > SourceDatas;
 			TArray<UInterchangeTranslatorBase* > Translators;
 			TArray<UInterchangeFactoryBase* > Factories;
@@ -113,8 +174,7 @@ namespace UE
 
 			FImportAsyncHelperData TaskData;
 
-			TPromise< UObject* > RootObject;
-			FGraphEventRef RootObjectCompletionEvent;
+			FAssetImportResultRef AssetImportResult;
 			
 			//If we cancel the tasks, we set this boolean to true
 			std::atomic<bool> bCancel;
@@ -126,21 +186,6 @@ namespace UE
 			void CleanUp();
 		};
 
-		class INTERCHANGEENGINE_API FAsyncImportResult
-		{
-		public:
-			FAsyncImportResult() = default;
-			explicit FAsyncImportResult(TFuture< UObject* >&& InFutureObject, const FGraphEventRef& InGraphEvent);
-
-			bool IsValid() const;
-			UObject* Get() const;
-			FAsyncImportResult Next(TFunction< UObject* (UObject*) > Continuation);
-
-		private:
-			TFuture< UObject* > FutureObject;
-			FGraphEventRef GraphEvent;
-		};
-
 		void SanitizeInvalidChar(FString& String);
 
 	} //ns interchange
@@ -240,7 +285,7 @@ public:
 	 */
 	UFUNCTION(BlueprintCallable, Category = "Interchange | Import Manager")
 	bool ImportAsset(const FString& ContentPath, const UInterchangeSourceData* SourceData, const FImportAssetParameters& ImportAssetParameters);
-	UE::Interchange::FAsyncImportResult ImportAssetAsync(const FString& ContentPath, const UInterchangeSourceData* SourceData, const FImportAssetParameters& ImportAssetParameters);
+	UE::Interchange::FAssetImportResultRef ImportAssetAsync(const FString& ContentPath, const UInterchangeSourceData* SourceData, const FImportAssetParameters& ImportAssetParameters);
 
 	/**
 	 * Call this to start an import scene process, the caller must specify a source data.
