Interchange skeletalmesh import now support morph target import.

diff --git a/Engine/Plugins/Runtime/Interchange/Source/Import/Private/Fbx/InterchangeFbxTranslator.cpp b/Engine/Plugins/Runtime/Interchange/Source/Import/Private/Fbx/InterchangeFbxTranslator.cpp
--- a/Engine/Plugins/Runtime/Interchange/Source/Import/Private/Fbx/InterchangeFbxTranslator.cpp
+++ b/Engine/Plugins/Runtime/Interchange/Source/Import/Private/Fbx/InterchangeFbxTranslator.cpp
@@ -199,9 +199,13 @@ TOptional<UE::Interchange::FSkeletalMeshLodPayloadData> UInterchangeFbxTranslato
 		//Nothing to load from this file
 		return TOptional<UE::Interchange::FSkeletalMeshLodPayloadData>();
 	}
+
+
 	//Buffer keep the ownership of the data, the large memory reader is use to serialize the TMap
 	FLargeMemoryReader Ar(FileData, FileDataSize);
 	SkeletalMeshLodPayload.LodMeshDescription.Serialize(Ar);
+	//Read the blend shape
+	Ar << SkeletalMeshLodPayload.LodBlendShapeMeshDescriptions;
 
 	return SkeletalMeshLodPayload;
 
diff --git a/Engine/Plugins/Runtime/Interchange/Source/Import/Private/Mesh/InterchangeSkeletalMeshFactory.cpp b/Engine/Plugins/Runtime/Interchange/Source/Import/Private/Mesh/InterchangeSkeletalMeshFactory.cpp
--- a/Engine/Plugins/Runtime/Interchange/Source/Import/Private/Mesh/InterchangeSkeletalMeshFactory.cpp
+++ b/Engine/Plugins/Runtime/Interchange/Source/Import/Private/Mesh/InterchangeSkeletalMeshFactory.cpp
@@ -702,6 +702,54 @@ namespace UE
 				return true;
 			}
 
+			bool CopyBlendShapesMeshDescriptionToSkeletalMeshImportData(const TArray<FName>& BlendShapeToImport, const TMap<FString, FMeshDescription>& LodBlendShapeMeshDescriptions, FSkeletalMeshImportData& DestinationSkeletalMeshImportData)
+			{
+				for (const TPair<FString, FMeshDescription>& Pair : LodBlendShapeMeshDescriptions)
+				{
+					FName BlendShapeName(*(Pair.Key));
+					//Skip blend shape that is not define in the BlendShapeToImport array
+					if (!BlendShapeToImport.Contains(BlendShapeName))
+					{
+						continue;
+					}
+
+					TArray<FVector> CompressPoints;
+					CompressPoints.Reserve(DestinationSkeletalMeshImportData.Points.Num());
+
+					const FMeshDescription& SourceMeshDescription = Pair.Value;
+					FStaticMeshConstAttributes Attributes(SourceMeshDescription);
+					TVertexAttributesConstRef<FVector> VertexPositions = Attributes.GetVertexPositions();
+
+					//Create the morph target source data
+					FString& MorphTargetName = DestinationSkeletalMeshImportData.MorphTargetNames.AddDefaulted_GetRef();
+					MorphTargetName = Pair.Key;
+					TSet<uint32>& ModifiedPoints = DestinationSkeletalMeshImportData.MorphTargetModifiedPoints.AddDefaulted_GetRef();
+					FSkeletalMeshImportData& MorphTargetData = DestinationSkeletalMeshImportData.MorphTargets.AddDefaulted_GetRef();
+
+					//Reserve the point and influences
+					MorphTargetData.Points.AddZeroed(SourceMeshDescription.Vertices().Num());
+
+					for (FVertexID VertexID : SourceMeshDescription.Vertices().GetElementIDs())
+					{
+						//We can use GetValue because the Meshdescription was compacted before the copy
+						MorphTargetData.Points[VertexID.GetValue()] = VertexPositions[VertexID];
+					}
+
+					for (int32 PointIdx = 0; PointIdx < DestinationSkeletalMeshImportData.Points.Num(); ++PointIdx)
+					{
+						int32 OriginalPointIdx = DestinationSkeletalMeshImportData.PointToRawMap[PointIdx];
+						//Rebuild the data with only the modified point
+						if ((MorphTargetData.Points[OriginalPointIdx] - DestinationSkeletalMeshImportData.Points[PointIdx]).SizeSquared() > FMath::Square(THRESH_POINTS_ARE_SAME))
+						{
+							ModifiedPoints.Add(PointIdx);
+							CompressPoints.Add(MorphTargetData.Points[OriginalPointIdx]);
+						}
+					}
+					MorphTargetData.Points = CompressPoints;
+				}
+				return true;
+			}
+
 			/**
 			 * Fill the Materials array using the raw skeletalmesh geometry data (using material imported name)
 			 * Find the material from the dependencies of the skeletalmesh before searching in all package.
@@ -1146,6 +1194,12 @@ UObject* UInterchangeSkeletalMeshFactory::CreateAsset(const UInterchangeSkeletal
 					FElementIDRemappings ElementIDRemappings;
 					LodMeshPayload->LodMeshDescription.Compact(ElementIDRemappings);
 					UE::Interchange::Private::CopyMeshDescriptionToSkeletalMeshImportData(LodMeshPayload->LodMeshDescription, SkeletalMeshImportData);
+
+					//Pipeline can remove names from the list to control which blendshapes they import
+					TArray<FName> BlendShapeToImport;
+					LodDataNode->GetBlendShapes(BlendShapeToImport);
+					//Copy also the blend shapes data
+					UE::Interchange::Private::CopyBlendShapesMeshDescriptionToSkeletalMeshImportData(BlendShapeToImport, LodMeshPayload->LodBlendShapeMeshDescriptions, SkeletalMeshImportData);
 				}
 
 				ensure(ImportedResource->LODModels.Add(new FSkeletalMeshLODModel()) == CurrentLodIndex);
diff --git a/Engine/Plugins/Runtime/Interchange/Source/Import/Public/Mesh/InterchangeSkeletalMeshPayload.h b/Engine/Plugins/Runtime/Interchange/Source/Import/Public/Mesh/InterchangeSkeletalMeshPayload.h
--- a/Engine/Plugins/Runtime/Interchange/Source/Import/Public/Mesh/InterchangeSkeletalMeshPayload.h
+++ b/Engine/Plugins/Runtime/Interchange/Source/Import/Public/Mesh/InterchangeSkeletalMeshPayload.h
@@ -13,6 +13,7 @@ namespace UE
 		{
 			//Currently the skeletalmesh payload data is editor only, we have to move to something available at runtime
 			FMeshDescription LodMeshDescription;
+			TMap<FString, FMeshDescription> LodBlendShapeMeshDescriptions;
 		};
 	}//ns Interchange
 }//ns UE
diff --git a/Engine/Plugins/Runtime/Interchange/Source/Nodes/Public/InterchangeJointNode.h b/Engine/Plugins/Runtime/Interchange/Source/Nodes/Public/InterchangeJointNode.h
--- a/Engine/Plugins/Runtime/Interchange/Source/Nodes/Public/InterchangeJointNode.h
+++ b/Engine/Plugins/Runtime/Interchange/Source/Nodes/Public/InterchangeJointNode.h
@@ -42,7 +42,7 @@ public:
 
 	virtual FGuid GetHash() const override
 	{
-		return Attributes.GetStorageHash();
+		return Attributes->GetStorageHash();
 	}
 
 public:
diff --git a/Engine/Plugins/Runtime/Interchange/Source/Nodes/Public/InterchangeMaterialNode.h b/Engine/Plugins/Runtime/Interchange/Source/Nodes/Public/InterchangeMaterialNode.h
--- a/Engine/Plugins/Runtime/Interchange/Source/Nodes/Public/InterchangeMaterialNode.h
+++ b/Engine/Plugins/Runtime/Interchange/Source/Nodes/Public/InterchangeMaterialNode.h
@@ -95,7 +95,7 @@ public:
 		bIsMaterialNodeClassInitialized = false;
 		InitializeNode(UniqueID, DisplayLabel);
 		FString OperationName = GetTypeName() + TEXT(".SetAssetClassName");
-		InterchangePrivateNodeBase::SetCustomAttribute<FString>(Attributes, ClassNameAttributeKey, OperationName, InAssetClass);
+		InterchangePrivateNodeBase::SetCustomAttribute<FString>(*Attributes, ClassNameAttributeKey, OperationName, InAssetClass);
 		FillAssetClassFromAttribute();
 	}
 
@@ -122,16 +122,16 @@ public:
 
 	virtual FGuid GetHash() const override
 	{
-		return Attributes.GetStorageHash();
+		return Attributes->GetStorageHash();
 	}
 
 	virtual const TOptional<FString> GetPayLoadKey() const
 	{
-		if (!Attributes.ContainAttribute(UE::Interchange::FMaterialNodeStaticData::PayloadSourceFileKey()))
+		if (!Attributes->ContainAttribute(UE::Interchange::FMaterialNodeStaticData::PayloadSourceFileKey()))
 		{
 			return TOptional<FString>();
 		}
-		UE::Interchange::FAttributeStorage::TAttributeHandle<FString> AttributeHandle = Attributes.GetAttributeHandle<FString>(UE::Interchange::FMaterialNodeStaticData::PayloadSourceFileKey());
+		UE::Interchange::FAttributeStorage::TAttributeHandle<FString> AttributeHandle = Attributes->GetAttributeHandle<FString>(UE::Interchange::FMaterialNodeStaticData::PayloadSourceFileKey());
 		if (!AttributeHandle.IsValid())
 		{
 			return TOptional<FString>();
@@ -149,7 +149,7 @@ public:
 	UFUNCTION(BlueprintCallable, Category = "Interchange | Node | Material")
 	virtual void SetPayLoadKey(const FString& PayloadKey)
 	{
-		UE::Interchange::EAttributeStorageResult Result = Attributes.RegisterAttribute(UE::Interchange::FMaterialNodeStaticData::PayloadSourceFileKey(), PayloadKey);
+		UE::Interchange::EAttributeStorageResult Result = Attributes->RegisterAttribute(UE::Interchange::FMaterialNodeStaticData::PayloadSourceFileKey(), PayloadKey);
 		if (!IsAttributeStorageResultSuccess(Result))
 		{
 			LogAttributeStorageErrors(Result, TEXT("UInterchangeMaterialNode.SetPayLoadKey"), UE::Interchange::FMaterialNodeStaticData::PayloadSourceFileKey());
@@ -327,7 +327,7 @@ private:
 #if WITH_ENGINE
 		FString OperationName = GetTypeName() + TEXT(".GetAssetClassName");
 		FString ClassName;
-		InterchangePrivateNodeBase::GetCustomAttribute<FString>(Attributes, ClassNameAttributeKey, OperationName, ClassName);
+		InterchangePrivateNodeBase::GetCustomAttribute<FString>(*Attributes, ClassNameAttributeKey, OperationName, ClassName);
 		if (ClassName.Equals(UMaterial::StaticClass()->GetName()))
 		{
 			AssetClass = UMaterial::StaticClass();
diff --git a/Engine/Plugins/Runtime/Interchange/Source/Nodes/Public/InterchangeSceneNode.h b/Engine/Plugins/Runtime/Interchange/Source/Nodes/Public/InterchangeSceneNode.h
--- a/Engine/Plugins/Runtime/Interchange/Source/Nodes/Public/InterchangeSceneNode.h
+++ b/Engine/Plugins/Runtime/Interchange/Source/Nodes/Public/InterchangeSceneNode.h
@@ -30,7 +30,7 @@ public:
 
 	virtual FGuid GetHash() const override
 	{
-		return Attributes.GetStorageHash();
+		return Attributes->GetStorageHash();
 	}
 
 	/** Return false if the Attribute was not set previously.*/
diff --git a/Engine/Plugins/Runtime/Interchange/Source/Nodes/Public/InterchangeSkeletalMeshLodDataNode.h b/Engine/Plugins/Runtime/Interchange/Source/Nodes/Public/InterchangeSkeletalMeshLodDataNode.h
--- a/Engine/Plugins/Runtime/Interchange/Source/Nodes/Public/InterchangeSkeletalMeshLodDataNode.h
+++ b/Engine/Plugins/Runtime/Interchange/Source/Nodes/Public/InterchangeSkeletalMeshLodDataNode.h
@@ -15,16 +15,16 @@ namespace UE
 
 		struct FSkeletalMeshNodeLodDataStaticData : public FBaseNodeStaticData
 		{
-			static const FAttributeKey& TranslatorMeshKeysCountKey()
+			static const FString& GetBlendShapesNameBaseKey()
 			{
-				static FAttributeKey AttributeKey(TEXT("__LOD_DATA_IDS_COUNT_"));
-				return AttributeKey;
+				static FString BlendShapesName_BaseKey = TEXT("BlendShape_Names");
+				return BlendShapesName_BaseKey;
 			}
 
-			static const FAttributeKey& TranslatorMeshKeysBaseKey()
+			static const FString& GetTranslatorMeshKeysBaseKey()
 			{
-				static FAttributeKey AttributeKey(TEXT("__LOD_DATA_IDS_INDEX_"));
-				return AttributeKey;
+				static FString TranslatorMeshKeys_BaseKey = TEXT("TranslatorMesh_Keys");
+				return TranslatorMeshKeys_BaseKey;
 			}
 		};
 
@@ -40,6 +40,8 @@ public:
 	UInterchangeSkeletalMeshLodDataNode()
 		:UInterchangeBaseNode()
 	{
+		BlendShapesNames.Initialize(Attributes, UE::Interchange::FSkeletalMeshNodeLodDataStaticData::GetBlendShapesNameBaseKey());
+		TranslatorMeshKeys.Initialize(Attributes, UE::Interchange::FSkeletalMeshNodeLodDataStaticData::GetTranslatorMeshKeysBaseKey());
 	}
 
 	/**
@@ -53,7 +55,7 @@ public:
 
 	virtual FGuid GetHash() const override
 	{
-		return Attributes.GetStorageHash();
+		return Attributes->GetStorageHash();
 	}
 
 
@@ -71,138 +73,62 @@ public:
 		IMPLEMENT_NODE_ATTRIBUTE_SETTER_NODELEGATE(SkeletonID, FName)
 	}
 
-	/** Return The number of LOD this skeletalmesh has.*/
+	/* Translator mesh keys*/
+
 	UFUNCTION(BlueprintCallable, Category = "Interchange | Node | SkeletalMeshLodData")
 	int32 GetTranslatorMeshKeyCount() const
 	{
-		int32 MeshCount = 0;
-		if (Attributes.ContainAttribute(UE::Interchange::FSkeletalMeshNodeLodDataStaticData::TranslatorMeshKeysCountKey()))
-		{
-			UE::Interchange::FAttributeStorage::TAttributeHandle<int32> Handle = Attributes.GetAttributeHandle<int32>(UE::Interchange::FSkeletalMeshNodeLodDataStaticData::TranslatorMeshKeysCountKey());
-			if (Handle.IsValid())
-			{
-				Handle.Get(MeshCount);
-			}
-		}
-		return MeshCount;
+		return TranslatorMeshKeys.GetCount();
 	}
 
 	UFUNCTION(BlueprintCallable, Category = "Interchange | Node | SkeletalMeshLodData")
 	void GetTranslatorMeshKeys(TArray<FName>& OutTranslatorMeshKeys) const
 	{
-		OutTranslatorMeshKeys.Reset();
-
-		int32 MeshCount = 0;
-		if (!Attributes.ContainAttribute(UE::Interchange::FSkeletalMeshNodeLodDataStaticData::TranslatorMeshKeysCountKey()))
-		{
-			return;
-		}
-		UE::Interchange::FAttributeStorage::TAttributeHandle<int32> Handle = Attributes.GetAttributeHandle<int32>(UE::Interchange::FSkeletalMeshNodeLodDataStaticData::TranslatorMeshKeysCountKey());
-		if (!Handle.IsValid())
-		{
-			return;
-		}
-		Handle.Get(MeshCount);
-		for (int32 MeshIndex = 0; MeshIndex < MeshCount; ++MeshIndex)
-		{
-			FString DepIndexKeyString = UE::Interchange::FSkeletalMeshNodeLodDataStaticData::TranslatorMeshKeysBaseKey().ToString() + FString::FromInt(MeshIndex);
-			UE::Interchange::FAttributeKey DepIndexKey(*DepIndexKeyString);
-			UE::Interchange::FAttributeStorage::TAttributeHandle<FName> HandleMesh = Attributes.GetAttributeHandle<FName>(DepIndexKey);
-			if (!HandleMesh.IsValid())
-			{
-				continue;
-			}
-			FName& NodeUniqueID = OutTranslatorMeshKeys.AddDefaulted_GetRef();
-			HandleMesh.Get(NodeUniqueID);
-		}
+		TranslatorMeshKeys.GetNames(OutTranslatorMeshKeys);
 	}
 
 	UFUNCTION(BlueprintCallable, Category = "Interchange | Node | SkeletalMeshLodData")
 	bool AddTranslatorMeshKey(FName TranslatorMeshKey)
 	{
-		if (!Attributes.ContainAttribute(UE::Interchange::FSkeletalMeshNodeLodDataStaticData::TranslatorMeshKeysCountKey()))
-		{
-			const int32 DependencyCount = 0;
-			UE::Interchange::EAttributeStorageResult Result = Attributes.RegisterAttribute<int32>(UE::Interchange::FSkeletalMeshNodeLodDataStaticData::TranslatorMeshKeysCountKey(), DependencyCount);
-			if (!UE::Interchange::IsAttributeStorageResultSuccess(Result))
-			{
-				UE::Interchange::LogAttributeStorageErrors(Result, TEXT("UInterchangeSkeletalMeshLodDataNode.AddTranslatorMeshKey"), UE::Interchange::FSkeletalMeshNodeLodDataStaticData::TranslatorMeshKeysCountKey());
-				return false;
-			}
-		}
-		UE::Interchange::FAttributeStorage::TAttributeHandle<int32> Handle = Attributes.GetAttributeHandle<int32>(UE::Interchange::FSkeletalMeshNodeLodDataStaticData::TranslatorMeshKeysCountKey());
-		if (!ensure(Handle.IsValid()))
-		{
-			return false;
-		}
-		int32 MeshIndex = 0;
-		Handle.Get(MeshIndex);
-		FString MeshIndexKeyString = UE::Interchange::FSkeletalMeshNodeLodDataStaticData::TranslatorMeshKeysBaseKey().ToString() + FString::FromInt(MeshIndex);
-		UE::Interchange::FAttributeKey MeshIndexKey(*MeshIndexKeyString);
-		//Increment the counter
-		MeshIndex++;
-		Handle.Set(MeshIndex);
-
-		UE::Interchange::EAttributeStorageResult AddMeshResult = Attributes.RegisterAttribute<FName>(MeshIndexKey, TranslatorMeshKey);
-		return true;
+		return TranslatorMeshKeys.AddName(TranslatorMeshKey);
 	}
 
 	UFUNCTION(BlueprintCallable, Category = "Interchange | Node | SkeletalMeshLodData")
 	bool RemoveTranslatorMeshKey(FName TranslatorMeshKey)
 	{
-		int32 MeshCount = 0;
-		if (!Attributes.ContainAttribute(UE::Interchange::FSkeletalMeshNodeLodDataStaticData::TranslatorMeshKeysCountKey()))
-		{
-			return false;
-		}
-		UE::Interchange::FAttributeStorage::TAttributeHandle<int32> Handle = Attributes.GetAttributeHandle<int32>(UE::Interchange::FSkeletalMeshNodeLodDataStaticData::TranslatorMeshKeysCountKey());
-		if (!Handle.IsValid())
-		{
-			return false;
-		}
-		Handle.Get(MeshCount);
-		bool DecrementKey = false;
-		for (int32 MeshIndex = 0; MeshIndex < MeshCount; ++MeshIndex)
-		{
-			FString DepIndexKeyString = UE::Interchange::FSkeletalMeshNodeLodDataStaticData::TranslatorMeshKeysBaseKey().ToString() + FString::FromInt(MeshIndex);
-			UE::Interchange::FAttributeKey DepIndexKey(*DepIndexKeyString);
-			UE::Interchange::FAttributeStorage::TAttributeHandle<FName> HandleMesh = Attributes.GetAttributeHandle<FName>(DepIndexKey);
-			if (!HandleMesh.IsValid())
-			{
-				continue;
-			}
-			FName MeshUniqueID;
-			HandleMesh.Get(MeshUniqueID);
-			if (MeshUniqueID == TranslatorMeshKey)
-			{
-				//Remove this entry
-				Attributes.UnregisterAttribute(DepIndexKey);
-				Handle.Set(MeshCount-1);
-				//We have to rename the key for all the next item
-				DecrementKey = true;;
-			}
-			else if (DecrementKey)
-			{
-				FString NewDepIndexKeyString = UE::Interchange::FSkeletalMeshNodeLodDataStaticData::TranslatorMeshKeysBaseKey().ToString() + FString::FromInt(MeshIndex-1);
-				UE::Interchange::FAttributeKey NewDepIndexKey(*NewDepIndexKeyString);
-				UE::Interchange::EAttributeStorageResult UnregisterResult = Attributes.UnregisterAttribute(DepIndexKey);
-				if (UE::Interchange::IsAttributeStorageResultSuccess(UnregisterResult))
-				{
-					UE::Interchange::EAttributeStorageResult RegisterResult = Attributes.RegisterAttribute<FName>(NewDepIndexKey, MeshUniqueID);
-					if (!UE::Interchange::IsAttributeStorageResultSuccess(RegisterResult))
-					{
-						UE::Interchange::LogAttributeStorageErrors(RegisterResult, TEXT("UInterchangeSkeletalMeshLodDataNode.RemoveTranslatorMeshKey"), NewDepIndexKey);
-					}
-				}
-				else
-				{
-					UE::Interchange::LogAttributeStorageErrors(UnregisterResult, TEXT("UInterchangeSkeletalMeshLodDataNode.RemoveTranslatorMeshKey"), DepIndexKey);
-				}
-				//make sure the for loop iterate since HandleMesh is now invalid
-				continue;
-			}
-		}
-		return true;
+		return TranslatorMeshKeys.RemoveName(TranslatorMeshKey);
+	}
+
+	/* Blend shapes*/
+
+	UFUNCTION(BlueprintCallable, Category = "Interchange | Node | SkeletalMeshLodData")
+	int32 GetBlendShapeCount() const
+	{
+		return BlendShapesNames.GetCount();
+	}
+
+	UFUNCTION(BlueprintCallable, Category = "Interchange | Node | SkeletalMeshLodData")
+	void GetBlendShapes(TArray<FName>& OutBlendShapeNames) const
+	{
+		BlendShapesNames.GetNames(OutBlendShapeNames);
+	}
+
+	UFUNCTION(BlueprintCallable, Category = "Interchange | Node | SkeletalMeshLodData")
+	bool AddBlendShape(FName BlendShapeName)
+	{
+		return BlendShapesNames.AddName(BlendShapeName);
+	}
+
+	UFUNCTION(BlueprintCallable, Category = "Interchange | Node | SkeletalMeshLodData")
+	bool RemoveBlendShape(FName BlendShapeName)
+	{
+		return BlendShapesNames.RemoveName(BlendShapeName);
+	}
+
+	UFUNCTION(BlueprintCallable, Category = "Interchange | Node | SkeletalMeshLodData")
+	bool RemoveAllBlendShape()
+	{
+		return BlendShapesNames.RemoveAllNames();
 	}
 
 private:
@@ -219,5 +145,7 @@ private:
 	//SkeletalMesh
 	const UE::Interchange::FAttributeKey Macro_CustomSkeletonIDKey = UE::Interchange::FAttributeKey(TEXT("SkeletonID"));
 
+	UE::Interchange::FNameAttributeArrayHelper BlendShapesNames;
+	UE::Interchange::FNameAttributeArrayHelper TranslatorMeshKeys;
 protected:
 };
diff --git a/Engine/Plugins/Runtime/Interchange/Source/Nodes/Public/InterchangeSkeletalMeshNode.h b/Engine/Plugins/Runtime/Interchange/Source/Nodes/Public/InterchangeSkeletalMeshNode.h
--- a/Engine/Plugins/Runtime/Interchange/Source/Nodes/Public/InterchangeSkeletalMeshNode.h
+++ b/Engine/Plugins/Runtime/Interchange/Source/Nodes/Public/InterchangeSkeletalMeshNode.h
@@ -19,16 +19,10 @@ namespace UE
 
 		struct FSkeletalMeshNodeStaticData : public FBaseNodeStaticData
 		{
-			static const FAttributeKey& LodDataIdsCountKey()
+			static const FString& GetLodDependenciesBaseKey()
 			{
-				static FAttributeKey AttributeKey(TEXT("__LOD_DATA_IDS_COUNT_"));
-				return AttributeKey;
-			}
-
-			static const FAttributeKey& LodDataIdsBaseKey()
-			{
-				static FAttributeKey AttributeKey(TEXT("__LOD_DATA_IDS_INDEX_"));
-				return AttributeKey;
+				static FString LodDependencies_BaseKey = TEXT("Lod_Dependencies");
+				return LodDependencies_BaseKey;
 			}
 		};
 
@@ -47,6 +41,7 @@ public:
 #if WITH_ENGINE
 		AssetClass = nullptr;
 #endif
+		LodDependencies.Initialize(Attributes, UE::Interchange::FSkeletalMeshNodeStaticData::GetLodDependenciesBaseKey());
 	}
 
 	/**
@@ -60,8 +55,9 @@ public:
 	{
 		bIsNodeClassInitialized = false;
 		InitializeNode(UniqueID, DisplayLabel);
+
 		FString OperationName = GetTypeName() + TEXT(".SetAssetClassName");
-		InterchangePrivateNodeBase::SetCustomAttribute<FString>(Attributes, ClassNameAttributeKey, OperationName, InAssetClass);
+		InterchangePrivateNodeBase::SetCustomAttribute<FString>(*Attributes, ClassNameAttributeKey, OperationName, InAssetClass);
 		FillAssetClassFromAttribute();
 	}
 
@@ -101,7 +97,7 @@ public:
 
 	virtual FGuid GetHash() const override
 	{
-		return Attributes.GetStorageHash();
+		return Attributes->GetStorageHash();
 	}
 
 public:
@@ -109,141 +105,27 @@ public:
 	UFUNCTION(BlueprintCallable, Category = "Interchange | Node | SkeletalMesh")
 	int32 GetLodDataCount() const
 	{
-		int32 LodCount = 0;
-		if (Attributes.ContainAttribute(UE::Interchange::FSkeletalMeshNodeStaticData::LodDataIdsCountKey()))
-		{
-			UE::Interchange::FAttributeStorage::TAttributeHandle<int32> Handle = Attributes.GetAttributeHandle<int32>(UE::Interchange::FSkeletalMeshNodeStaticData::LodDataIdsCountKey());
-			if (Handle.IsValid())
-			{
-				Handle.Get(LodCount);
-			}
-		}
-		return LodCount;
+		return LodDependencies.GetCount();
 	}
 
 	UFUNCTION(BlueprintCallable, Category = "Interchange | Node | SkeletalMesh")
 	void GetLodDataUniqueIds(TArray<FName>& OutLodDataUniqueIds) const
 	{
-		OutLodDataUniqueIds.Reset();
-
-		int32 LodCount = 0;
-		if (!Attributes.ContainAttribute(UE::Interchange::FSkeletalMeshNodeStaticData::LodDataIdsCountKey()))
-		{
-			return;
-		}
-		UE::Interchange::FAttributeStorage::TAttributeHandle<int32> Handle = Attributes.GetAttributeHandle<int32>(UE::Interchange::FSkeletalMeshNodeStaticData::LodDataIdsCountKey());
-		if (!Handle.IsValid())
-		{
-			return;
-		}
-		Handle.Get(LodCount);
-		for (int32 LODIndex = 0; LODIndex < LodCount; ++LODIndex)
-		{
-			FString DepIndexKeyString = UE::Interchange::FSkeletalMeshNodeStaticData::LodDataIdsBaseKey().ToString() + FString::FromInt(LODIndex);
-			UE::Interchange::FAttributeKey DepIndexKey(*DepIndexKeyString);
-			UE::Interchange::FAttributeStorage::TAttributeHandle<FName> HandleLOD = Attributes.GetAttributeHandle<FName>(DepIndexKey);
-			if (!HandleLOD.IsValid())
-			{
-				continue;
-			}
-			FName& NodeUniqueID = OutLodDataUniqueIds.AddDefaulted_GetRef();
-			HandleLOD.Get(NodeUniqueID);
-		}
+		LodDependencies.GetNames(OutLodDataUniqueIds);
 	}
 
 	UFUNCTION(BlueprintCallable, Category = "Interchange | Node | SkeletalMesh")
 	bool AddLodDataUniqueId(FName LodDataUniqueId)
 	{
-		if (!Attributes.ContainAttribute(UE::Interchange::FSkeletalMeshNodeStaticData::LodDataIdsCountKey()))
-		{
-			const int32 DependencyCount = 0;
-			UE::Interchange::EAttributeStorageResult Result = Attributes.RegisterAttribute<int32>(UE::Interchange::FSkeletalMeshNodeStaticData::LodDataIdsCountKey(), DependencyCount);
-			if (!UE::Interchange::IsAttributeStorageResultSuccess(Result))
-			{
-				UE::Interchange::LogAttributeStorageErrors(Result, TEXT("UInterchangeSkeletalMeshNode.AddLodDataUniqueId"), UE::Interchange::FSkeletalMeshNodeStaticData::LodDataIdsCountKey());
-				return false;
-			}
-		}
-		UE::Interchange::FAttributeStorage::TAttributeHandle<int32> Handle = Attributes.GetAttributeHandle<int32>(UE::Interchange::FSkeletalMeshNodeStaticData::LodDataIdsCountKey());
-		if (!ensure(Handle.IsValid()))
-		{
-			return false;
-		}
-		int32 LODIndex = 0;
-		Handle.Get(LODIndex);
-		FString LodIndexKeyString = UE::Interchange::FSkeletalMeshNodeStaticData::LodDataIdsBaseKey().ToString() + FString::FromInt(LODIndex);
-		UE::Interchange::FAttributeKey LodIndexKey(*LodIndexKeyString);
-		//Increment the counter
-		LODIndex++;
-		Handle.Set(LODIndex);
-
-		UE::Interchange::EAttributeStorageResult AddLodResult = Attributes.RegisterAttribute<FName>(LodIndexKey, LodDataUniqueId);
-		if (!UE::Interchange::IsAttributeStorageResultSuccess(AddLodResult))
-		{
-			UE::Interchange::LogAttributeStorageErrors(AddLodResult, TEXT("UInterchangeSkeletalMeshNode.AddLodDataUniqueId"), LodIndexKey);
-			return false;
-		}
-		return true;
+		return LodDependencies.AddName(LodDataUniqueId);
 	}
 
 	UFUNCTION(BlueprintCallable, Category = "Interchange | Node | SkeletalMesh")
 	bool RemoveLodDataUniqueId(FName LodDataUniqueId)
 	{
-		int32 LodCount = 0;
-		if (!Attributes.ContainAttribute(UE::Interchange::FSkeletalMeshNodeStaticData::LodDataIdsCountKey()))
-		{
-			return false;
-		}
-		UE::Interchange::FAttributeStorage::TAttributeHandle<int32> Handle = Attributes.GetAttributeHandle<int32>(UE::Interchange::FSkeletalMeshNodeStaticData::LodDataIdsCountKey());
-		if (!Handle.IsValid())
-		{
-			return false;
-		}
-		Handle.Get(LodCount);
-		bool DecrementKey = false;
-		for (int32 LodIndex = 0; LodIndex < LodCount; ++LodIndex)
-		{
-			FString DepIndexKeyString = UE::Interchange::FSkeletalMeshNodeStaticData::LodDataIdsBaseKey().ToString() + FString::FromInt(LodIndex);
-			UE::Interchange::FAttributeKey DepIndexKey(*DepIndexKeyString);
-			UE::Interchange::FAttributeStorage::TAttributeHandle<FName> HandleMesh = Attributes.GetAttributeHandle<FName>(DepIndexKey);
-			if (!HandleMesh.IsValid())
-			{
-				continue;
-			}
-			FName MeshUniqueID;
-			HandleMesh.Get(MeshUniqueID);
-			if (MeshUniqueID == LodDataUniqueId)
-			{
-				//Remove this entry
-				Attributes.UnregisterAttribute(DepIndexKey);
-				Handle.Set(LodCount - 1);
-				//We have to rename the key for all the next item
-				DecrementKey = true;;
-			}
-			else if (DecrementKey)
-			{
-				FString NewDepIndexKeyString = UE::Interchange::FSkeletalMeshNodeStaticData::LodDataIdsBaseKey().ToString() + FString::FromInt(LodIndex - 1);
-				UE::Interchange::FAttributeKey NewDepIndexKey(*NewDepIndexKeyString);
-				UE::Interchange::EAttributeStorageResult UnregisterResult = Attributes.UnregisterAttribute(DepIndexKey);
-				if (UE::Interchange::IsAttributeStorageResultSuccess(UnregisterResult))
-				{
-					UE::Interchange::EAttributeStorageResult RegisterResult = Attributes.RegisterAttribute<FName>(NewDepIndexKey, MeshUniqueID);
-					if (!UE::Interchange::IsAttributeStorageResultSuccess(RegisterResult))
-					{
-						UE::Interchange::LogAttributeStorageErrors(RegisterResult, TEXT("UInterchangeSkeletalMeshNode.RemoveLodDataUniqueId"), NewDepIndexKey);
-					}
-				}
-				else
-				{
-					UE::Interchange::LogAttributeStorageErrors(UnregisterResult, TEXT("UInterchangeSkeletalMeshNode.RemoveLodDataUniqueId"), DepIndexKey);
-				}
-
-				//Avoid doing more code in the for since the HandleMesh is now invalid
-				continue;
-			}
-		}
-		return true;
+		return LodDependencies.RemoveName(LodDataUniqueId);
 	}
+
 private:
 
 	void FillAssetClassFromAttribute()
@@ -251,7 +133,7 @@ private:
 #if WITH_ENGINE
 		FString OperationName = GetTypeName() + TEXT(".GetAssetClassName");
 		FString ClassName;
-		InterchangePrivateNodeBase::GetCustomAttribute<FString>(Attributes, ClassNameAttributeKey, OperationName, ClassName);
+		InterchangePrivateNodeBase::GetCustomAttribute<FString>(*Attributes, ClassNameAttributeKey, OperationName, ClassName);
 		if (ClassName.Equals(USkeletalMesh::StaticClass()->GetName()))
 		{
 			AssetClass = USkeletalMesh::StaticClass();
@@ -280,6 +162,7 @@ private:
 
 	const UE::Interchange::FAttributeKey ClassNameAttributeKey = UE::Interchange::FAttributeKey(TEXT("__ClassTypeAttribute__"));
 
+	UE::Interchange::FNameAttributeArrayHelper LodDependencies;
 protected:
 #if WITH_ENGINE
 	TSubclassOf<USkeletalMesh> AssetClass = nullptr;
diff --git a/Engine/Plugins/Runtime/Interchange/Source/Nodes/Public/InterchangeSkeletonNode.h b/Engine/Plugins/Runtime/Interchange/Source/Nodes/Public/InterchangeSkeletonNode.h
--- a/Engine/Plugins/Runtime/Interchange/Source/Nodes/Public/InterchangeSkeletonNode.h
+++ b/Engine/Plugins/Runtime/Interchange/Source/Nodes/Public/InterchangeSkeletonNode.h
@@ -37,7 +37,7 @@ public:
 		bIsNodeClassInitialized = false;
 		InitializeNode(UniqueID, DisplayLabel);
 		FString OperationName = GetTypeName() + TEXT(".SetAssetClassName");
-		InterchangePrivateNodeBase::SetCustomAttribute<FString>(Attributes, ClassNameAttributeKey, OperationName, InAssetClass);
+		InterchangePrivateNodeBase::SetCustomAttribute<FString>(*Attributes, ClassNameAttributeKey, OperationName, InAssetClass);
 		FillAssetClassFromAttribute();
 	}
 
@@ -77,7 +77,7 @@ public:
 
 	virtual FGuid GetHash() const override
 	{
-		return Attributes.GetStorageHash();
+		return Attributes->GetStorageHash();
 	}
 
 public:
@@ -101,7 +101,7 @@ private:
 #if WITH_ENGINE
 		FString OperationName = GetTypeName() + TEXT(".GetAssetClassName");
 		FString ClassName;
-		InterchangePrivateNodeBase::GetCustomAttribute<FString>(Attributes, ClassNameAttributeKey, OperationName, ClassName);
+		InterchangePrivateNodeBase::GetCustomAttribute<FString>(*Attributes, ClassNameAttributeKey, OperationName, ClassName);
 		if (ClassName.Equals(USkeleton::StaticClass()->GetName()))
 		{
 			AssetClass = USkeleton::StaticClass();
diff --git a/Engine/Plugins/Runtime/Interchange/Source/Nodes/Public/InterchangeTextureNode.h b/Engine/Plugins/Runtime/Interchange/Source/Nodes/Public/InterchangeTextureNode.h
--- a/Engine/Plugins/Runtime/Interchange/Source/Nodes/Public/InterchangeTextureNode.h
+++ b/Engine/Plugins/Runtime/Interchange/Source/Nodes/Public/InterchangeTextureNode.h
@@ -56,7 +56,7 @@ public:
 		bIsTextureNodeClassInitialized = false;
 		InitializeNode(UniqueID, DisplayLabel);
 		FString OperationName = GetTypeName() + TEXT(".SetAssetClassName");
-		InterchangePrivateNodeBase::SetCustomAttribute<FString>(Attributes, ClassNameAttributeKey, OperationName, InAssetClass);
+		InterchangePrivateNodeBase::SetCustomAttribute<FString>(*Attributes, ClassNameAttributeKey, OperationName, InAssetClass);
 		FillAssetClassFromAttribute();
 	}
 
@@ -96,17 +96,17 @@ public:
 
 	virtual FGuid GetHash() const override
 	{
-		return Attributes.GetStorageHash();
+		return Attributes->GetStorageHash();
 	}
 
 	/** Texture node Interface Begin */
 	virtual const TOptional<FString> GetPayLoadKey() const
 	{
-		if (!Attributes.ContainAttribute(UE::Interchange::FTextureNodeStaticData::PayloadSourceFileKey()))
+		if (!Attributes->ContainAttribute(UE::Interchange::FTextureNodeStaticData::PayloadSourceFileKey()))
 		{
 			return TOptional<FString>();
 		}
-		UE::Interchange::FAttributeStorage::TAttributeHandle<FString> AttributeHandle = Attributes.GetAttributeHandle<FString>(UE::Interchange::FTextureNodeStaticData::PayloadSourceFileKey());
+		UE::Interchange::FAttributeStorage::TAttributeHandle<FString> AttributeHandle = Attributes->GetAttributeHandle<FString>(UE::Interchange::FTextureNodeStaticData::PayloadSourceFileKey());
 		if (!AttributeHandle.IsValid())
 		{
 			return TOptional<FString>();
@@ -124,7 +124,7 @@ public:
 	UFUNCTION(BlueprintCallable, Category = "Interchange | Node | Texture")
 	virtual void SetPayLoadKey(const FString& PayloadKey)
 	{
-		UE::Interchange::EAttributeStorageResult Result = Attributes.RegisterAttribute(UE::Interchange::FTextureNodeStaticData::PayloadSourceFileKey(), PayloadKey);
+		UE::Interchange::EAttributeStorageResult Result = Attributes->RegisterAttribute(UE::Interchange::FTextureNodeStaticData::PayloadSourceFileKey(), PayloadKey);
 		if (!IsAttributeStorageResultSuccess(Result))
 		{
 			LogAttributeStorageErrors(Result, TEXT("UInterchangeTextureNode.SetPayLoadKey"), UE::Interchange::FTextureNodeStaticData::PayloadSourceFileKey());
@@ -755,7 +755,7 @@ private:
 #if WITH_ENGINE
 		FString OperationName = GetTypeName() + TEXT(".GetAssetClassName");
 		FString ClassName;
-		InterchangePrivateNodeBase::GetCustomAttribute<FString>(Attributes, ClassNameAttributeKey, OperationName, ClassName);
+		InterchangePrivateNodeBase::GetCustomAttribute<FString>(*Attributes, ClassNameAttributeKey, OperationName, ClassName);
 		if (ClassName.Equals(UTexture2D::StaticClass()->GetName()))
 		{
 			AssetClass = UTexture2D::StaticClass();
diff --git a/Engine/Plugins/Runtime/Interchange/Source/Parsers/Fbx/Private/FbxMaterial.cpp b/Engine/Plugins/Runtime/Interchange/Source/Parsers/Fbx/Private/FbxMaterial.cpp
--- a/Engine/Plugins/Runtime/Interchange/Source/Parsers/Fbx/Private/FbxMaterial.cpp
+++ b/Engine/Plugins/Runtime/Interchange/Source/Parsers/Fbx/Private/FbxMaterial.cpp
@@ -76,6 +76,7 @@ namespace UE
 							int32 UnsupportedTextureCount = FbxProperty.GetSrcObjectCount<FbxLayeredTexture>();
 							UnsupportedTextureCount += FbxProperty.GetSrcObjectCount<FbxProceduralTexture>();
 							int32 TextureCount = FbxProperty.GetSrcObjectCount<FbxTexture>();
+							bool bFoundValidTexture = false;
 							if (UnsupportedTextureCount > 0)
 							{
 								JSonErrorMessages.Add(TEXT("{\"Msg\" : {\"Type\" : \"Warning\",\n\"Msg\" : \"Layered or procedural Textures are not supported (material" + MaterialName + TEXT(")\"}}")));
@@ -109,9 +110,11 @@ namespace UE
 									MaterialNode->AddTextureParameterData(MaterialParameterName, TextureNode->GetUniqueID(), UVChannelIndex, ScaleU, ScaleV);
 									MaterialNode->SetDependencyUID(TextureNode->GetUniqueID());
 									bSetMaterial = true;
+									bFoundValidTexture = true;
 								}
 							}
-							else if (MaterialParameterName == EInterchangeMaterialNodeParameterName::BaseColor)
+							
+							if (!bFoundValidTexture && MaterialParameterName == EInterchangeMaterialNodeParameterName::BaseColor)
 							{
 								//We support only the base color has a vector color
 								//TODO support all basic attributes has vector or scalar
diff --git a/Engine/Plugins/Runtime/Interchange/Source/Parsers/Fbx/Private/FbxMesh.cpp b/Engine/Plugins/Runtime/Interchange/Source/Parsers/Fbx/Private/FbxMesh.cpp
--- a/Engine/Plugins/Runtime/Interchange/Source/Parsers/Fbx/Private/FbxMesh.cpp
+++ b/Engine/Plugins/Runtime/Interchange/Source/Parsers/Fbx/Private/FbxMesh.cpp
@@ -224,6 +224,81 @@ namespace UE
 				return FillMeshDescriptionFromFbxMesh(MeshType, SortedJoints);
 			}
 
+			bool FMeshDescriptionImporter::FillMeshDescriptionFromFbxShape(FbxShape* Shape)
+			{
+				//For the shape we just need the vertex positions
+				if (!ensure(bInitialized) || !ensure(MeshDescription))
+				{
+					return false;
+				}
+				FStaticMeshAttributes Attributes(*MeshDescription);
+				FbxMesh* Mesh = MeshNode->GetMesh();
+
+				//The shape should be for a specified mesh
+				if (!ensure(Mesh))
+				{
+					return false;
+				}
+
+				//Extract the points into a simplified MeshDescription
+				{
+					TRACE_CPUPROFILER_EVENT_SCOPE(GatherMorphTargetPoints);
+					MeshDescription->SuspendVertexIndexing();
+
+					// Construct the matrices for the conversion from right handed to left handed system
+					FbxAMatrix TotalMatrix;
+					FbxAMatrix TotalMatrixForNormal;
+					TotalMatrix = ComputeNodeMatrix(MeshNode);
+					TotalMatrixForNormal = TotalMatrix.Inverse();
+					TotalMatrixForNormal = TotalMatrixForNormal.Transpose();
+
+					FbxGeometryBase* GeoBase = (Shape != nullptr ? (FbxGeometryBase*)Shape : (FbxGeometryBase*)Mesh);
+
+					int32 VertexCount = GeoBase->GetControlPointsCount();
+					if (Shape)
+					{
+						if (!ensure(Mesh->GetControlPointsCount() == VertexCount))
+						{
+							//TODO warn the user
+							return false;
+						}
+					}
+
+					TVertexAttributesRef<FVector> VertexPositions = Attributes.GetVertexPositions();
+					int32 VertexOffset = MeshDescription->Vertices().Num();
+					// The below code expects Num() to be equivalent to GetArraySize(), i.e. that all added elements are appended, not inserted into existing gaps
+					check(VertexOffset == MeshDescription->Vertices().GetArraySize());
+					//Fill the vertex array
+					auto GetFinalPosition = [&TotalMatrix](FbxVector4& FbxPosition)
+					{
+						FbxPosition = TotalMatrix.MultT(FbxPosition);
+						return FFbxConvert::ConvertPos(FbxPosition);
+					};
+
+					MeshDescription->ReserveNewVertices(VertexCount);
+					for (int32 VertexIndex = 0; VertexIndex < VertexCount; ++VertexIndex)
+					{
+						int32 RealVertexIndex = VertexOffset + VertexIndex;
+						const FVector VertexPosition = GetFinalPosition(GeoBase->GetControlPoints()[VertexIndex]);
+						//Maybe we want to do some work here like computing the deltas
+						//const FVector MeshVertexPosition = GetFinalPosition(Mesh->GetControlPoints()[VertexIndex]);
+						FVertexID AddedVertexId = MeshDescription->CreateVertex();
+						if (AddedVertexId.GetValue() != RealVertexIndex)
+						{
+							//TODO add an error log
+							//AddTokenizedErrorMessage(FTokenizedMessage::Create(EMessageSeverity::Error, FText::Format(LOCTEXT("Error_CannotCreateVertex", "Cannot create valid vertex for mesh '{0}'"), FText::FromString(Mesh->GetName()))), FFbxErrors::StaticMesh_BuildError);
+							return false;
+						}
+						//Add the delta position, so we do not have to recompute it later
+						VertexPositions[AddedVertexId] = VertexPosition;// -MeshVertexPosition;
+					}
+					MeshDescription->ResumeVertexIndexing();
+					FElementIDRemappings OutRemappings;
+					MeshDescription->Compact(OutRemappings);
+				}
+				return true;
+			}
+
 			bool FMeshDescriptionImporter::FillMeshDescriptionFromFbxMesh(EMeshType MeshType, TArray<FbxNode*>* SortedJoints)
 			{
 				if (!ensure(bInitialized) || !ensure(MeshDescription))
diff --git a/Engine/Plugins/Runtime/Interchange/Source/Parsers/Fbx/Private/FbxMesh.h b/Engine/Plugins/Runtime/Interchange/Source/Parsers/Fbx/Private/FbxMesh.h
--- a/Engine/Plugins/Runtime/Interchange/Source/Parsers/Fbx/Private/FbxMesh.h
+++ b/Engine/Plugins/Runtime/Interchange/Source/Parsers/Fbx/Private/FbxMesh.h
@@ -22,6 +22,7 @@ namespace UE
 				FMeshDescriptionImporter(FMeshDescription* InMeshDescription, FbxNode* InMeshNode, FbxScene* InSDKScene, FbxGeometryConverter* InSDKGeometryConverter);
 				bool FillStaticMeshDescriptionFromFbxMesh();
 				bool FillSkinnedMeshDescriptionFromFbxMesh(TArray<FbxNode*>* SortedJoints);
+				bool FillMeshDescriptionFromFbxShape(FbxShape* Shape);
 			private:
 				
 				enum class EMeshType : uint8
diff --git a/Engine/Plugins/Runtime/Interchange/Source/Parsers/Fbx/Private/FbxSkeletalMesh.cpp b/Engine/Plugins/Runtime/Interchange/Source/Parsers/Fbx/Private/FbxSkeletalMesh.cpp
--- a/Engine/Plugins/Runtime/Interchange/Source/Parsers/Fbx/Private/FbxSkeletalMesh.cpp
+++ b/Engine/Plugins/Runtime/Interchange/Source/Parsers/Fbx/Private/FbxSkeletalMesh.cpp
@@ -50,22 +50,41 @@ namespace UE
 				{
 					return false;
 				}
-				FbxSkin* Skin = (FbxSkin*)Mesh->GetDeformer(0, FbxDeformer::eSkin);
-				if (!ensure(Skin))
+
+				FMeshDescriptionImporter MeshDescriptionImporter(&SkeletalMeshDescription, Node, SDKScene, SDKGeometryConverter);
+				if (!MeshDescriptionImporter.FillSkinnedMeshDescriptionFromFbxMesh(&SortedJoints))
 				{
 					return false;
 				}
 
-				FMeshDescriptionImporter MeshDescriptionImporter(&SkeletalMeshDescription, Node, SDKScene, SDKGeometryConverter);
-				if (!MeshDescriptionImporter.FillSkinnedMeshDescriptionFromFbxMesh(&SortedJoints))
+				return true;
+			}
+
+			
+
+			bool FSkeletalMeshGeometryPayload::AddFbxShapeToMeshDescription(int32 MeshIndex, FbxShape* Shape, FMeshDescription& StaticMeshDescription, FStaticMeshAttributes& StaticMeshAttribute, TArray<FString>& JSonErrorMessages)
+			{
+				if (!ensure(SDKScene) || !ensure(SDKGeometryConverter))
+				{
+					return false;
+				}
+
+				if (!ensure(SkelMeshNodeArray.IsValidIndex(MeshIndex)))
 				{
 					return false;
 				}
 
+				FbxNode* Node = SkelMeshNodeArray[MeshIndex];
+				if (!ensure(Node))
+				{
+					return false;
+				}
 
-				//////////////////////////////////////////////////////////////////////////
-				//////////////////////////////////////////////////////////////////////////
-				//////////////////////////////////////////////////////////////////////////
+				FMeshDescriptionImporter MeshDescriptionImporter(&StaticMeshDescription, Node, SDKScene, SDKGeometryConverter);
+				if (!MeshDescriptionImporter.FillMeshDescriptionFromFbxShape(Shape))
+				{
+					return false;
+				}
 
 				return true;
 			}
@@ -98,10 +117,40 @@ namespace UE
 					}
 				}
 
+				//Import the BlendShape
+				TMap<FString, FMeshDescription> BlendShapes;
+				if (ShapeNameToShapeArray.Num() > 0)
+				{
+					for (const TPair<FString, TArray<FbxShape*>>& Pair : ShapeNameToShapeArray)
+					{
+						const TArray<FbxShape*> FbxShapes = Pair.Value;
+						FMeshDescription& BlendShapeDescription = BlendShapes.FindOrAdd(Pair.Key);
+						FStaticMeshAttributes StaticMeshAttribute(BlendShapeDescription);
+						StaticMeshAttribute.Register();
+						ensure(FbxShapes.Num() == SkelMeshNodeArray.Num());
+						for (int32 MeshIndex = 0; MeshIndex < SkelMeshNodeArray.Num(); ++MeshIndex)
+						{
+							if (!FbxShapes.IsValidIndex(MeshIndex))
+							{
+								//TODO warn the user
+								break;
+							}
+							FbxShape* Shape = FbxShapes[MeshIndex];
+							if (!AddFbxShapeToMeshDescription(MeshIndex, Shape, BlendShapeDescription, StaticMeshAttribute, JSonErrorMessages))
+							{
+								continue;
+							}
+						}
+					}
+				}
 				//Dump the MeshDescription to a file
 				{
 					FLargeMemoryWriter Ar;
 					SkeletalMeshDescription.Serialize(Ar);
+
+					//Add the LOD BlendShape Array to the archive
+					Ar << BlendShapes;
+
 					uint8* ArchiveData = Ar.GetData();
 					int64 ArchiveSize = Ar.TotalSize();
 					TArray64<uint8> Buffer(ArchiveData, ArchiveSize);
@@ -134,7 +183,7 @@ namespace UE
 							}
 						}
 					}
-					
+
 					UInterchangeSkeletalMeshNode* SkeletalMeshNode = nullptr;
 
 					bool bOverrideSkeletalMeshName = false;
@@ -147,7 +196,10 @@ namespace UE
 
 					for (int32 LODIndex = 0; LODIndex < MaxNumberOfLOD; LODIndex++)
 					{
+						//Array use to retrieve the payload data, so we easily find the request payload data
 						TArray<FbxNode*> SkelMeshNodeArray;
+						TArray<FbxNode*> SortedJoints;
+						TMap<FString, TArray<FbxShape*>> ShapeNameToShapeArray;
 						for (int32 j = 0; j < NodeArray.Num(); j++)
 						{
 							FbxNode* Node = NodeArray[j];
@@ -179,7 +231,7 @@ namespace UE
 							continue;
 						}
 						//Find the skeleton for this skeletal mesh LOD
-						TArray<FbxNode*> SortedJoints;
+						
 						FbxArray<FbxAMatrix> LocalsPerLink;
 						if (!FFbxHelper::FindSkeletonJoints(SDKScene, SkelMeshNodeArray, SortedJoints, LocalsPerLink))
 						{
@@ -229,9 +281,9 @@ namespace UE
 						//We will use this information to retrieve the payload
 						FSHA1 Sha;
 						
-
-						for (FbxNode* MeshNode : SkelMeshNodeArray)
+						for (int32 NodeIndex = 0; NodeIndex < SkelMeshNodeArray.Num(); ++NodeIndex)
 						{
+							FbxNode* MeshNode = SkelMeshNodeArray[NodeIndex];
 							FString MeshNodePath = FFbxHelper::GetFbxNodeHierarchyName(MeshNode);
 						
 							int32 NodeMaterialCount = MeshNode->GetMaterialCount();
@@ -255,6 +307,69 @@ namespace UE
 
 							TArray<TCHAR> IDArray = MeshNodePath.GetCharArray();
 							Sha.Update((uint8*)IDArray.GetData(), IDArray.Num()* IDArray.GetTypeSize());
+
+							//Add the Blendshape names to the node
+							FbxGeometry* Geometry = (FbxGeometry*)MeshNode->GetNodeAttribute();
+							if (Geometry)
+							{
+								const int32 BlendShapeDeformerCount = Geometry->GetDeformerCount(FbxDeformer::eBlendShape);
+								if (BlendShapeDeformerCount > 0)
+								{
+									for (int32 BlendShapeIndex = 0; BlendShapeIndex < BlendShapeDeformerCount; ++BlendShapeIndex)
+									{
+										FbxBlendShape* BlendShape = (FbxBlendShape*)Geometry->GetDeformer(BlendShapeIndex, FbxDeformer::eBlendShape);
+										const int32 BlendShapeChannelCount = BlendShape->GetBlendShapeChannelCount();
+										FString BlendShapeName = FFbxHelper::GetFbxObjectName(BlendShape);
+										// see below where this is used for explanation...
+										const bool bMightBeBadMAXFile = (BlendShapeName == FString("Morpher"));
+										for (int32 ChannelIndex = 0; ChannelIndex < BlendShapeChannelCount; ++ChannelIndex)
+										{
+											FbxBlendShapeChannel* Channel = BlendShape->GetBlendShapeChannel(ChannelIndex);
+											if (Channel)
+											{
+												//Find which shape should we use according to the weight.
+												const int32 CurrentChannelShapeCount = Channel->GetTargetShapeCount();
+												FString ChannelName = FFbxHelper::GetFbxObjectName(Channel);
+												// Maya adds the name of the blendshape and an underscore to the front of the channel name, so remove it
+												if (ChannelName.StartsWith(BlendShapeName))
+												{
+													ChannelName.RightInline(ChannelName.Len() - (BlendShapeName.Len() + 1), false);
+												}
+												for (int32 ShapeIndex = 0; ShapeIndex < CurrentChannelShapeCount; ++ShapeIndex)
+												{
+													FbxShape* Shape = Channel->GetTargetShape(ShapeIndex);
+													FString ShapeName;
+													if (CurrentChannelShapeCount > 1)
+													{
+														ShapeName = FFbxHelper::GetFbxObjectName(Shape);
+													}
+													else
+													{
+														if (bMightBeBadMAXFile)
+														{
+															ShapeName = FFbxHelper::GetFbxObjectName(Shape);
+														}
+														else
+														{
+															// Maya concatenates the number of the shape to the end of its name, so instead use the name of the channel
+															ShapeName = ChannelName;
+														}
+													}
+
+													TArray<FbxShape*>& ShapeArray = ShapeNameToShapeArray.FindOrAdd(ShapeName);
+													if (ShapeArray.Num() == 0)
+													{
+														//Add the Shape name to the LodDataNode so we can retrieve the shapes when someone query the payload
+														LodDataNode->AddBlendShape(*ShapeName);
+														ShapeArray.AddZeroed(SkelMeshNodeArray.Num());
+													}
+													ShapeArray[NodeIndex] = Shape;
+												}
+											}
+										}
+									}
+								}
+							}
 						}
 						Sha.Final();
 						// Retrieve the hash and use it to construct a pseudo-GUID. 
@@ -267,6 +382,7 @@ namespace UE
 							TSharedPtr<FSkeletalMeshGeometryPayload> GeoPayload = MakeShared<FSkeletalMeshGeometryPayload>();
 							GeoPayload->SkelMeshNodeArray = SkelMeshNodeArray;
 							GeoPayload->SortedJoints = SortedJoints;
+							GeoPayload->ShapeNameToShapeArray = ShapeNameToShapeArray;
 							GeoPayload->SDKScene = SDKScene;
 							GeoPayload->SDKGeometryConverter = SDKGeometryConverter;
 							PayloadContexts.Add(PayloadKey, GeoPayload);
diff --git a/Engine/Plugins/Runtime/Interchange/Source/Parsers/Fbx/Private/FbxSkeletalMesh.h b/Engine/Plugins/Runtime/Interchange/Source/Parsers/Fbx/Private/FbxSkeletalMesh.h
--- a/Engine/Plugins/Runtime/Interchange/Source/Parsers/Fbx/Private/FbxSkeletalMesh.h
+++ b/Engine/Plugins/Runtime/Interchange/Source/Parsers/Fbx/Private/FbxSkeletalMesh.h
@@ -8,6 +8,7 @@
 
 /** Forward declarations */
 class FSkeletalMeshAttributes;
+class FStaticMeshAttributes;
 class UInterchangeBaseNode;
 class UInterchangeBaseNodeContainer;
 class UInterchangeSkeletalMeshNode;
@@ -30,11 +31,13 @@ namespace UE
 				virtual bool FetchPayloadToFile(const FString& PayloadFilepath, TArray<FString>& JSonErrorMessages) override;
 				TArray<FbxNode*> SkelMeshNodeArray;
 				TArray<FbxNode*> SortedJoints;
+				TMap<FString, TArray<FbxShape*>> ShapeNameToShapeArray;
 				FbxScene* SDKScene = nullptr;
 				FbxGeometryConverter* SDKGeometryConverter = nullptr;
 
 			private:
 				bool AddFbxMeshToMeshDescription(int32 MeshIndex, FMeshDescription& SkeletalMeshDescription, FSkeletalMeshAttributes& SkeletalMeshAttribute, TArray<FString>& JSonErrorMessages);
+				bool AddFbxShapeToMeshDescription(int32 MeshIndex, FbxShape* Shape, FMeshDescription& StaticMeshDescription, FStaticMeshAttributes& StaticMeshAttribute, TArray<FString>& JSonErrorMessages);
 			};
 
 			class FFbxSkeletalMesh
diff --git a/Engine/Source/Runtime/Interchange/Core/Private/Nodes/InterchangeBaseNode.cpp b/Engine/Source/Runtime/Interchange/Core/Private/Nodes/InterchangeBaseNode.cpp
--- a/Engine/Source/Runtime/Interchange/Core/Private/Nodes/InterchangeBaseNode.cpp
+++ b/Engine/Source/Runtime/Interchange/Core/Private/Nodes/InterchangeBaseNode.cpp
@@ -12,17 +12,20 @@
 
 void UInterchangeBaseNode::InitializeNode(const FName& UniqueID, const FName& DisplayLabel)
 {
-	UE::Interchange::EAttributeStorageResult Result = Attributes.RegisterAttribute(UE::Interchange::FBaseNodeStaticData::UniqueIDKey(), UniqueID, UE::Interchange::EAttributeProperty::NoHash);
+	UE::Interchange::EAttributeStorageResult Result = Attributes->RegisterAttribute(UE::Interchange::FBaseNodeStaticData::UniqueIDKey(), UniqueID, UE::Interchange::EAttributeProperty::NoHash);
 	if (!IsAttributeStorageResultSuccess(Result))
 	{
 		LogAttributeStorageErrors(Result, TEXT("RegisterAttribute"), UE::Interchange::FBaseNodeStaticData::UniqueIDKey());
 	}
 
-	Result = Attributes.RegisterAttribute(UE::Interchange::FBaseNodeStaticData::DisplayLabelKey(), DisplayLabel, UE::Interchange::EAttributeProperty::NoHash);
+	Result = Attributes->RegisterAttribute(UE::Interchange::FBaseNodeStaticData::DisplayLabelKey(), DisplayLabel, UE::Interchange::EAttributeProperty::NoHash);
 	if (!IsAttributeStorageResultSuccess(Result))
 	{
 		LogAttributeStorageErrors(Result, TEXT("RegisterAttribute"), UE::Interchange::FBaseNodeStaticData::DisplayLabelKey());
 	}
+
+	Dependencies.Initialize(Attributes, UE::Interchange::FBaseNodeStaticData::GetDependenciesBaseKey());
+
 	bIsInitialized = true;
 }
 
@@ -34,7 +37,7 @@ FString UInterchangeBaseNode::GetTypeName() const
 
 bool UInterchangeBaseNode::HasAttribute(const UE::Interchange::FAttributeKey& NodeAttributeKey) const
 {
-	return Attributes.ContainAttribute(NodeAttributeKey);
+	return Attributes->ContainAttribute(NodeAttributeKey);
 }
 
 
@@ -42,25 +45,25 @@ FName UInterchangeBaseNode::GetUniqueID() const
 {
 	ensure(bIsInitialized);
 	FName UniqueID = NAME_None;
-	Attributes.GetAttributeHandle<FName>(UE::Interchange::FBaseNodeStaticData::UniqueIDKey()).Get(UniqueID);
+	Attributes->GetAttributeHandle<FName>(UE::Interchange::FBaseNodeStaticData::UniqueIDKey()).Get(UniqueID);
 	return UniqueID;
 }
 
 FName UInterchangeBaseNode::GetDisplayLabel() const
 {
 	ensure(bIsInitialized);
-	checkSlow(Attributes.ContainAttribute(UE::Interchange::FBaseNodeStaticData::DisplayLabelKey()));
+	checkSlow(Attributes->ContainAttribute(UE::Interchange::FBaseNodeStaticData::DisplayLabelKey()));
 	FName DisplayLabel = NAME_None;
-	Attributes.GetAttributeHandle<FName>(UE::Interchange::FBaseNodeStaticData::DisplayLabelKey()).Get(DisplayLabel);
+	Attributes->GetAttributeHandle<FName>(UE::Interchange::FBaseNodeStaticData::DisplayLabelKey()).Get(DisplayLabel);
 	return DisplayLabel;
 }
 
 bool UInterchangeBaseNode::SetDisplayLabel(FName DisplayLabel)
 {
-	UE::Interchange::EAttributeStorageResult Result = Attributes.RegisterAttribute(UE::Interchange::FBaseNodeStaticData::DisplayLabelKey(), DisplayLabel);
+	UE::Interchange::EAttributeStorageResult Result = Attributes->RegisterAttribute(UE::Interchange::FBaseNodeStaticData::DisplayLabelKey(), DisplayLabel);
 	if (IsAttributeStorageResultSuccess(Result))
 	{
-		UE::Interchange::FAttributeStorage::TAttributeHandle<FName> Handle = Attributes.GetAttributeHandle<FName>(UE::Interchange::FBaseNodeStaticData::DisplayLabelKey());
+		UE::Interchange::FAttributeStorage::TAttributeHandle<FName> Handle = Attributes->GetAttributeHandle<FName>(UE::Interchange::FBaseNodeStaticData::DisplayLabelKey());
 		return Handle.IsValid();
 	}
 	return false;
@@ -68,13 +71,13 @@ bool UInterchangeBaseNode::SetDisplayLabel(FName DisplayLabel)
 
 FName UInterchangeBaseNode::GetParentUID() const
 {
-	if (!Attributes.ContainAttribute(UE::Interchange::FBaseNodeStaticData::ParentIDKey()))
+	if (!Attributes->ContainAttribute(UE::Interchange::FBaseNodeStaticData::ParentIDKey()))
 	{
 		return InvalidNodeUID();
 	}
 
 	FName ParentUniqueID = NAME_None;
-	UE::Interchange::FAttributeStorage::TAttributeHandle<FName> Handle = Attributes.GetAttributeHandle<FName>(UE::Interchange::FBaseNodeStaticData::ParentIDKey());
+	UE::Interchange::FAttributeStorage::TAttributeHandle<FName> Handle = Attributes->GetAttributeHandle<FName>(UE::Interchange::FBaseNodeStaticData::ParentIDKey());
 	if(Handle.IsValid())
 	{
 		Handle.Get(ParentUniqueID);
@@ -85,71 +88,38 @@ FName UInterchangeBaseNode::GetParentUID() const
 
 bool UInterchangeBaseNode::SetParentUID(FName ParentUID)
 {
-	UE::Interchange::EAttributeStorageResult Result = Attributes.RegisterAttribute(UE::Interchange::FBaseNodeStaticData::ParentIDKey(), ParentUID);
+	UE::Interchange::EAttributeStorageResult Result = Attributes->RegisterAttribute(UE::Interchange::FBaseNodeStaticData::ParentIDKey(), ParentUID);
 	if(IsAttributeStorageResultSuccess(Result))
 	{
-		UE::Interchange::FAttributeStorage::TAttributeHandle<FName> Handle = Attributes.GetAttributeHandle<FName>(UE::Interchange::FBaseNodeStaticData::ParentIDKey());
+		UE::Interchange::FAttributeStorage::TAttributeHandle<FName> Handle = Attributes->GetAttributeHandle<FName>(UE::Interchange::FBaseNodeStaticData::ParentIDKey());
 		return Handle.IsValid();
 	}
 	return false;
 }
 
-void UInterchangeBaseNode::GetDependecies(TArray<FName>& OutDependencies ) const
+int32 UInterchangeBaseNode::GetDependeciesCount() const
 {
-	OutDependencies.Reset();
+	return Dependencies.GetCount();
+}
 
-	int32 DepenedenciesCount = 0;
-	if (!Attributes.ContainAttribute(UE::Interchange::FBaseNodeStaticData::DependencyCountKey()))
-	{
-		return;
-	}
-	UE::Interchange::FAttributeStorage::TAttributeHandle<int32> Handle = Attributes.GetAttributeHandle<int32>(UE::Interchange::FBaseNodeStaticData::DependencyCountKey());
-	if (!Handle.IsValid())
-	{
-		return;
-	}
-	Handle.Get(DepenedenciesCount);
-	for (int32 DepIndex = 0; DepIndex < DepenedenciesCount; ++DepIndex)
-	{
-		FString DepIndexKeyString = UE::Interchange::FBaseNodeStaticData::DependencyBaseKey().Key.ToString() + FString::FromInt(DepIndex);
-		UE::Interchange::FAttributeKey DepIndexKey(*DepIndexKeyString);
-		UE::Interchange::FAttributeStorage::TAttributeHandle<FName> HandleDep = Attributes.GetAttributeHandle<FName>(DepIndexKey);
-		if (!HandleDep.IsValid())
-		{
-			continue;
-		}
-		FName& NodeUniqueID = OutDependencies.AddDefaulted_GetRef();
-		HandleDep.Get(NodeUniqueID);
-	}
+void UInterchangeBaseNode::GetDependecies(TArray<FName>& OutDependencies ) const
+{
+	Dependencies.GetNames(OutDependencies);
 }
 
 bool UInterchangeBaseNode::SetDependencyUID(FName DependencyUID)
 {
-	if (!Attributes.ContainAttribute(UE::Interchange::FBaseNodeStaticData::DependencyCountKey()))
-	{
-		const int32 DependencyCount = 0;
-		UE::Interchange::EAttributeStorageResult Result = Attributes.RegisterAttribute<int32>(UE::Interchange::FBaseNodeStaticData::DependencyCountKey(), DependencyCount);
-	}
-	UE::Interchange::FAttributeStorage::TAttributeHandle<int32> Handle = Attributes.GetAttributeHandle<int32>(UE::Interchange::FBaseNodeStaticData::DependencyCountKey());
-	if(!ensure(Handle.IsValid()))
-	{
-		return false;
-	}
-	int32 DepIndex = 0;
-	Handle.Get(DepIndex);
-	FString DepIndexKeyString = UE::Interchange::FBaseNodeStaticData::DependencyBaseKey().Key.ToString() + FString::FromInt(DepIndex);
-	UE::Interchange::FAttributeKey DepIndexKey(*DepIndexKeyString);
-	//Increment the counter
-	DepIndex++;
-	Handle.Set(DepIndex);
-
-	UE::Interchange::EAttributeStorageResult DepResult = Attributes.RegisterAttribute<FName>(DepIndexKey, DependencyUID);
-	return true;
+	return Dependencies.AddName(DependencyUID);
+}
+
+bool UInterchangeBaseNode::RemoveDependencyUID(FName DependencyUID)
+{
+	return Dependencies.RemoveName(DependencyUID);
 }
 
 bool UInterchangeBaseNode::IsEnabled() const
 {
-	UE::Interchange::FAttributeStorage::TAttributeHandle<bool> Handle = Attributes.GetAttributeHandle<bool>(UE::Interchange::FBaseNodeStaticData::IsEnabledKey());
+	UE::Interchange::FAttributeStorage::TAttributeHandle<bool> Handle = Attributes->GetAttributeHandle<bool>(UE::Interchange::FBaseNodeStaticData::IsEnabledKey());
 	if (Handle.IsValid())
 	{
 		bool bIsEnabled = false;
@@ -161,10 +131,10 @@ bool UInterchangeBaseNode::IsEnabled() const
 
 bool UInterchangeBaseNode::SetEnabled(const bool bIsEnabled)
 {
-	UE::Interchange::EAttributeStorageResult Result = Attributes.RegisterAttribute(UE::Interchange::FBaseNodeStaticData::IsEnabledKey(), bIsEnabled);
+	UE::Interchange::EAttributeStorageResult Result = Attributes->RegisterAttribute(UE::Interchange::FBaseNodeStaticData::IsEnabledKey(), bIsEnabled);
 	if (IsAttributeStorageResultSuccess(Result))
 	{
-		UE::Interchange::FAttributeStorage::TAttributeHandle<bool> Handle = Attributes.GetAttributeHandle<bool>(UE::Interchange::FBaseNodeStaticData::IsEnabledKey());
+		UE::Interchange::FAttributeStorage::TAttributeHandle<bool> Handle = Attributes->GetAttributeHandle<bool>(UE::Interchange::FBaseNodeStaticData::IsEnabledKey());
 		return Handle.IsValid();
 	}
 	return false;
@@ -172,7 +142,7 @@ bool UInterchangeBaseNode::SetEnabled(const bool bIsEnabled)
 
 FGuid UInterchangeBaseNode::GetHash() const
 {
-	return Attributes.GetStorageHash();
+	return Attributes->GetStorageHash();
 }
 
 class UClass* UInterchangeBaseNode::GetAssetClass() const
@@ -223,11 +193,13 @@ void UInterchangeBaseNode::FillAllCustomAttributeFromAsset(UObject* Object) cons
 
 void UInterchangeBaseNode::Serialize(FArchive& Ar)
 {
-	Ar << Attributes;
+	UE::Interchange::FAttributeStorage& RefAttributes = *(Attributes.Get());
+	Ar << RefAttributes;
 	if (Ar.IsLoading())
 	{
 		//The node is consider Initialize if the UniqueID and the Display label are set properly
-		bIsInitialized = (Attributes.GetAttributeHandle<FName>(UE::Interchange::FBaseNodeStaticData::UniqueIDKey()).IsValid() &&
-							Attributes.GetAttributeHandle<FName>(UE::Interchange::FBaseNodeStaticData::DisplayLabelKey()).IsValid());
+		bIsInitialized = (Attributes->GetAttributeHandle<FName>(UE::Interchange::FBaseNodeStaticData::UniqueIDKey()).IsValid() &&
+						  Attributes->GetAttributeHandle<FName>(UE::Interchange::FBaseNodeStaticData::DisplayLabelKey()).IsValid());
+
 	}
 }
diff --git a/Engine/Source/Runtime/Interchange/Core/Private/Tests/StorageTest.cpp b/Engine/Source/Runtime/Interchange/Core/Private/Tests/StorageTest.cpp
--- a/Engine/Source/Runtime/Interchange/Core/Private/Tests/StorageTest.cpp
+++ b/Engine/Source/Runtime/Interchange/Core/Private/Tests/StorageTest.cpp
@@ -18,15 +18,16 @@ bool FAttributeStorageTest::RunTest(const FString& Parameters)
 
 	FString KeyPrefixString = TEXT("TestKey");
 	FName KeyPrefix(*KeyPrefixString);
-	TAtomic<int32> KeySuffixCounter(1);
+	int32 KeySuffixCounter = 1;
 	auto CreateUniqueKey = [&KeyPrefix, &KeySuffixCounter]()->FAttributeKey
 	{
-		KeyPrefix.SetNumber(KeySuffixCounter.IncrementExchange());
+		check(IsInGameThread());
+		KeyPrefix.SetNumber(KeySuffixCounter++);
 		const FAttributeKey UniqueKeyRet = FAttributeKey(KeyPrefix);
 		return UniqueKeyRet;
 	};
 
-	FAttributeStorage TestStorage;
+	FAttributeStorage TestStorage = FAttributeStorage();
 	//This seed make the tests deterministic
 	FRandomStream RandomStream(564389);
 	const FAttributeKey TestInt32KeyName = CreateUniqueKey();
@@ -261,6 +262,11 @@ bool FAttributeStorageTest::RunTest(const FString& Parameters)
 		Keys.AddZeroed(BatchSize*BatchCount);
 		bool ThreadError = false;
 		FAttributeStorage::TAttributeHandle<FRandomStream> RandomStreamHandle = TestStorage.GetAttributeHandle<FRandomStream>(RandomStreamKey);
+		for (FAttributeKey& Key : Keys)
+		{
+			Key = CreateUniqueKey();
+		}
+
 		//Iterate all vertex to compute normals for all vertex instance
 		ParallelFor(BatchCount,
 			[&TestStorage, &RandomStreamHandle, &BatchSize, &Keys, &CreateUniqueKey, &ThreadError](const int32 BatchIndex)
@@ -272,7 +278,6 @@ bool FAttributeStorageTest::RunTest(const FString& Parameters)
 			{
 				const uint64 RealIndex = (BatchIndex * BatchSize) + AttributeIndex;
 				const FVector ValueRef = RandomStream.VRand();
-				Keys[RealIndex] = CreateUniqueKey();
 				if (!IsAttributeStorageResultSuccess(TestStorage.RegisterAttribute(Keys[RealIndex], ValueRef)))
 				{
 					ThreadError = true;
@@ -317,8 +322,8 @@ bool FAttributeStorageTest::RunTest(const FString& Parameters)
 
 	//Test compare storage
 	{
-		FAttributeStorage BaseStorage;
-		FAttributeStorage VersionStorage;
+		FAttributeStorage BaseStorage = FAttributeStorage();
+		FAttributeStorage VersionStorage = FAttributeStorage();
 		
 		//Adding many FVector
 		for (int32 AddedIndex = 0; AddedIndex < 3; ++AddedIndex)
@@ -408,7 +413,7 @@ bool FAttributeStorageTest::RunTest(const FString& Parameters)
 		FMemoryWriter Ar(MemoryMocked, bMemoryPersistent);
 		Ar << TestStorage;
 
-		FAttributeStorage FromMemoryMocked;
+		FAttributeStorage FromMemoryMocked = FAttributeStorage();
 		FMemoryReader ArRead(MemoryMocked, bMemoryPersistent);
 		ArRead << FromMemoryMocked;
 
diff --git a/Engine/Source/Runtime/Interchange/Core/Public/Nodes/InterchangeBaseNode.h b/Engine/Source/Runtime/Interchange/Core/Public/Nodes/InterchangeBaseNode.h
--- a/Engine/Source/Runtime/Interchange/Core/Public/Nodes/InterchangeBaseNode.h
+++ b/Engine/Source/Runtime/Interchange/Core/Public/Nodes/InterchangeBaseNode.h
@@ -133,17 +133,17 @@ bool FillCustom##AttributeName##FromAsset(UObject* Asset)											\
 
 #define IMPLEMENT_NODE_ATTRIBUTE_GETTER(AttributeName, AttributeType)																					\
 	FString OperationName = GetTypeName() + TEXT(".Get" #AttributeName);																				\
-	return InterchangePrivateNodeBase::GetCustomAttribute<AttributeType>(Attributes, Macro_Custom##AttributeName##Key, OperationName, AttributeValue);
+	return InterchangePrivateNodeBase::GetCustomAttribute<AttributeType>(*Attributes, Macro_Custom##AttributeName##Key, OperationName, AttributeValue);
 
 #define IMPLEMENT_NODE_ATTRIBUTE_SETTER_NODELEGATE(AttributeName, AttributeType)															\
 	FString OperationName = GetTypeName() + TEXT(".Set" #AttributeName);																				\
-	return InterchangePrivateNodeBase::SetCustomAttribute<AttributeType>(Attributes, Macro_Custom##AttributeName##Key, OperationName, AttributeValue);
+	return InterchangePrivateNodeBase::SetCustomAttribute<AttributeType>(*Attributes, Macro_Custom##AttributeName##Key, OperationName, AttributeValue);
 
 #if WITH_ENGINE
 
 #define IMPLEMENT_NODE_ATTRIBUTE_SETTER(NodeClassName, AttributeName, AttributeType, AssetType)														\
 	FString OperationName = GetTypeName() + TEXT(".Set" #AttributeName);																			\
-	if(InterchangePrivateNodeBase::SetCustomAttribute<AttributeType>(Attributes, Macro_Custom##AttributeName##Key, OperationName, AttributeValue))	\
+	if(InterchangePrivateNodeBase::SetCustomAttribute<AttributeType>(*Attributes, Macro_Custom##AttributeName##Key, OperationName, AttributeValue))	\
 	{																																				\
 		if(bAddApplyDelegate)																														\
 		{																																			\
@@ -160,7 +160,7 @@ bool FillCustom##AttributeName##FromAsset(UObject* Asset)											\
 
 #define IMPLEMENT_NODE_ATTRIBUTE_SETTER(NodeClassName, AttributeName, AttributeType, AssetType)															\
 	FString OperationName = GetTypeName() + TEXT(".Set" #AttributeName);																				\
-	return InterchangePrivateNodeBase::SetCustomAttribute<AttributeType>(Attributes, Macro_Custom##AttributeName##Key, OperationName, AttributeValue);
+	return InterchangePrivateNodeBase::SetCustomAttribute<AttributeType>(*Attributes, Macro_Custom##AttributeName##Key, OperationName, AttributeValue);
 
 #endif //#if WITH_ENGINE
 
@@ -174,6 +174,229 @@ namespace UE
 		DECLARE_DELEGATE_RetVal_OneParam(bool, FApplyAttributeToAsset, UObject*);
 		DECLARE_DELEGATE_RetVal_OneParam(bool, FFillAttributeToAsset, UObject*);
 
+		class FNameAttributeArrayHelper
+		{
+		public:
+			~FNameAttributeArrayHelper()
+			{
+				Attributes = nullptr;
+				KeyCount = NAME_None;
+			}
+
+			void Initialize(TSharedPtr<FAttributeStorage, ESPMode::ThreadSafe> InAttributes, const FString& BaseKeyName)
+			{
+				Attributes = InAttributes;
+				check(Attributes.IsValid());
+				FString BaseTryName = TEXT("__") + BaseKeyName;
+				KeyCount = FName(*BaseTryName);
+			}
+
+			int32 GetCount() const
+			{
+				//The Class must be initialise properly before we can use it
+				TSharedPtr<FAttributeStorage, ESPMode::ThreadSafe> AttributePtr = Attributes.Pin();
+				if (!ensure(AttributePtr.IsValid()))
+				{
+					return 0;
+				}
+				int32 NameCount = 0;
+				if (AttributePtr->ContainAttribute(GetKeyCount()))
+				{
+					FAttributeStorage::TAttributeHandle<int32> Handle = AttributePtr->GetAttributeHandle<int32>(GetKeyCount());
+					if (Handle.IsValid())
+					{
+						Handle.Get(NameCount);
+					}
+				}
+				return NameCount;
+			}
+
+			void GetNames(TArray<FName>& OutNames) const
+			{
+				//The Class must be initialise properly before we can use it
+				TSharedPtr<FAttributeStorage, ESPMode::ThreadSafe> AttributePtr = Attributes.Pin();
+				if (!ensure(AttributePtr.IsValid()))
+				{
+					OutNames.Empty();
+					return;
+				}
+				int32 NameCount = 0;
+				if (!AttributePtr->ContainAttribute(GetKeyCount()))
+				{
+					OutNames.Empty();
+					return;
+				}
+
+				//Reuse as much memory we can to avoid allocation
+				OutNames.Reset(NameCount);
+
+				FAttributeStorage::TAttributeHandle<int32> Handle = AttributePtr->GetAttributeHandle<int32>(GetKeyCount());
+				if (!Handle.IsValid())
+				{
+					return;
+				}
+				Handle.Get(NameCount);
+				for (int32 NameIndex = 0; NameIndex < NameCount; ++NameIndex)
+				{
+					FAttributeKey DepIndexKey = GetIndexKey(NameIndex);
+					FAttributeStorage::TAttributeHandle<FName> HandleName = AttributePtr->GetAttributeHandle<FName>(DepIndexKey);
+					if (!HandleName.IsValid())
+					{
+						continue;
+					}
+					FName& OutName = OutNames.AddDefaulted_GetRef();
+					HandleName.Get(OutName);
+				}
+			}
+
+			bool AddName(FName Name)
+			{
+				//The Class must be initialise properly before we can use it
+				TSharedPtr<FAttributeStorage, ESPMode::ThreadSafe> AttributePtr = Attributes.Pin();
+				if (!ensure(AttributePtr.IsValid()))
+				{
+					return false;
+				}
+
+				if (!AttributePtr->ContainAttribute(GetKeyCount()))
+				{
+					const int32 DependencyCount = 0;
+					EAttributeStorageResult Result = AttributePtr->RegisterAttribute<int32>(GetKeyCount(), DependencyCount);
+					if (!IsAttributeStorageResultSuccess(Result))
+					{
+						LogAttributeStorageErrors(Result, TEXT("FNameAttributeArrayHelper.AddName"), GetKeyCount());
+						return false;
+					}
+				}
+				FAttributeStorage::TAttributeHandle<int32> Handle = AttributePtr->GetAttributeHandle<int32>(GetKeyCount());
+				if (!ensure(Handle.IsValid()))
+				{
+					return false;
+				}
+				int32 NameIndex = 0;
+				Handle.Get(NameIndex);
+				FAttributeKey NameIndexKey = GetIndexKey(NameIndex);
+				//Increment the name counter
+				NameIndex++;
+				Handle.Set(NameIndex);
+
+				EAttributeStorageResult AddNameResult = AttributePtr->RegisterAttribute<FName>(NameIndexKey, Name);
+				if (!IsAttributeStorageResultSuccess(AddNameResult))
+				{
+					LogAttributeStorageErrors(AddNameResult, TEXT("FNameAttributeArrayHelper.AddName"), NameIndexKey);
+					return false;
+				}
+				return true;
+			}
+
+			bool RemoveName(FName NameToDelete)
+			{
+				//The Class must be initialise properly before we can use it
+				TSharedPtr<FAttributeStorage, ESPMode::ThreadSafe> AttributePtr = Attributes.Pin();
+				if (!ensure(AttributePtr.IsValid()))
+				{
+					return false;
+				}
+
+				int32 NameCount = 0;
+				if (!AttributePtr->ContainAttribute(GetKeyCount()))
+				{
+					return false;
+				}
+				FAttributeStorage::TAttributeHandle<int32> Handle = AttributePtr->GetAttributeHandle<int32>(GetKeyCount());
+				if (!Handle.IsValid())
+				{
+					return false;
+				}
+				Handle.Get(NameCount);
+				bool DecrementKey = false;
+				for (int32 NameIndex = 0; NameIndex < NameCount; ++NameIndex)
+				{
+					FAttributeKey DepIndexKey = GetIndexKey(NameIndex);
+					FAttributeStorage::TAttributeHandle<FName> HandleName = AttributePtr->GetAttributeHandle<FName>(DepIndexKey);
+					if (!HandleName.IsValid())
+					{
+						continue;
+					}
+					FName Name;
+					HandleName.Get(Name);
+					if (Name == NameToDelete)
+					{
+						//Remove this entry
+						AttributePtr->UnregisterAttribute(DepIndexKey);
+						Handle.Set(NameCount - 1);
+						//We have to rename the key for all the next item
+						DecrementKey = true;
+					}
+					else if (DecrementKey)
+					{
+						FAttributeKey NewDepIndexKey = GetIndexKey(NameIndex - 1);
+						EAttributeStorageResult UnregisterResult = AttributePtr->UnregisterAttribute(DepIndexKey);
+						if (IsAttributeStorageResultSuccess(UnregisterResult))
+						{
+							EAttributeStorageResult RegisterResult = AttributePtr->RegisterAttribute<FName>(NewDepIndexKey, Name);
+							if (!IsAttributeStorageResultSuccess(RegisterResult))
+							{
+								LogAttributeStorageErrors(RegisterResult, TEXT("FNameAttributeArrayHelper.RemoveName"), NewDepIndexKey);
+							}
+						}
+						else
+						{
+							LogAttributeStorageErrors(UnregisterResult, TEXT("FNameAttributeArrayHelper.RemoveName"), DepIndexKey);
+						}
+
+						//Avoid doing more code in the for since the HandleName is now invalid
+						continue;
+					}
+				}
+				return true;
+			}
+
+			bool RemoveAllNames()
+			{
+				//The Class must be initialise properly before we can use it
+				TSharedPtr<FAttributeStorage, ESPMode::ThreadSafe> AttributePtr = Attributes.Pin();
+				if (!ensure(AttributePtr.IsValid()))
+				{
+					return false;
+				}
+
+				int32 NameCount = 0;
+				if (!AttributePtr->ContainAttribute(GetKeyCount()))
+				{
+					return false;
+				}
+				FAttributeStorage::TAttributeHandle<int32> HandleCount = AttributePtr->GetAttributeHandle<int32>(GetKeyCount());
+				if (!HandleCount.IsValid())
+				{
+					return false;
+				}
+				HandleCount.Get(NameCount);
+				//Remove all attribute one by one
+				for (int32 NameIndex = 0; NameIndex < NameCount; ++NameIndex)
+				{
+					FAttributeKey DepIndexKey = GetIndexKey(NameIndex);
+					AttributePtr->UnregisterAttribute(DepIndexKey);
+				}
+				//Make sure Count is zero
+				NameCount = 0;
+				HandleCount.Set(NameCount);
+				return true;
+			}
+
+		private:
+			TWeakPtr<FAttributeStorage, ESPMode::ThreadSafe> Attributes = nullptr;
+
+			FAttributeKey KeyCount; //Assign in Initialize function, it will ensure if its the default value (NAME_None) when using the class
+			FAttributeKey GetKeyCount() const { ensure(KeyCount.Key != NAME_None); return KeyCount; }
+
+			FAttributeKey GetIndexKey(int32 Index) const
+			{
+				FString DepIndexKeyString = GetKeyCount().ToString() + TEXT("_NameIndex_") + FString::FromInt(Index);
+				return FAttributeKey(*DepIndexKeyString);
+			}
+		};
+
 		/**
 		 * Helper struct use to declare static const data we use in the UInterchangeBaseNode
 		 * Node that derive from UInterchangeBaseNode can also add a struct that derive from this one to add there static data
@@ -205,16 +428,10 @@ namespace UE
 				return AttributeKey;
 			}
 
-			static const FAttributeKey& DependencyCountKey()
+			static const FString& GetDependenciesBaseKey()
 			{
-				static FAttributeKey AttributeKey(TEXT("__DEPENDENCY_COUNT_"));
-				return AttributeKey;
-			}
-
-			static const FAttributeKey& DependencyBaseKey()
-			{
-				static FAttributeKey AttributeKey(TEXT("__DEPENDENCY_INDEX_"));
-				return AttributeKey;
+				static FString BaseNodeDependencies_BaseKey = TEXT("BaseNodeDependencies__");
+				return BaseNodeDependencies_BaseKey;
 			}
 		};
 
@@ -233,7 +450,10 @@ class INTERCHANGECORE_API UInterchangeBaseNode : public UObject
 
 public:
 	UInterchangeBaseNode()
-	{}
+	{
+		Attributes = MakeShared<UE::Interchange::FAttributeStorage, ESPMode::ThreadSafe>();
+		Dependencies.Initialize(Attributes, UE::Interchange::FBaseNodeStaticData::GetDependenciesBaseKey());
+	}
 
 	virtual ~UInterchangeBaseNode() = default;
 
@@ -260,11 +480,11 @@ public:
 	template<typename T>
 	UE::Interchange::FAttributeStorage::TAttributeHandle<T> RegisterAttribute(const UE::Interchange::FAttributeKey& NodeAttributeKey, const T& Value)
 	{
-		const UE::Interchange::EAttributeStorageResult Result = Attributes.RegisterAttribute(NodeAttributeKey, Value);
+		const UE::Interchange::EAttributeStorageResult Result = Attributes->RegisterAttribute(NodeAttributeKey, Value);
 		
 		if (IsAttributeStorageResultSuccess(Result))
 		{
-			return Attributes.GetAttributeHandle<T>(NodeAttributeKey);
+			return Attributes->GetAttributeHandle<T>(NodeAttributeKey);
 		}
 		LogAttributeStorageErrors(Result, TEXT("RegisterAttribute"), NodeAttributeKey);
 		return UE::Interchange::FAttributeStorage::TAttributeHandle<T>();
@@ -284,7 +504,7 @@ public:
 	 */
 	virtual UE::Interchange::EAttributeTypes GetAttributeType(const UE::Interchange::FAttributeKey& NodeAttributeKey) const
 	{
-		return Attributes.GetAttributeType(NodeAttributeKey);
+		return Attributes->GetAttributeType(NodeAttributeKey);
 	}
 
 	/**
@@ -296,7 +516,7 @@ public:
 	template<typename T>
 	UE::Interchange::FAttributeStorage::TAttributeHandle<T> GetAttributeHandle(const UE::Interchange::FAttributeKey& NodeAttributeKey) const
 	{
-		return Attributes.GetAttributeHandle<T>(NodeAttributeKey);
+		return Attributes->GetAttributeHandle<T>(NodeAttributeKey);
 	}
 
 	/**
@@ -332,6 +552,13 @@ public:
 	UFUNCTION(BlueprintCallable, Category = "Interchange | Node")
 	bool SetParentUID(FName ParentUID);
 
+	/**
+	 * This function allow to retrieve the number of dependencies for this object.
+	 *
+	 */
+	UFUNCTION(BlueprintCallable, Category = "Interchange | Node")
+	int32 GetDependeciesCount() const;
+
 	/**
 	 * This function allow to retrieve the dependency for this object.
 	 * 
@@ -345,6 +572,12 @@ public:
 	UFUNCTION(BlueprintCallable, Category = "Interchange | Node")
 	bool SetDependencyUID(FName DependencyUID);
 
+	/**
+	 * Remove one dependency from this object.
+	 */
+	UFUNCTION(BlueprintCallable, Category = "Interchange | Node")
+	bool RemoveDependencyUID(FName DependencyUID);
+
 	/**
 	 * IsEnable true mean that the node will be import/export, if false it will be discarded.
 	 * Return false if this node was disabled. Return true if the attribute is not there or if it was enabled.
@@ -391,24 +624,24 @@ public:
 
 	static void CompareNodeStorage(UInterchangeBaseNode* NodeA, const UInterchangeBaseNode* NodeB, TArray<UE::Interchange::FAttributeKey>& RemovedAttributes, TArray<UE::Interchange::FAttributeKey>& AddedAttributes, TArray<UE::Interchange::FAttributeKey>& ModifiedAttributes)
 	{
-		UE::Interchange::FAttributeStorage::CompareStorage(NodeA->Attributes, NodeB->Attributes, RemovedAttributes, AddedAttributes, ModifiedAttributes);
+		UE::Interchange::FAttributeStorage::CompareStorage(*(NodeA->Attributes), *(NodeB->Attributes), RemovedAttributes, AddedAttributes, ModifiedAttributes);
 	}
 	
 	static void CopyStorageAttributes(const UInterchangeBaseNode* SourceNode, UInterchangeBaseNode* DestinationNode, TArray<UE::Interchange::FAttributeKey>& AttributeKeys)
 	{
-		UE::Interchange::FAttributeStorage::CopyStorageAttributes(SourceNode->Attributes, DestinationNode->Attributes, AttributeKeys);
+		UE::Interchange::FAttributeStorage::CopyStorageAttributes(*(SourceNode->Attributes), *(DestinationNode->Attributes), AttributeKeys);
 	}
 
 	static void CopyStorage(const UInterchangeBaseNode* SourceNode, UInterchangeBaseNode* DestinationNode)
 	{
-		DestinationNode->Attributes = SourceNode->Attributes;
+		*(DestinationNode->Attributes) = *(SourceNode->Attributes);
 	}
 	
 	UPROPERTY()
 	mutable FSoftObjectPath ReferenceObject;
 protected:
 	/** The storage use to store the Key value attribute for this node. */
-	UE::Interchange::FAttributeStorage Attributes;
+	TSharedPtr<UE::Interchange::FAttributeStorage, ESPMode::ThreadSafe> Attributes;
 
 	/* This array hold the delegate to apply the attribute that has to be set on an UObject */
 	TMap<UClass*, TArray<UE::Interchange::FApplyAttributeToAsset>> ApplyCustomAttributeDelegates;
@@ -416,4 +649,6 @@ protected:
 	TMap<UClass*, TArray<UE::Interchange::FFillAttributeToAsset>> FillCustomAttributeDelegates;
 
 	bool bIsInitialized = false;
+
+	UE::Interchange::FNameAttributeArrayHelper Dependencies;
 };
