Using a compilation error instead of a runtime error for the unsuported attribute types.

diff --git a/Engine/Source/Runtime/Interchange/Core/Private/Types/AttributeStorage.cpp b/Engine/Source/Runtime/Interchange/Core/Private/Types/AttributeStorage.cpp
--- a/Engine/Source/Runtime/Interchange/Core/Private/Types/AttributeStorage.cpp
+++ b/Engine/Source/Runtime/Interchange/Core/Private/Types/AttributeStorage.cpp
@@ -47,11 +47,6 @@ namespace UE
 				UE_LOG(LogInterchangeCore, Warning, TEXT("Attribute storage operation [%s] Key[%s]: Stored attribute value type do not match parameter value type."), *OperationName, *(AttributeKey.Key));
 			}
 
-			if (HasAttributeStorageResult(Result, EAttributeStorageResult::Operation_Error_UnsupportedType))
-			{
-				UE_LOG(LogInterchangeCore, Warning, TEXT("Attribute storage operation [%s] Key[%s]: The storage do not support this type of value, please add unsupported type to EAttributeTypes."), *OperationName, *(AttributeKey.Key));
-			}
-
 			if (HasAttributeStorageResult(Result, EAttributeStorageResult::Operation_Error_InvalidStorage))
 			{
 				UE_LOG(LogInterchangeCore, Warning, TEXT("Attribute storage operation [%s] Key[%s]: The storage is invalid (NULL)."), *OperationName, *(AttributeKey.Key));
@@ -329,14 +324,11 @@ namespace UE
 
 		EAttributeStorageResult FAttributeStorage::GetAttribute(const FAttributeKey& ElementAttributeKey, FString& OutValue, TSpecializeType<FString >) const
 		{
+			static_assert(TAttributeTypeTraits<FString>::GetType() != EAttributeTypes::None, "Not a supported type for the attributes. Check EAttributeTypes for the supported types");
+
 			//Lock the storage
 			FScopeLock ScopeLock(&StorageMutex);
 
-			if (!IsSupportedType(OutValue))
-			{
-				return EAttributeStorageResult::Operation_Error_UnsupportedType;
-			}
-
 			const EAttributeTypes ValueType = TAttributeTypeTraits<FString>::GetType();
 			const FAttributeAllocationInfo* AttributeAllocationInfo = AttributeAllocationTable.Find(ElementAttributeKey);
 			if (AttributeAllocationInfo)
@@ -369,14 +361,11 @@ namespace UE
 
 		EAttributeStorageResult FAttributeStorage::GetAttribute(const FAttributeKey& ElementAttributeKey, FName& OutValue, TSpecializeType<FName >) const
 		{
+			static_assert(TAttributeTypeTraits<FName>::GetType() != EAttributeTypes::None, "Not a supported type for the attributes. Check EAttributeTypes for supported types");
+
 			//Lock the storage
 			FScopeLock ScopeLock(&StorageMutex);
 
-			if (!IsSupportedType(OutValue))
-			{
-				return EAttributeStorageResult::Operation_Error_UnsupportedType;
-			}
-
 			const EAttributeTypes ValueType = TAttributeTypeTraits<FName>::GetType();
 
 
diff --git a/Engine/Source/Runtime/Interchange/Core/Public/Types/AttributeStorage.h b/Engine/Source/Runtime/Interchange/Core/Public/Types/AttributeStorage.h
--- a/Engine/Source/Runtime/Interchange/Core/Public/Types/AttributeStorage.h
+++ b/Engine/Source/Runtime/Interchange/Core/Public/Types/AttributeStorage.h
@@ -238,12 +238,10 @@ namespace UE
 			Operation_Error_CannotRemoveAttribute			= (0x1 << 24),
 			//We try to override an attribute but the specified options do not allow override.
 			Operation_Error_CannotOverrideAttribute			= (0x1 << 25),
-			//The type pass in parameter is not supported, add new type to support in EAttributeTypes
-			Operation_Error_UnsupportedType					= (0x1 << 26),
 			//The storage is invalid (nullptr).
-			Operation_Error_InvalidStorage					= (0x1 << 27),
+			Operation_Error_InvalidStorage					= (0x1 << 26),
 			//Cannot get a valid value data pointer.
-			Operation_Error_InvalidMultiSizeValueData		= (0x1 << 28),
+			Operation_Error_InvalidMultiSizeValueData		= (0x1 << 27),
 		};
 
 		ENUM_CLASS_FLAGS(EAttributeStorageResult)
@@ -342,6 +340,8 @@ namespace UE
 			template<typename T>
 			class TAttributeHandle
 			{
+				static_assert(TAttributeTypeTraits<T>::GetType() != EAttributeTypes::None, "Unsupported attribute type");
+
 			public:
 				TAttributeHandle()
 				: AttributeStorage(nullptr)
@@ -382,13 +382,7 @@ namespace UE
 				{
 					AttributeStorage = const_cast<FAttributeStorage*>(InAttributeStorage);
 					Key = InKey;
-					//Log some error in case something is wrong
-					const EAttributeTypes ValueType = TAttributeTypeTraits<T>::GetType();
-					if (ValueType == EAttributeTypes::None)
-					{
-						//Value Type is unsupported
-						LogAttributeStorageErrors(EAttributeStorageResult::Operation_Error_UnsupportedType, TEXT("GetAttributeHandle"), Key);
-					}
+				
 					//Look for storage validity
 					if (AttributeStorage == nullptr)
 					{
@@ -403,7 +397,7 @@ namespace UE
 							//Storage do not contain the key
 							LogAttributeStorageErrors(EAttributeStorageResult::Operation_Error_CannotFoundKey, TEXT("GetAttributeHandle"), Key);
 						}
-						if (AttributeStorage->GetAttributeType(Key) != ValueType)
+						if (AttributeStorage->GetAttributeType(Key) != TAttributeTypeTraits<T>::GetType())
 						{
 							//Value Type is different from the existing key
 							LogAttributeStorageErrors(EAttributeStorageResult::Operation_Error_WrongType, TEXT("GetAttributeHandle"), Key);
@@ -428,20 +422,17 @@ namespace UE
 			 * @Param DefaultValue - is the default value for the registered attribute
 			 *
 			 * @note Possible errors
-			 * - An attribute of a unsupported type see IsSupportedType
 			 * - Key exist with a different type
 			 * - Storage is corrupted
 			 */
 			template<typename T>
 			EAttributeStorageResult RegisterAttribute(const FAttributeKey& ElementAttributeKey, const T& DefaultValue, EAttributeProperty AttributeProperty = EAttributeProperty::None)
 			{
+				static_assert(TAttributeTypeTraits<T>::GetType() != EAttributeTypes::None, "T is not a supported type for the attributes. Check EAttributeTypes for the supported types");
+
 				//Lock the storage
 				FScopeLock ScopeLock(&StorageMutex);
 
-				if (!IsSupportedType(DefaultValue))
-				{
-					return EAttributeStorageResult::Operation_Error_UnsupportedType;
-				}
 
 				const EAttributeTypes ValueType = TAttributeTypeTraits<T>::GetType();
 				const uint64 ValueSize = GetValueSize(DefaultValue);
@@ -503,25 +494,7 @@ namespace UE
 				return AttributeHandle;
 			}
 
-		
-			/**
-			 * Return true if the pass value type is support has an attribute
-			 * See EAttributeTypes to get all the supported type.
-			 *
-			 * @param Value is the value we check the type
-			 */
-			template<typename T>
-			static bool IsSupportedType(const T& Value)
-			{
-				const EAttributeTypes ValueType = TAttributeTypeTraits<T>::GetType();
-				if (ValueType == EAttributeTypes::None)
-				{
-					return false;
-				}
-				return true;
-			}
 
-		
 			/**
 			 * Return the attribute type if the key exist, return None if the key is missing
 			 *
@@ -667,7 +640,6 @@ namespace UE
 			 * @param Value is the value we want to add to the storage
 			 *
 			 * @note Possible errors
-			 * - An attribute of a unsupported type see IsSupportedType
 			 * - Key exist with a different type
 			 * - Key exist with a wrong size
 			 */
@@ -718,14 +690,11 @@ namespace UE
 			template<typename T>
 			EAttributeStorageResult SetAttribute(FAttributeAllocationInfo* AttributeAllocationInfo, const T& Value, TSpecializeType<T>)
 			{
+				static_assert(TAttributeTypeTraits<T>::GetType() != EAttributeTypes::None, "T is not a supported type for the attributes. Check EAttributeTypes for the supported types");
+
 				//Lock the storage
 				FScopeLock ScopeLock(&StorageMutex);
 
-				if (!IsSupportedType(Value))
-				{
-					return EAttributeStorageResult::Operation_Error_UnsupportedType;
-				}
-
 				if (!AttributeAllocationInfo)
 				{
 					return EAttributeStorageResult::Operation_Error_CannotFoundKey;
@@ -791,14 +760,11 @@ namespace UE
 			template<typename T>
 			EAttributeStorageResult GetAttribute(const FAttributeKey& ElementAttributeKey, T& OutValue, TSpecializeType<T>) const
 			{
+				static_assert(TAttributeTypeTraits<T>::GetType() != EAttributeTypes::None, "T is not a supported type for the attributes. Check EAttributeTypes for the supported types");
+
 				//Lock the storage
 				FScopeLock ScopeLock(&StorageMutex);
 
-				if (!IsSupportedType(OutValue))
-				{
-					return EAttributeStorageResult::Operation_Error_UnsupportedType;
-				}
-
 				const EAttributeTypes ValueType = TAttributeTypeTraits<T>::GetType();
 				const uint64 ValueSize = GetValueSize(OutValue);
 
