Make sure the source file always use absolute path.

diff --git a/Engine/Source/Runtime/Interchange/Core/Public/InterchangeSourceData.h b/Engine/Source/Runtime/Interchange/Core/Public/InterchangeSourceData.h
--- a/Engine/Source/Runtime/Interchange/Core/Public/InterchangeSourceData.h
+++ b/Engine/Source/Runtime/Interchange/Core/Public/InterchangeSourceData.h
@@ -30,9 +30,8 @@ public:
 	}
 
 	UInterchangeSourceData(FString InFilename)
-	: Filename(InFilename)
 	{
-		FileContentHashCache.Reset();
+		SetFilename(InFilename);
 	}
 	
 	UFUNCTION(BlueprintCallable, Category = "Interchange | Import Manager")
@@ -54,7 +53,7 @@ public:
 	UFUNCTION(BlueprintCallable, Category = "Interchange | Import Manager")
 	FORCEINLINE bool SetFilename(const FString& InFilename)
 	{
-		Filename = InFilename;
+		Filename = FPaths::ConvertRelativePathToFull(InFilename);
 		//Reset the cache
 		FileContentHashCache.Reset();
 		return true;
