Better sanitazing of the asset names in interchange.
It can now handle properly the dot character.

diff --git a/Engine/Source/Runtime/Interchange/Engine/Private/InterchangeManager.cpp b/Engine/Source/Runtime/Interchange/Engine/Private/InterchangeManager.cpp
--- a/Engine/Source/Runtime/Interchange/Engine/Private/InterchangeManager.cpp
+++ b/Engine/Source/Runtime/Interchange/Engine/Private/InterchangeManager.cpp
@@ -296,12 +296,22 @@ void UE::Interchange::FAssetImportResult::AddReferencedObjects(FReferenceCollect
 	Collector.AddReferencedObjects(ImportedAssets);
 }
 
-void UE::Interchange::SanitizeInvalidChar(FString& String)
+void UE::Interchange::SanitizeObjectPath(FString& ObjectPath)
 {
 	const TCHAR* InvalidChar = INVALID_OBJECTPATH_CHARACTERS;
 	while (*InvalidChar)
 	{
-		String.ReplaceCharInline(*InvalidChar, TCHAR('_'), ESearchCase::CaseSensitive);
+		ObjectPath.ReplaceCharInline(*InvalidChar, TCHAR('_'), ESearchCase::CaseSensitive);
+		++InvalidChar;
+	}
+}
+
+void UE::Interchange::SanitizeObjectName(FString& ObjectName)
+{
+	const TCHAR* InvalidChar = INVALID_OBJECTNAME_CHARACTERS;
+	while (*InvalidChar)
+	{
+		ObjectName.ReplaceCharInline(*InvalidChar, TCHAR('_'), ESearchCase::CaseSensitive);
 		++InvalidChar;
 	}
 }
@@ -540,7 +550,7 @@ UE::Interchange::FAssetImportResultRef UInterchangeManager::ImportAssetAsync(con
 	FString PackageBasePath = ContentPath;
 	if(!ImportAssetParameters.ReimportAsset)
 	{
-		UE::Interchange::SanitizeInvalidChar(PackageBasePath);
+		UE::Interchange::SanitizeObjectPath(PackageBasePath);
 	}
 
 	bool bCanShowDialog = !ImportAssetParameters.bIsAutomated && !IsAttended();
diff --git a/Engine/Source/Runtime/Interchange/Engine/Private/Tasks/InterchangeTaskCreateAsset.cpp b/Engine/Source/Runtime/Interchange/Engine/Private/Tasks/InterchangeTaskCreateAsset.cpp
--- a/Engine/Source/Runtime/Interchange/Engine/Private/Tasks/InterchangeTaskCreateAsset.cpp
+++ b/Engine/Source/Runtime/Interchange/Engine/Private/Tasks/InterchangeTaskCreateAsset.cpp
@@ -33,11 +33,12 @@ namespace UE
 
 				// Set the asset name and the package name
 				OutAssetName = NodeDisplayName;
+				SanitizeObjectName(OutAssetName);
 
-				OutPackageName = FPaths::Combine(*PackageBasePath, *OutAssetName);
+				FString SanitazedPackageBasePath = PackageBasePath;
+				SanitizeObjectPath(SanitazedPackageBasePath);
 
-				//Sanitize only the package name
-				UE::Interchange::SanitizeInvalidChar(OutPackageName);
+				OutPackageName = FPaths::Combine(*SanitazedPackageBasePath, *OutAssetName);
 			}
 		}//ns Private
 	}//ns Interchange
@@ -108,7 +109,7 @@ void UE::Interchange::FTaskCreatePackage::DoTask(ENamedThreads::Type CurrentThre
 		{
 			if (!NodeAsset->HasAnyInternalFlags(EInternalObjectFlags::Async))
 			{
-				//Sice the async flag is not set we must be in the game thread
+				//Since the async flag is not set we must be in the game thread
 				ensure(IsInGameThread());
 				NodeAsset->SetInternalFlags(EInternalObjectFlags::Async);
 			}
diff --git a/Engine/Source/Runtime/Interchange/Engine/Public/InterchangeManager.h b/Engine/Source/Runtime/Interchange/Engine/Public/InterchangeManager.h
--- a/Engine/Source/Runtime/Interchange/Engine/Public/InterchangeManager.h
+++ b/Engine/Source/Runtime/Interchange/Engine/Public/InterchangeManager.h
@@ -187,7 +187,9 @@ namespace UE
 			void CleanUp();
 		};
 
-		void SanitizeInvalidChar(FString& String);
+		void SanitizeObjectPath(FString& ObjectPath);
+
+		void SanitizeObjectName(FString& ObjectName);
 
 	} //ns interchange
 } //ns UE
