Fix skeletal mesh re-import crash

diff --git a/Engine/Source/Runtime/Interchange/Engine/Private/InterchangeManager.cpp b/Engine/Source/Runtime/Interchange/Engine/Private/InterchangeManager.cpp
--- a/Engine/Source/Runtime/Interchange/Engine/Private/InterchangeManager.cpp
+++ b/Engine/Source/Runtime/Interchange/Engine/Private/InterchangeManager.cpp
@@ -553,6 +553,10 @@ UE::Interchange::FAssetImportResultRef UInterchangeManager::ImportAssetAsync(con
 	{
 		UE::Interchange::SanitizeObjectPath(PackageBasePath);
 	}
+	else
+	{
+		PackageBasePath = FPaths::GetPath(ImportAssetParameters.ReimportAsset->GetPathName());
+	}
 
 	bool bCanShowDialog = !ImportAssetParameters.bIsAutomated && !IsAttended();
 
diff --git a/Engine/Source/Runtime/Interchange/Engine/Private/Tasks/InterchangeTaskCreateAsset.cpp b/Engine/Source/Runtime/Interchange/Engine/Private/Tasks/InterchangeTaskCreateAsset.cpp
--- a/Engine/Source/Runtime/Interchange/Engine/Private/Tasks/InterchangeTaskCreateAsset.cpp
+++ b/Engine/Source/Runtime/Interchange/Engine/Private/Tasks/InterchangeTaskCreateAsset.cpp
@@ -29,7 +29,6 @@ namespace UE
 				const UInterchangeSourceData* SourceData = AsyncHelper.SourceDatas[SourceIndex];
 				check(SourceData);
 				FString NodeDisplayName = Node->GetDisplayLabel();
-				const FString BaseFileName = FPaths::GetBaseFilename(SourceData->GetFilename());
 
 				// Set the asset name and the package name
 				OutAssetName = NodeDisplayName;
@@ -64,8 +63,22 @@ void UE::Interchange::FTaskCreatePackage::DoTask(ENamedThreads::Type CurrentThre
 	//If we do a reimport no need to create a package
 	if (AsyncHelper->TaskData.ReimportObject)
 	{
-		Pkg = AsyncHelper->TaskData.ReimportObject->GetPackage();
-		PackageName = Pkg->GetPathName();
+		Private::InternalGetPackageName(*AsyncHelper, SourceIndex, PackageBasePath, Node, PackageName, AssetName);
+		UPackage* ResultFindPackage = FindPackage(nullptr, *PackageName);
+		UObject* FindOuter = (ResultFindPackage == nullptr ? ANY_PACKAGE : ResultFindPackage);
+		UObject* ExistingObject = FindObject<UObject>(FindOuter, *AssetName);
+		if (ExistingObject)
+		{
+			Pkg = ExistingObject->GetPackage();
+			PackageName = Pkg->GetPathName();
+		}
+		else
+		{
+			const FText Message = FText::Format(NSLOCTEXT("Interchange", "CannotFindPackageDuringReimportErrorMsg", "Cannot find package named '{0}', for asset {1}."), FText::FromString(PackageName), FText::FromString(AssetName));
+			UE_LOG(LogInterchangeEngine, Warning, TEXT("%s"), *Message.ToString());
+			//Skip this asset
+			return;
+		}
 	}
 	else
 	{
@@ -149,10 +162,26 @@ void UE::Interchange::FTaskCreateAsset::DoTask(ENamedThreads::Type CurrentThread
 	FString PackageName;
 	FString AssetName;
 	Private::InternalGetPackageName(*AsyncHelper, SourceIndex, PackageBasePath, Node, PackageName, AssetName);
+	bool bSkipAsset = false;
+	UObject* ExistingObject = nullptr;
 	if (AsyncHelper->TaskData.ReimportObject)
 	{
-		Pkg = AsyncHelper->TaskData.ReimportObject->GetPackage();
-		PackageName = Pkg->GetPathName();
+		//When we re-import one particular asset, if the source file contain other assets, we want to set the node= reference UObject for those asset to the existing asset
+		//The way to discover this case is to compare the re-import asset with the node asset.
+		UPackage* ResultFindPackage = FindPackage(nullptr, *PackageName);
+		UObject* FindOuter = (ResultFindPackage == nullptr ? ANY_PACKAGE : ResultFindPackage);
+		ExistingObject = FindObject<UObject>(FindOuter, *AssetName);
+		bSkipAsset = !ExistingObject || ExistingObject != AsyncHelper->TaskData.ReimportObject;
+		if (!bSkipAsset)
+		{
+			Pkg = AsyncHelper->TaskData.ReimportObject->GetPackage();
+			PackageName = Pkg->GetPathName();
+		}
+		else if(ExistingObject)
+		{
+			Pkg = ExistingObject->GetPackage();
+			PackageName = Pkg->GetPathName();
+		}
 	}
 	else
 	{
@@ -176,40 +205,51 @@ void UE::Interchange::FTaskCreateAsset::DoTask(ENamedThreads::Type CurrentThread
 		Pkg = *PkgPtr;
 	}
 
-	UInterchangeTranslatorBase* Translator = AsyncHelper->Translators[SourceIndex];
-	//Import Asset describe by the node
-	UInterchangeFactoryBase::FCreateAssetParams CreateAssetParams;
-	CreateAssetParams.AssetName = AssetName;
-	CreateAssetParams.AssetNode = Node;
-	CreateAssetParams.Parent = Pkg;
-	CreateAssetParams.SourceData = AsyncHelper->SourceDatas[SourceIndex];
-	CreateAssetParams.Translator = Translator;
-	if (AsyncHelper->BaseNodeContainers.IsValidIndex(SourceIndex))
+	UObject* NodeAsset = nullptr;
+	if (bSkipAsset)
 	{
-		CreateAssetParams.NodeContainer = AsyncHelper->BaseNodeContainers[SourceIndex].Get();
+		NodeAsset = ExistingObject;
 	}
-	CreateAssetParams.ReimportObject = AsyncHelper->TaskData.ReimportObject;
+	else
+	{
+		UInterchangeTranslatorBase* Translator = AsyncHelper->Translators[SourceIndex];
+		//Import Asset describe by the node
+		UInterchangeFactoryBase::FCreateAssetParams CreateAssetParams;
+		CreateAssetParams.AssetName = AssetName;
+		CreateAssetParams.AssetNode = Node;
+		CreateAssetParams.Parent = Pkg;
+		CreateAssetParams.SourceData = AsyncHelper->SourceDatas[SourceIndex];
+		CreateAssetParams.Translator = Translator;
+		if (AsyncHelper->BaseNodeContainers.IsValidIndex(SourceIndex))
+		{
+			CreateAssetParams.NodeContainer = AsyncHelper->BaseNodeContainers[SourceIndex].Get();
+		}
+		CreateAssetParams.ReimportObject = AsyncHelper->TaskData.ReimportObject;
 
-	CreateAssetParams.ReimportStrategyFlags = EReimportStrategyFlags::ApplyNoProperties;
-	//CreateAssetParams.ReimportStrategyFlags = EReimportStrategyFlags::ApplyPipelineProperties;
-	//CreateAssetParams.ReimportStrategyFlags = EReimportStrategyFlags::ApplyEditorChangedProperties;
+		CreateAssetParams.ReimportStrategyFlags = EReimportStrategyFlags::ApplyNoProperties;
+		//CreateAssetParams.ReimportStrategyFlags = EReimportStrategyFlags::ApplyPipelineProperties;
+		//CreateAssetParams.ReimportStrategyFlags = EReimportStrategyFlags::ApplyEditorChangedProperties;
 
-	UObject* NodeAsset = Factory->CreateAsset(CreateAssetParams);
+		NodeAsset = Factory->CreateAsset(CreateAssetParams);
+	}
 	if (NodeAsset)
 	{
-		FScopeLock Lock(&AsyncHelper->ImportedAssetsPerSourceIndexLock);
-		TArray<UE::Interchange::FImportAsyncHelper::FImportedAssetInfo>& ImportedInfos = AsyncHelper->ImportedAssetsPerSourceIndex.FindOrAdd(SourceIndex);
-		UE::Interchange::FImportAsyncHelper::FImportedAssetInfo* AssetInfoPtr = ImportedInfos.FindByPredicate([NodeAsset](const UE::Interchange::FImportAsyncHelper::FImportedAssetInfo& CurInfo)
+		if (!bSkipAsset)
 		{
-			return CurInfo.ImportAsset == NodeAsset;
-		});
+			FScopeLock Lock(&AsyncHelper->ImportedAssetsPerSourceIndexLock);
+			TArray<UE::Interchange::FImportAsyncHelper::FImportedAssetInfo>& ImportedInfos = AsyncHelper->ImportedAssetsPerSourceIndex.FindOrAdd(SourceIndex);
+			UE::Interchange::FImportAsyncHelper::FImportedAssetInfo* AssetInfoPtr = ImportedInfos.FindByPredicate([NodeAsset](const UE::Interchange::FImportAsyncHelper::FImportedAssetInfo& CurInfo)
+			{
+				return CurInfo.ImportAsset == NodeAsset;
+			});
 
-		if (!AssetInfoPtr)
-		{
-			UE::Interchange::FImportAsyncHelper::FImportedAssetInfo& AssetInfo = ImportedInfos.AddDefaulted_GetRef();
-			AssetInfo.ImportAsset = NodeAsset;
-			AssetInfo.Factory = Factory;
-			AssetInfo.NodeUniqueId = Node->GetUniqueID();
+			if (!AssetInfoPtr)
+			{
+				UE::Interchange::FImportAsyncHelper::FImportedAssetInfo& AssetInfo = ImportedInfos.AddDefaulted_GetRef();
+				AssetInfo.ImportAsset = NodeAsset;
+				AssetInfo.Factory = Factory;
+				AssetInfo.NodeUniqueId = Node->GetUniqueID();
+			}
 		}
 
 		Node->ReferenceObject = FSoftObjectPath(NodeAsset);
