First pass implementaiton of a message passing system for Interchange. This permits messages, warnings and errors to be output from translators, pipelines and factories, as well as from external proce

diff --git a/Engine/Source/Runtime/Interchange/Engine/Private/Tasks/InterchangeTaskCreateAsset.cpp b/Engine/Source/Runtime/Interchange/Engine/Private/Tasks/InterchangeTaskCreateAsset.cpp
--- a/Engine/Source/Runtime/Interchange/Engine/Private/Tasks/InterchangeTaskCreateAsset.cpp
+++ b/Engine/Source/Runtime/Interchange/Engine/Private/Tasks/InterchangeTaskCreateAsset.cpp
@@ -6,6 +6,7 @@
 #include "InterchangeFactoryBase.h"
 #include "InterchangeEngineLogPrivate.h"
 #include "InterchangeManager.h"
+#include "InterchangeResult.h"
 #include "InterchangeSourceData.h"
 #include "InterchangeTranslatorBase.h"
 #include "Misc/Paths.h"
@@ -57,6 +58,18 @@ void UE::Interchange::FTaskCreatePackage::DoTask(ENamedThreads::Type CurrentThre
 		return;
 	}
 
+	//The create package thread must always execute on the game thread
+	check(IsInGameThread());
+
+	// Create factory
+	UInterchangeFactoryBase* Factory = NewObject<UInterchangeFactoryBase>(GetTransientPackage(), FactoryClass);
+	Factory->SetResultsContainer(AsyncHelper->AssetImportResult->GetResults());
+
+	{
+		FScopeLock Lock(&AsyncHelper->CreatedFactoriesLock);
+		AsyncHelper->CreatedFactories.Add(Node->GetUniqueID(), Factory);
+	}
+
 	UPackage* Pkg = nullptr;
 	FString PackageName;
 	FString AssetName;
@@ -74,23 +87,28 @@ void UE::Interchange::FTaskCreatePackage::DoTask(ENamedThreads::Type CurrentThre
 		}
 		else
 		{
-			const FText Message = FText::Format(NSLOCTEXT("Interchange", "CannotFindPackageDuringReimportErrorMsg", "Cannot find package named '{0}', for asset {1}."), FText::FromString(PackageName), FText::FromString(AssetName));
-			UE_LOG(LogInterchangeEngine, Warning, TEXT("%s"), *Message.ToString());
+			UInterchangeResultError_Generic* Message = Factory->AddMessage<UInterchangeResultError_Generic>();
+			Message->SourceAssetName = AsyncHelper->SourceDatas[SourceIndex]->GetFilename();
+			Message->DestinationAssetName = AssetName;
+			Message->AssetType = Node->GetAssetClass();
+			Message->Text = NSLOCTEXT("Interchange", "CannotFindPackageDuringReimport", "Cannot find an existing package.");
+
 			//Skip this asset
 			return;
 		}
 	}
 	else
 	{
-		//The create package thread must always execute on the game thread
-		check(IsInGameThread());
-
 		Private::InternalGetPackageName(*AsyncHelper, SourceIndex, PackageBasePath, Node, PackageName, AssetName);
 		// We can not create assets that share the name of a map file in the same location
 		if (UE::Interchange::FPackageUtils::IsMapPackageAsset(PackageName))
 		{
-			const FText Message = FText::Format(NSLOCTEXT("Interchange", "AssetNameInUseByMap", "You can not create an asset named '{0}' because there is already a map file with this name in this folder."), FText::FromString(AssetName));
-			UE_LOG(LogInterchangeEngine, Warning, TEXT("%s"), *Message.ToString());
+			UInterchangeResultError_Generic* Message = Factory->AddMessage<UInterchangeResultError_Generic>();
+			Message->SourceAssetName = AsyncHelper->SourceDatas[SourceIndex]->GetFilename();
+			Message->DestinationAssetName = AssetName;
+			Message->AssetType = Node->GetAssetClass();
+			Message->Text = NSLOCTEXT("Interchange", "MapExistsWithSameName", "You cannot create an asset with this name, as there is already a map file with the same name in this folder.");
+
 			//Skip this asset
 			return;
 		}
@@ -98,8 +116,12 @@ void UE::Interchange::FTaskCreatePackage::DoTask(ENamedThreads::Type CurrentThre
 		Pkg = CreatePackage(*PackageName);
 		if (Pkg == nullptr)
 		{
-			const FText Message = FText::Format(NSLOCTEXT("Interchange", "CannotCreatePackageErrorMsg", "Cannot create package named '{0}', will not import asset {1}."), FText::FromString(PackageName), FText::FromString(AssetName));
-			UE_LOG(LogInterchangeEngine, Warning, TEXT("%s"), *Message.ToString());
+			UInterchangeResultError_Generic* Message = Factory->AddMessage<UInterchangeResultError_Generic>();
+			Message->SourceAssetName = AsyncHelper->SourceDatas[SourceIndex]->GetFilename();
+			Message->DestinationAssetName = AssetName;
+			Message->AssetType = Node->GetAssetClass();
+			Message->Text = FText::Format(NSLOCTEXT("Interchange", "CouldntCreatePackage", "It was not possible to create a package named '{0}'; the asset will not be imported."), FText::FromString(PackageName));
+
 			//Skip this asset
 			return;
 		}
@@ -135,6 +157,7 @@ void UE::Interchange::FTaskCreatePackage::DoTask(ENamedThreads::Type CurrentThre
 			Node->ReferenceObject = FSoftObjectPath(NodeAsset);
 		}
 	}
+
 	// Make sure the destination package is loaded
 	Pkg->FullyLoad();
 	
@@ -150,7 +173,7 @@ void UE::Interchange::FTaskCreateAsset::DoTask(ENamedThreads::Type CurrentThread
 	INTERCHANGE_TRACE_ASYNCHRONOUS_TASK(CreateAsset)
 #endif
 	TSharedPtr<UE::Interchange::FImportAsyncHelper, ESPMode::ThreadSafe> AsyncHelper = WeakAsyncHelper.Pin();
-	check(WeakAsyncHelper.IsValid());
+	check(AsyncHelper.IsValid());
 
 	//Verify if the task was cancel
 	if (AsyncHelper->bCancel)
@@ -158,6 +181,12 @@ void UE::Interchange::FTaskCreateAsset::DoTask(ENamedThreads::Type CurrentThread
 		return;
 	}
 
+	UInterchangeFactoryBase* Factory = nullptr;
+	{
+		FScopeLock Lock(&AsyncHelper->CreatedFactoriesLock);
+		Factory = AsyncHelper->CreatedFactories.FindChecked(Node->GetUniqueID());
+	}
+
 	UPackage* Pkg = nullptr;
 	FString PackageName;
 	FString AssetName;
@@ -190,15 +219,22 @@ void UE::Interchange::FTaskCreateAsset::DoTask(ENamedThreads::Type CurrentThread
 
 		if (!PkgPtr || !(*PkgPtr))
 		{
-			const FText Message = FText::Format(NSLOCTEXT("Interchange", "CannotCreateAssetNoPackageErrorMsg", "Cannot create asset named '{1}', package '{0}'was not created properly."), FText::FromString(PackageName), FText::FromString(AssetName));
-			UE_LOG(LogInterchangeEngine, Warning, TEXT("%s"), *Message.ToString());
+			UInterchangeResultError_Generic* Message = Factory->AddMessage<UInterchangeResultError_Generic>();
+			Message->SourceAssetName = AsyncHelper->SourceDatas[SourceIndex]->GetFilename();
+			Message->DestinationAssetName = AssetName;
+			Message->AssetType = Node->GetAssetClass();
+			Message->Text = NSLOCTEXT("Interchange", "BadPackage", "It was not possible to create the asset as its package was not created correctly.");
+
 			return;
 		}
 
 		if (!AsyncHelper->SourceDatas.IsValidIndex(SourceIndex) || !AsyncHelper->Translators.IsValidIndex(SourceIndex))
 		{
-			const FText Message = FText::Format(NSLOCTEXT("Interchange", "CannotCreateAssetMissingDataErrorMsg", "Cannot create asset named '{0}', Source data or translator is invalid."), FText::FromString(AssetName));
-			UE_LOG(LogInterchangeEngine, Warning, TEXT("%s"), *Message.ToString());
+			UInterchangeResultError_Generic* Message = Factory->AddMessage<UInterchangeResultError_Generic>();
+			Message->DestinationAssetName = AssetName;
+			Message->AssetType = Node->GetAssetClass();
+			Message->Text = NSLOCTEXT("Interchange", "SourceDataOrTranslatorInvalid", "It was not possible to create the asset as its translator was not created correctly.");
+
 			return;
 		}
 
@@ -250,6 +286,22 @@ void UE::Interchange::FTaskCreateAsset::DoTask(ENamedThreads::Type CurrentThread
 				AssetInfo.Factory = Factory;
 				AssetInfo.NodeUniqueId = Node->GetUniqueID();
 			}
+
+			// Fill in destination asset and type in any results which have been added previously by a translator or pipeline, now that we have a corresponding factory.
+			UInterchangeResultsContainer* Results = AsyncHelper->AssetImportResult->GetResults();
+			for (UInterchangeResult* Result : Results->GetResults())
+			{
+				if (!Result->InterchangeKey.IsEmpty() && (Result->DestinationAssetName.IsEmpty() || Result->AssetType == nullptr))
+				{
+					TArray<FString> TargetAssets;
+					Node->GetTargetAssetUids(TargetAssets);
+					if (TargetAssets.Contains(Result->InterchangeKey))
+					{
+						Result->DestinationAssetName = NodeAsset->GetPathName();
+						Result->AssetType = NodeAsset->GetClass();
+					}
+				}
+			}
 		}
 
 		Node->ReferenceObject = FSoftObjectPath(NodeAsset);
diff --git a/Engine/Source/Runtime/Interchange/Engine/Private/Tasks/InterchangeTaskCreateAsset.h b/Engine/Source/Runtime/Interchange/Engine/Private/Tasks/InterchangeTaskCreateAsset.h
--- a/Engine/Source/Runtime/Interchange/Engine/Private/Tasks/InterchangeTaskCreateAsset.h
+++ b/Engine/Source/Runtime/Interchange/Engine/Private/Tasks/InterchangeTaskCreateAsset.h
@@ -25,27 +25,28 @@ namespace UE
 			int32 SourceIndex;
 			TWeakPtr<FImportAsyncHelper, ESPMode::ThreadSafe> WeakAsyncHelper;
 			UInterchangeBaseNode* Node;
-			UInterchangeFactoryBase* Factory;
+			const UClass* FactoryClass;
 
 		public:
-			FTaskCreatePackage(const FString& InPackageBasePath, const int32 InSourceIndex, TWeakPtr<FImportAsyncHelper, ESPMode::ThreadSafe> InAsyncHelper, UInterchangeBaseNode* InNode, UInterchangeFactoryBase* InFactory)
+			FTaskCreatePackage(const FString& InPackageBasePath, const int32 InSourceIndex, TWeakPtr<FImportAsyncHelper, ESPMode::ThreadSafe> InAsyncHelper, UInterchangeBaseNode* InNode, const UClass* InFactoryClass)
 				: PackageBasePath(InPackageBasePath)
 				, SourceIndex(InSourceIndex)
 				, WeakAsyncHelper(InAsyncHelper)
 				, Node(InNode)
-				, Factory(InFactory)
+				, FactoryClass(InFactoryClass)
 			{
 				check(Node);
-				check(Factory);
+				check(FactoryClass);
 			}
 
-			ENamedThreads::Type GetDesiredThread()
+			ENamedThreads::Type GetDesiredThread() const
 			{
-				if (WeakAsyncHelper.IsValid() && WeakAsyncHelper.Pin()->TaskData.ReimportObject)
-				{
-					//When doing a reimport the package already exist, so we can get it outside of the main thread
-					return ENamedThreads::AnyBackgroundThreadNormalTask;
-				}
+				// This is no longer true now that we have to construct a factory here
+				//if (WeakAsyncHelper.IsValid() && WeakAsyncHelper.Pin()->TaskData.ReimportObject)
+				//{
+				//	//When doing a reimport the package already exist, so we can get it outside of the main thread
+				//	return ENamedThreads::AnyBackgroundThreadNormalTask;
+				//}
 				return ENamedThreads::GameThread;
 			}
 
@@ -69,23 +70,22 @@ namespace UE
 			int32 SourceIndex;
 			TWeakPtr<FImportAsyncHelper, ESPMode::ThreadSafe> WeakAsyncHelper;
 			UInterchangeBaseNode* Node;
-			UInterchangeFactoryBase* Factory;
+			bool bCanRunOnAnyThread;
 
 		public:
-			FTaskCreateAsset(const FString& InPackageBasePath, const int32 InSourceIndex, TWeakPtr<FImportAsyncHelper, ESPMode::ThreadSafe> InAsyncHelper, UInterchangeBaseNode* InNode, UInterchangeFactoryBase* InFactory)
+			FTaskCreateAsset(const FString& InPackageBasePath, const int32 InSourceIndex, TWeakPtr<FImportAsyncHelper, ESPMode::ThreadSafe> InAsyncHelper, UInterchangeBaseNode* InNode, bool bInCanRunOnAnyThread)
 				: PackageBasePath(InPackageBasePath)
 				, SourceIndex(InSourceIndex)
 				, WeakAsyncHelper(InAsyncHelper)
 				, Node(InNode)
-				, Factory(InFactory)
+				, bCanRunOnAnyThread(bInCanRunOnAnyThread)
 			{
 				check(Node);
-				check(Factory);
 			}
 
-			ENamedThreads::Type GetDesiredThread()
+			ENamedThreads::Type GetDesiredThread() const
 			{
-				return Factory->CanExecuteOnAnyThread() ? ENamedThreads::AnyBackgroundThreadNormalTask : ENamedThreads::GameThread;
+				return bCanRunOnAnyThread ? ENamedThreads::AnyBackgroundThreadNormalTask : ENamedThreads::GameThread;
 			}
 
 			static ESubsequentsMode::Type GetSubsequentsMode()
diff --git a/Engine/Source/Runtime/Interchange/Engine/Private/Tasks/InterchangeTaskParsing.cpp b/Engine/Source/Runtime/Interchange/Engine/Private/Tasks/InterchangeTaskParsing.cpp
--- a/Engine/Source/Runtime/Interchange/Engine/Private/Tasks/InterchangeTaskParsing.cpp
+++ b/Engine/Source/Runtime/Interchange/Engine/Private/Tasks/InterchangeTaskParsing.cpp
@@ -39,7 +39,7 @@ void UE::Interchange::FTaskParsing::DoTask(ENamedThreads::Type CurrentThread, co
 		TArray<FString> Dependencies;
 		FGraphEventRef GraphEventRef;
 		FGraphEventArray Prerequistes;
-		UInterchangeFactoryBase* Factory;
+		const UClass* FactoryClass;
 	};
 
 
@@ -64,18 +64,19 @@ void UE::Interchange::FTaskParsing::DoTask(ENamedThreads::Type CurrentThread, co
 			{
 				if (Node->GetAssetClass() != nullptr)
 				{
-					UInterchangeFactoryBase* NodeFactory = InterchangeManager->GetRegisterFactory(Node->GetAssetClass());
-					if (!NodeFactory)
+					const UClass* RegisteredFactoryClass = InterchangeManager->GetRegisteredFactoryClass(Node->GetAssetClass());
+					if (!RegisteredFactoryClass)
 					{
 						//nothing we can import from this element
 						return;
 					}
+
 					FTaskData& NodeTaskData = TaskDatas.AddDefaulted_GetRef();
 					NodeTaskData.UniqueID = Node->GetUniqueID();
 					NodeTaskData.SourceIndex = SourceIndex;
 					NodeTaskData.Node = Node;
 					Node->GetFactoryDependencies(NodeTaskData.Dependencies);
-					NodeTaskData.Factory = NodeFactory;
+					NodeTaskData.FactoryClass = RegisteredFactoryClass;
 				}
 			});
 		}
@@ -110,13 +111,22 @@ void UE::Interchange::FTaskParsing::DoTask(ENamedThreads::Type CurrentThread, co
 			}
 		}
 
+		const int32 SourceIndex = TaskDatas[TaskIndex].SourceIndex;
+		const UClass* const FactoryClass = TaskDatas[TaskIndex].FactoryClass;
+		UInterchangeBaseNode* const FactoryNode = TaskDatas[TaskIndex].Node;
+		const bool bFactoryCanRunOnAnyThread = FactoryClass->GetDefaultObject<UInterchangeFactoryBase>()->CanExecuteOnAnyThread();
+
 		//Add create package task has a prerequisite of FTaskCreateAsset. Create package task is a game thread task
 		FGraphEventArray CreatePackagePrerequistes;
-		int32 CreatePackageTaskIndex = AsyncHelper->CreatePackageTasks.Add(TGraphTask<FTaskCreatePackage>::CreateTask(&(TaskDatas[TaskIndex].Prerequistes)).ConstructAndDispatchWhenReady(PackageBasePath, TaskDatas[TaskIndex].SourceIndex, WeakAsyncHelper, TaskDatas[TaskIndex].Node, TaskDatas[TaskIndex].Factory));
+		int32 CreatePackageTaskIndex = AsyncHelper->CreatePackageTasks.Add(
+			TGraphTask<FTaskCreatePackage>::CreateTask(&(TaskDatas[TaskIndex].Prerequistes)).ConstructAndDispatchWhenReady(PackageBasePath, SourceIndex, WeakAsyncHelper, FactoryNode, FactoryClass)
+		);
 		CreatePackagePrerequistes.Add(AsyncHelper->CreatePackageTasks[CreatePackageTaskIndex]);
 
 		FGraphEventArray PostPipelinePrerequistes;
-		int32 CreateTaskIndex = AsyncHelper->CreateAssetTasks.Add(TGraphTask<FTaskCreateAsset>::CreateTask(&(CreatePackagePrerequistes)).ConstructAndDispatchWhenReady(PackageBasePath, TaskDatas[TaskIndex].SourceIndex, WeakAsyncHelper, TaskDatas[TaskIndex].Node, TaskDatas[TaskIndex].Factory));
+		int32 CreateTaskIndex = AsyncHelper->CreateAssetTasks.Add(
+			TGraphTask<FTaskCreateAsset>::CreateTask(&(CreatePackagePrerequistes)).ConstructAndDispatchWhenReady(PackageBasePath, SourceIndex, WeakAsyncHelper, FactoryNode, bFactoryCanRunOnAnyThread)
+		);
 		PostPipelinePrerequistes.Add(AsyncHelper->CreateAssetTasks[CreateTaskIndex]);
 		
 		TaskDatas[TaskIndex].GraphEventRef = AsyncHelper->CreateAssetTasks[CreateTaskIndex];
@@ -124,7 +134,9 @@ void UE::Interchange::FTaskParsing::DoTask(ENamedThreads::Type CurrentThread, co
 		for (int32 GraphPipelineIndex = 0; GraphPipelineIndex < AsyncHelper->Pipelines.Num(); ++GraphPipelineIndex)
 		{
 			int32 GraphPipelineTaskIndex = INDEX_NONE;
-			GraphPipelineTaskIndex = AsyncHelper->PipelinePostImportTasks.Add(TGraphTask<FTaskPipelinePostImport>::CreateTask(&(PostPipelinePrerequistes)).ConstructAndDispatchWhenReady(TaskDatas[TaskIndex].SourceIndex, GraphPipelineIndex, WeakAsyncHelper));
+			GraphPipelineTaskIndex = AsyncHelper->PipelinePostImportTasks.Add(
+				TGraphTask<FTaskPipelinePostImport>::CreateTask(&(PostPipelinePrerequistes)).ConstructAndDispatchWhenReady(SourceIndex, GraphPipelineIndex, WeakAsyncHelper)
+			);
 			//Ensure we run the pipeline in the same order we create the task, since pipeline modify the node container, its important that its not process in parallel, Adding the one we start to the prerequisites
 			//is the way to go here
 			PostPipelinePrerequistes.Add(AsyncHelper->PipelinePostImportTasks[GraphPipelineTaskIndex]);
diff --git a/Engine/Source/Runtime/Interchange/Engine/Private/Tasks/InterchangeTaskPipeline.cpp b/Engine/Source/Runtime/Interchange/Engine/Private/Tasks/InterchangeTaskPipeline.cpp
--- a/Engine/Source/Runtime/Interchange/Engine/Private/Tasks/InterchangeTaskPipeline.cpp
+++ b/Engine/Source/Runtime/Interchange/Engine/Private/Tasks/InterchangeTaskPipeline.cpp
@@ -23,6 +23,8 @@ void UE::Interchange::FTaskPipelinePreImport::DoTask(ENamedThreads::Type Current
 
 	if (UInterchangePipelineBase* Pipeline = PipelineBase.Get())
 	{
+		Pipeline->SetResultsContainer(AsyncHelper->AssetImportResult->GetResults());
+
 		for (int32 GraphIndex = 0; GraphIndex < AsyncHelper->BaseNodeContainers.Num(); ++GraphIndex)
 		{
 			//Verify if the task was cancel
@@ -86,6 +88,8 @@ void UE::Interchange::FTaskPipelinePostImport::DoTask(ENamedThreads::Type Curren
 		return;
 	}
 	UInterchangePipelineBase* Pipeline = AsyncHelper->Pipelines[PipelineIndex];
+	Pipeline->SetResultsContainer(AsyncHelper->AssetImportResult->GetResults());
+
 	//Call the pipeline outside of the lock, we do this in case the pipeline take a long time. We call it for each asset created by this import
 	for (int32 AssetIndex = 0; AssetIndex < ImportAssets.Num(); ++AssetIndex)
 	{
diff --git a/Engine/Source/Runtime/Interchange/Engine/Private/Tasks/InterchangeTaskTranslator.cpp b/Engine/Source/Runtime/Interchange/Engine/Private/Tasks/InterchangeTaskTranslator.cpp
--- a/Engine/Source/Runtime/Interchange/Engine/Private/Tasks/InterchangeTaskTranslator.cpp
+++ b/Engine/Source/Runtime/Interchange/Engine/Private/Tasks/InterchangeTaskTranslator.cpp
@@ -25,11 +25,10 @@ void UE::Interchange::FTaskTranslator::DoTask(ENamedThreads::Type CurrentThread,
 	{
 		return;
 	}
-	const UInterchangeSourceData* SourceData = AsyncHelper->SourceDatas.IsValidIndex(SourceIndex) ? AsyncHelper->SourceDatas[SourceIndex] : nullptr;
-	if (!SourceData)
-	{
-		return;
-	}
+
+	check(AsyncHelper->AssetImportResult->GetResults() != nullptr);
+	check(Translator->SourceData != nullptr);
+	Translator->SetResultsContainer(AsyncHelper->AssetImportResult->GetResults());
 
 	if (!AsyncHelper->BaseNodeContainers.IsValidIndex(SourceIndex) || !AsyncHelper->BaseNodeContainers[SourceIndex].IsValid())
 	{
@@ -44,5 +43,5 @@ void UE::Interchange::FTaskTranslator::DoTask(ENamedThreads::Type CurrentThread,
 
 	//Translate the source data
 	UInterchangeBaseNodeContainer& BaseNodeContainer = *(AsyncHelper->BaseNodeContainers[SourceIndex].Get());
-	Translator->Translate(SourceData, BaseNodeContainer);
+	Translator->Translate(BaseNodeContainer);
 }
diff --git a/Engine/Source/Runtime/Interchange/Engine/Public/InterchangeManager.h b/Engine/Source/Runtime/Interchange/Engine/Public/InterchangeManager.h
--- a/Engine/Source/Runtime/Interchange/Engine/Public/InterchangeManager.h
+++ b/Engine/Source/Runtime/Interchange/Engine/Public/InterchangeManager.h
@@ -9,9 +9,8 @@
 #include "Delegates/DelegateCombinations.h"
 #include "HAL/Thread.h"
 #include "HAL/ThreadSafeBool.h"
-#include "InterchangeFactoryBase.h"
-#include "InterchangePipelineBase.h"
 #include "InterchangePipelineConfigurationBase.h"
+#include "InterchangeResultsContainer.h"
 #include "InterchangeSourceData.h"
 #include "InterchangeTranslatorBase.h"
 #include "InterchangeWriterBase.h"
@@ -25,6 +24,9 @@
 #include "InterchangeManager.generated.h"
 
 class FAsyncTaskNotification;
+class UInterchangeFactoryBase;
+class UInterchangePipelineBase;
+
 namespace UE
 {
 	namespace Interchange
@@ -104,6 +106,9 @@ namespace UE
 			// Helper to get the first asset of a certain class. Use when expecting a single asset of that class to be imported since the order isn't deterministic.
 			UObject* GetFirstAssetOfClass(UClass* InClass) const;
 
+			// Return the results of this asset import operation
+			UInterchangeResultsContainer* GetResults() const { return Results; }
+
 			// Adds an asset to the list of imported assets.
 			void AddImportedAsset(UObject* ImportedAsset);
 
@@ -119,6 +124,7 @@ namespace UE
 
 			TArray< UObject* > ImportedAssets;
 			mutable FRWLock ImportedAssetsRWLock;
+			UInterchangeResultsContainer* Results;
 
 			FGraphEventRef GraphEvent; // WaitUntilDone waits for this event to be triggered.
 
@@ -142,13 +148,11 @@ namespace UE
 
 			//The following Arrays are per source data
 			TArray<TStrongObjectPtr<UInterchangeBaseNodeContainer>> BaseNodeContainers;
-			TArray<UInterchangeSourceData* > SourceDatas;
-			TArray<UInterchangeTranslatorBase* > Translators;
-			TArray<UInterchangeFactoryBase* > Factories;
+			TArray<UInterchangeSourceData*> SourceDatas;
+			TArray<UInterchangeTranslatorBase*> Translators;
 
 			//Pipelines array is not per source data 
-			TArray<UInterchangePipelineBase* > Pipelines;
-			
+			TArray<UInterchangePipelineBase*> Pipelines;
 
 			TArray<FGraphEventRef> TranslatorTasks;
 			TArray<FGraphEventRef> PipelinePreImportTasks;
@@ -164,6 +168,10 @@ namespace UE
 			FCriticalSection CreatedPackagesLock;
 			TMap<FString, UPackage*> CreatedPackages;
 
+			// Created factories map, Key is factory node UID
+			FCriticalSection CreatedFactoriesLock;
+			TMap<FString, UInterchangeFactoryBase*> CreatedFactories;
+
 			struct FImportedAssetInfo
 			{
 				UObject* ImportAsset;
@@ -234,12 +242,14 @@ public:
 	DECLARE_MULTICAST_DELEGATE_OneParam(FInterchangeOnAssetPostImport, UObject*);
 	/** delegate type fired when new assets have been reimported. Note: InCreatedObject can be NULL if import failed. Params: UObject* InCreatedObject */
 	DECLARE_MULTICAST_DELEGATE_OneParam(FInterchangeOnAssetPostReimport, UObject*);
-	
+	/** delegate type fired when import results in an error */
+	DECLARE_MULTICAST_DELEGATE_OneParam(FInterchangeOnBatchImportComplete, TStrongObjectPtr<UInterchangeResultsContainer>);
 
 	// Delegates used to register and unregister
 
 	FInterchangeOnAssetPostImport OnAssetPostImport;
 	FInterchangeOnAssetPostReimport OnAssetPostReimport;
+	FInterchangeOnBatchImportComplete OnBatchImportComplete;
 	// Called when before the application is exiting.
 	FSimpleMulticastDelegate OnPreDestroyInterchangeManager;
 
@@ -335,20 +345,10 @@ public:
 	/**
 	* Script helper to get a registered factory for a specified class
 	* @Param FactoryClass: The class we search a registerd factory
-	* @return: if found, we return the factory that is registered. Return NULL if nothing found.
+	* @return: if found, we return the factory class that is registered. Return NULL if nothing found.
 	*/
 	UFUNCTION(BlueprintCallable, Category = "Interchange | Import Manager")
-	UInterchangeFactoryBase* GetRegisterFactory(UClass* FactoryClass)
-	{
-		for (auto Kvp : RegisteredFactories)
-		{
-			if (FactoryClass->IsChildOf(Kvp.Key))
-			{
-				return Kvp.Value;
-			}
-		}
-		return nullptr;
-	}
+	const UClass* GetRegisteredFactoryClass(const UClass* ClassToMake) const;
 
 	/**
 	 * Return an FImportAsynHelper pointer. The pointer is deleted when ReleaseAsyncHelper is call.
@@ -447,16 +447,16 @@ private:
 
 	//The manager will create translator at every import, translator must be able to retrieve payload information when the factory ask for it.
 	//The translator stored has value is only use to know if we can use this type of translator.
-	UPROPERTY()
-	TArray<TObjectPtr<UInterchangeTranslatorBase>> RegisteredTranslators;
+//	UPROPERTY()
+//	TArray<TObjectPtr<UInterchangeTranslatorBase>> RegisteredTranslators;
 	
 	//The manager will create only one pipeline per type
-	UPROPERTY()
-	TMap<TObjectPtr<const UClass>, TObjectPtr<UInterchangePipelineBase> > RegisteredPipelines;
+//	UPROPERTY()
+//	TMap<TObjectPtr<const UClass>, TObjectPtr<UInterchangePipelineBase> > RegisteredPipelines;
 
 	//The manager will create only one factory per type
 	UPROPERTY()
-	TMap<TObjectPtr<const UClass>, TObjectPtr<UInterchangeFactoryBase> > RegisteredFactories;
+	TMap<TObjectPtr<const UClass>, TObjectPtr<const UClass> > RegisteredFactoryClasses;
 
 	//The manager will create only one writer per type
 	UPROPERTY()
diff --git a/Engine/Source/Runtime/StaticMeshDescription/Private/StaticMeshDescription.cpp b/Engine/Source/Runtime/StaticMeshDescription/Private/StaticMeshDescription.cpp
--- a/Engine/Source/Runtime/StaticMeshDescription/Private/StaticMeshDescription.cpp
+++ b/Engine/Source/Runtime/StaticMeshDescription/Private/StaticMeshDescription.cpp
@@ -108,6 +108,12 @@ void UStaticMeshDescription::CreateCube(FVector Center, FVector HalfExtents, FPo
 		UVs[VertexInstanceIDs[2]] = FVector2D(1.0f, 1.0f);
 		UVs[VertexInstanceIDs[3]] = FVector2D(0.0f, 1.0f);
 
+		TVertexInstanceAttributesRef<FVector4> Colors = GetVertexInstanceColors();
+		Colors[VertexInstanceIDs[0]] = FLinearColor::Red;
+		Colors[VertexInstanceIDs[1]] = FLinearColor::Green;
+		Colors[VertexInstanceIDs[2]] = FLinearColor::Blue;
+		Colors[VertexInstanceIDs[3]] = FLinearColor::White;
+
 		TArray<FEdgeID> EdgeIDs;
 		EdgeIDs.Reserve(4);
 
