Fix static analysis warning.

diff --git a/Engine/Source/Runtime/Interchange/Engine/Private/Tasks/InterchangeTaskCompletion.cpp b/Engine/Source/Runtime/Interchange/Engine/Private/Tasks/InterchangeTaskCompletion.cpp
--- a/Engine/Source/Runtime/Interchange/Engine/Private/Tasks/InterchangeTaskCompletion.cpp
+++ b/Engine/Source/Runtime/Interchange/Engine/Private/Tasks/InterchangeTaskCompletion.cpp
@@ -65,14 +65,16 @@ void UE::Interchange::FTaskCompletion::DoTask(ENamedThreads::Type CurrentThread,
 				AssetInfo.Factory->PostImportGameThreadCallback(Arguments);
 			}
 
-			if (Asset)
+			if (Asset == nullptr)
 			{
-				UInterchangeResultSuccess* Message = Results->Add<UInterchangeResultSuccess>();
-				Message->SourceAssetName = AsyncHelper->SourceDatas[SourceIndex]->GetFilename();
-				Message->DestinationAssetName = Asset->GetPathName();
-				Message->AssetType = Asset->GetClass();
+				continue;
 			}
 
+			UInterchangeResultSuccess* Message = Results->Add<UInterchangeResultSuccess>();
+			Message->SourceAssetName = AsyncHelper->SourceDatas[SourceIndex]->GetFilename();
+			Message->DestinationAssetName = Asset->GetPathName();
+			Message->AssetType = Asset->GetClass();
+
 			//Clear any async flag from the created asset
 			const EInternalObjectFlags AsyncFlags = EInternalObjectFlags::Async | EInternalObjectFlags::AsyncLoading;
 			Asset->ClearInternalFlags(AsyncFlags);
