Fixed static analysis warning.

diff --git a/Engine/Source/Runtime/Interchange/Core/Private/InterchangeResultsContainer.cpp b/Engine/Source/Runtime/Interchange/Core/Private/InterchangeResultsContainer.cpp
--- a/Engine/Source/Runtime/Interchange/Core/Private/InterchangeResultsContainer.cpp
+++ b/Engine/Source/Runtime/Interchange/Core/Private/InterchangeResultsContainer.cpp
@@ -1,7 +1,5 @@
 // Copyright Epic Games, Inc. All Rights Reserved.
 
-#pragma once
-
 #include "InterchangeResultsContainer.h"
 
 
