Changed checks to ensures to prevent irritating editor crashes if InterchangeWorker isn't up to date.

diff --git a/Engine/Source/Runtime/Interchange/Core/Public/InterchangeTranslatorBase.h b/Engine/Source/Runtime/Interchange/Core/Public/InterchangeTranslatorBase.h
--- a/Engine/Source/Runtime/Interchange/Core/Public/InterchangeTranslatorBase.h
+++ b/Engine/Source/Runtime/Interchange/Core/Public/InterchangeTranslatorBase.h
@@ -60,8 +60,8 @@ public:
 	template <typename T>
 	T* AddMessage() const
 	{
-		check(Results != nullptr);
-		check(SourceData != nullptr);
+		ensure(Results != nullptr);
+		ensure(SourceData != nullptr);
 		T* Item = Results->Add<T>();
 		Item->SourceAssetName = SourceData->GetFilename();
 		return Item;
@@ -70,8 +70,8 @@ public:
 
 	void AddMessage(UInterchangeResult* Item) const
 	{
-		check(Results != nullptr);
-		check(SourceData != nullptr);
+		ensure(Results != nullptr);
+		ensure(SourceData != nullptr);
 		Results->Add(Item);
 		Item->SourceAssetName = SourceData->GetFilename();
 	}
