InterchangeBaseNode: rename GetAssetClass to GetObjectClass

diff --git a/Engine/Plugins/Editor/InterchangeEditor/Source/Pipelines/Private/SInterchangeGraphInspectorWindow.cpp b/Engine/Plugins/Editor/InterchangeEditor/Source/Pipelines/Private/SInterchangeGraphInspectorWindow.cpp
--- a/Engine/Plugins/Editor/InterchangeEditor/Source/Pipelines/Private/SInterchangeGraphInspectorWindow.cpp
+++ b/Engine/Plugins/Editor/InterchangeEditor/Source/Pipelines/Private/SInterchangeGraphInspectorWindow.cpp
@@ -86,7 +86,7 @@ public:
 		check(InterchangeNode);
 		check(InterchangeBaseNodeContainer);
 
-		UClass* IconClass = InterchangeNode->GetAssetClass();
+		UClass* IconClass = InterchangeNode->GetObjectClass();
 		const FSlateBrush* TypeIcon = nullptr;
 		FName IconName = InterchangeNode->GetIconName();
 		if (IconClass)
@@ -106,9 +106,9 @@ public:
 
 		//Prepare the tooltip
 		FString Tooltip = InterchangeNode->GetDisplayLabel();
-		if (InterchangeNode->GetAssetClass())
+		if (InterchangeNode->GetObjectClass())
 		{
-			Tooltip += TEXT(" [") + InterchangeNode->GetAssetClass()->GetName() + TEXT("]");
+			Tooltip += TEXT(" [") + InterchangeNode->GetObjectClass()->GetName() + TEXT("]");
 		}
 
 		this->ChildSlot
diff --git a/Engine/Plugins/Runtime/Interchange/Source/Import/Private/Material/InterchangeMaterialFactory.cpp b/Engine/Plugins/Runtime/Interchange/Source/Import/Private/Material/InterchangeMaterialFactory.cpp
--- a/Engine/Plugins/Runtime/Interchange/Source/Import/Private/Material/InterchangeMaterialFactory.cpp
+++ b/Engine/Plugins/Runtime/Interchange/Source/Import/Private/Material/InterchangeMaterialFactory.cpp
@@ -34,7 +34,7 @@ UObject* UInterchangeMaterialFactory::CreateEmptyAsset(const FCreateAssetParams&
 	UObject* Material = nullptr;
 
 #if WITH_EDITORONLY_DATA
-	if (!Arguments.AssetNode || !Arguments.AssetNode->GetAssetClass()->IsChildOf(GetFactoryClass()))
+	if (!Arguments.AssetNode || !Arguments.AssetNode->GetObjectClass()->IsChildOf(GetFactoryClass()))
 	{
 		return nullptr;
 	}
@@ -52,7 +52,7 @@ UObject* UInterchangeMaterialFactory::CreateEmptyAsset(const FCreateAssetParams&
 		return nullptr;
 	}
 
-	const UClass* MaterialClass = MaterialFactoryNode->GetAssetClass();
+	const UClass* MaterialClass = MaterialFactoryNode->GetObjectClass();
 	if (!ensure(MaterialClass && MaterialClass->IsChildOf(GetFactoryClass())))
 	{
 		return nullptr;
@@ -91,7 +91,7 @@ UObject* UInterchangeMaterialFactory::CreateAsset(const UInterchangeMaterialFact
 
 #else
 
-	if (!Arguments.AssetNode || !Arguments.AssetNode->GetAssetClass()->IsChildOf(GetFactoryClass()))
+	if (!Arguments.AssetNode || !Arguments.AssetNode->GetObjectClass()->IsChildOf(GetFactoryClass()))
 	{
 		return nullptr;
 	}
@@ -109,7 +109,7 @@ UObject* UInterchangeMaterialFactory::CreateAsset(const UInterchangeMaterialFact
 		return nullptr;
 	}
 
-	const UClass* MaterialClass = MaterialFactoryNode->GetAssetClass();
+	const UClass* MaterialClass = MaterialFactoryNode->GetObjectClass();
 	check(MaterialClass && MaterialClass->IsChildOf(GetFactoryClass()));
 
 	// create an asset if it doesn't exist
diff --git a/Engine/Plugins/Runtime/Interchange/Source/Import/Private/Mesh/InterchangePhysicsAssetFactory.cpp b/Engine/Plugins/Runtime/Interchange/Source/Import/Private/Mesh/InterchangePhysicsAssetFactory.cpp
--- a/Engine/Plugins/Runtime/Interchange/Source/Import/Private/Mesh/InterchangePhysicsAssetFactory.cpp
+++ b/Engine/Plugins/Runtime/Interchange/Source/Import/Private/Mesh/InterchangePhysicsAssetFactory.cpp
@@ -19,7 +19,7 @@ UObject* UInterchangePhysicsAssetFactory::CreateEmptyAsset(const FCreateAssetPar
 	UObject* PhysicsAsset = nullptr;
 
 #if WITH_EDITORONLY_DATA
-	if (!Arguments.AssetNode || !Arguments.AssetNode->GetAssetClass()->IsChildOf(GetFactoryClass()))
+	if (!Arguments.AssetNode || !Arguments.AssetNode->GetObjectClass()->IsChildOf(GetFactoryClass()))
 	{
 		return nullptr;
 	}
@@ -65,7 +65,7 @@ UObject* UInterchangePhysicsAssetFactory::CreateAsset(const UInterchangePhysicsA
 
 #else
 
-	if (!Arguments.AssetNode || !Arguments.AssetNode->GetAssetClass()->IsChildOf(GetFactoryClass()))
+	if (!Arguments.AssetNode || !Arguments.AssetNode->GetObjectClass()->IsChildOf(GetFactoryClass()))
 	{
 		return nullptr;
 	}
@@ -76,7 +76,7 @@ UObject* UInterchangePhysicsAssetFactory::CreateAsset(const UInterchangePhysicsA
 		return nullptr;
 	}
 
-	const UClass* PhysicsAssetClass = PhysicsAssetNode->GetAssetClass();
+	const UClass* PhysicsAssetClass = PhysicsAssetNode->GetObjectClass();
 	check(PhysicsAssetClass && PhysicsAssetClass->IsChildOf(GetFactoryClass()));
 
 	// create an asset if it doesn't exist
diff --git a/Engine/Plugins/Runtime/Interchange/Source/Import/Private/Mesh/InterchangeSkeletalMeshFactory.cpp b/Engine/Plugins/Runtime/Interchange/Source/Import/Private/Mesh/InterchangeSkeletalMeshFactory.cpp
--- a/Engine/Plugins/Runtime/Interchange/Source/Import/Private/Mesh/InterchangeSkeletalMeshFactory.cpp
+++ b/Engine/Plugins/Runtime/Interchange/Source/Import/Private/Mesh/InterchangeSkeletalMeshFactory.cpp
@@ -677,7 +677,7 @@ UObject* UInterchangeSkeletalMeshFactory::CreateEmptyAsset(const FCreateAssetPar
 
 #else
 	USkeletalMesh* SkeletalMesh = nullptr;
-	if (!Arguments.AssetNode || !Arguments.AssetNode->GetAssetClass()->IsChildOf(GetFactoryClass()))
+	if (!Arguments.AssetNode || !Arguments.AssetNode->GetObjectClass()->IsChildOf(GetFactoryClass()))
 	{
 		return nullptr;
 	}
@@ -725,7 +725,7 @@ UObject* UInterchangeSkeletalMeshFactory::CreateAsset(const UInterchangeSkeletal
 
 #else
 
-	if (!Arguments.AssetNode || !Arguments.AssetNode->GetAssetClass()->IsChildOf(GetFactoryClass()))
+	if (!Arguments.AssetNode || !Arguments.AssetNode->GetObjectClass()->IsChildOf(GetFactoryClass()))
 	{
 		return nullptr;
 	}
@@ -743,7 +743,7 @@ UObject* UInterchangeSkeletalMeshFactory::CreateAsset(const UInterchangeSkeletal
 		return nullptr;
 	}
 
-	const UClass* SkeletalMeshClass = SkeletalMeshFactoryNode->GetAssetClass();
+	const UClass* SkeletalMeshClass = SkeletalMeshFactoryNode->GetObjectClass();
 	check(SkeletalMeshClass && SkeletalMeshClass->IsChildOf(GetFactoryClass()));
 
 	// create an asset if it doesn't exist
diff --git a/Engine/Plugins/Runtime/Interchange/Source/Import/Private/Mesh/InterchangeSkeletonFactory.cpp b/Engine/Plugins/Runtime/Interchange/Source/Import/Private/Mesh/InterchangeSkeletonFactory.cpp
--- a/Engine/Plugins/Runtime/Interchange/Source/Import/Private/Mesh/InterchangeSkeletonFactory.cpp
+++ b/Engine/Plugins/Runtime/Interchange/Source/Import/Private/Mesh/InterchangeSkeletonFactory.cpp
@@ -20,7 +20,7 @@ UObject* UInterchangeSkeletonFactory::CreateEmptyAsset(const FCreateAssetParams&
 	UObject* Skeleton = nullptr;
 
 #if WITH_EDITORONLY_DATA
-	if (!Arguments.AssetNode || !Arguments.AssetNode->GetAssetClass()->IsChildOf(GetFactoryClass()))
+	if (!Arguments.AssetNode || !Arguments.AssetNode->GetObjectClass()->IsChildOf(GetFactoryClass()))
 	{
 		return nullptr;
 	}
@@ -66,7 +66,7 @@ UObject* UInterchangeSkeletonFactory::CreateAsset(const UInterchangeSkeletonFact
 
 #else
 
-	if (!Arguments.AssetNode || !Arguments.AssetNode->GetAssetClass()->IsChildOf(GetFactoryClass()))
+	if (!Arguments.AssetNode || !Arguments.AssetNode->GetObjectClass()->IsChildOf(GetFactoryClass()))
 	{
 		return nullptr;
 	}
@@ -77,7 +77,7 @@ UObject* UInterchangeSkeletonFactory::CreateAsset(const UInterchangeSkeletonFact
 		return nullptr;
 	}
 
-	const UClass* SkeletonClass = SkeletonNode->GetAssetClass();
+	const UClass* SkeletonClass = SkeletonNode->GetObjectClass();
 	check(SkeletonClass && SkeletonClass->IsChildOf(GetFactoryClass()));
 
 	// create an asset if it doesn't exist
diff --git a/Engine/Plugins/Runtime/Interchange/Source/Import/Private/Mesh/InterchangeStaticMeshFactory.cpp b/Engine/Plugins/Runtime/Interchange/Source/Import/Private/Mesh/InterchangeStaticMeshFactory.cpp
--- a/Engine/Plugins/Runtime/Interchange/Source/Import/Private/Mesh/InterchangeStaticMeshFactory.cpp
+++ b/Engine/Plugins/Runtime/Interchange/Source/Import/Private/Mesh/InterchangeStaticMeshFactory.cpp
@@ -42,7 +42,7 @@ UObject* UInterchangeStaticMeshFactory::CreateEmptyAsset(const FCreateAssetParam
 
 #else
 	UStaticMesh* StaticMesh = nullptr;
-	if (!Arguments.AssetNode || !Arguments.AssetNode->GetAssetClass()->IsChildOf(GetFactoryClass()))
+	if (!Arguments.AssetNode || !Arguments.AssetNode->GetObjectClass()->IsChildOf(GetFactoryClass()))
 	{
 		return nullptr;
 	}
@@ -88,7 +88,7 @@ UObject* UInterchangeStaticMeshFactory::CreateAsset(const UInterchangeStaticMesh
 
 #else
 
-	if (!Arguments.AssetNode || !Arguments.AssetNode->GetAssetClass()->IsChildOf(GetFactoryClass()))
+	if (!Arguments.AssetNode || !Arguments.AssetNode->GetObjectClass()->IsChildOf(GetFactoryClass()))
 	{
 		return nullptr;
 	}
@@ -106,7 +106,7 @@ UObject* UInterchangeStaticMeshFactory::CreateAsset(const UInterchangeStaticMesh
 		return nullptr;
 	}
 
-	const UClass* StaticMeshClass = StaticMeshFactoryNode->GetAssetClass();
+	const UClass* StaticMeshClass = StaticMeshFactoryNode->GetObjectClass();
 	check(StaticMeshClass && StaticMeshClass->IsChildOf(GetFactoryClass()));
 
 	// create an asset if it doesn't exist
diff --git a/Engine/Plugins/Runtime/Interchange/Source/Import/Private/Texture/InterchangeTextureFactory.cpp b/Engine/Plugins/Runtime/Interchange/Source/Import/Private/Texture/InterchangeTextureFactory.cpp
--- a/Engine/Plugins/Runtime/Interchange/Source/Import/Private/Texture/InterchangeTextureFactory.cpp
+++ b/Engine/Plugins/Runtime/Interchange/Source/Import/Private/Texture/InterchangeTextureFactory.cpp
@@ -523,7 +523,7 @@ UObject* UInterchangeTextureFactory::CreateEmptyAsset(const FCreateAssetParams&
 		return nullptr;
 	}
 
-	const UClass* TextureClass = Arguments.AssetNode->GetAssetClass();
+	const UClass* TextureClass = Arguments.AssetNode->GetObjectClass();
 	if (!TextureClass || !TextureClass->IsChildOf(UTexture::StaticClass()))
 	{
 		UE_LOG(LogInterchangeImport, Error, TEXT("UInterchangeTextureFactory: Asset node parameter class doesnt derive from UTexture."));
@@ -596,7 +596,7 @@ UObject* UInterchangeTextureFactory::CreateAsset(const UInterchangeTextureFactor
 		return nullptr;
 	}
 
-	const UClass* TextureClass = Arguments.AssetNode->GetAssetClass();
+	const UClass* TextureClass = Arguments.AssetNode->GetObjectClass();
 	if (!TextureClass || !TextureClass->IsChildOf(UTexture::StaticClass()))
 	{
 		UE_LOG(LogInterchangeImport, Error, TEXT("UInterchangeTextureFactory: Asset node parameter class doesnt derive from UTexture."));
diff --git a/Engine/Plugins/Runtime/Interchange/Source/Nodes/Public/InterchangeMaterialFactoryNode.h b/Engine/Plugins/Runtime/Interchange/Source/Nodes/Public/InterchangeMaterialFactoryNode.h
--- a/Engine/Plugins/Runtime/Interchange/Source/Nodes/Public/InterchangeMaterialFactoryNode.h
+++ b/Engine/Plugins/Runtime/Interchange/Source/Nodes/Public/InterchangeMaterialFactoryNode.h
@@ -57,7 +57,7 @@ public:
 
 	/** Get the class this node want to create */
 	UFUNCTION(BlueprintCallable, Category = "Interchange | Node | MaterialFactory")
-	virtual class UClass* GetAssetClass() const override
+	virtual class UClass* GetObjectClass() const override
 	{
 		ensure(bIsMaterialNodeClassInitialized);
 #if WITH_ENGINE
diff --git a/Engine/Plugins/Runtime/Interchange/Source/Nodes/Public/InterchangePhysicsAssetFactoryNode.h b/Engine/Plugins/Runtime/Interchange/Source/Nodes/Public/InterchangePhysicsAssetFactoryNode.h
--- a/Engine/Plugins/Runtime/Interchange/Source/Nodes/Public/InterchangePhysicsAssetFactoryNode.h
+++ b/Engine/Plugins/Runtime/Interchange/Source/Nodes/Public/InterchangePhysicsAssetFactoryNode.h
@@ -65,7 +65,7 @@ public:
 
 	/** Get the class this node want to create */
 	UFUNCTION(BlueprintCallable, Category = "Interchange | Node | Skeleton")
-	virtual class UClass* GetAssetClass() const override
+	virtual class UClass* GetObjectClass() const override
 	{
 		ensure(bIsNodeClassInitialized);
 #if WITH_ENGINE
diff --git a/Engine/Plugins/Runtime/Interchange/Source/Nodes/Public/InterchangeSkeletalMeshFactoryNode.h b/Engine/Plugins/Runtime/Interchange/Source/Nodes/Public/InterchangeSkeletalMeshFactoryNode.h
--- a/Engine/Plugins/Runtime/Interchange/Source/Nodes/Public/InterchangeSkeletalMeshFactoryNode.h
+++ b/Engine/Plugins/Runtime/Interchange/Source/Nodes/Public/InterchangeSkeletalMeshFactoryNode.h
@@ -107,7 +107,7 @@ public:
 
 	/** Get the class this node want to create */
 	UFUNCTION(BlueprintCallable, Category = "Interchange | Node | SkeletalMesh")
-	virtual class UClass* GetAssetClass() const override
+	virtual class UClass* GetObjectClass() const override
 	{
 		ensure(bIsNodeClassInitialized);
 #if WITH_ENGINE
diff --git a/Engine/Plugins/Runtime/Interchange/Source/Nodes/Public/InterchangeSkeletonFactoryNode.h b/Engine/Plugins/Runtime/Interchange/Source/Nodes/Public/InterchangeSkeletonFactoryNode.h
--- a/Engine/Plugins/Runtime/Interchange/Source/Nodes/Public/InterchangeSkeletonFactoryNode.h
+++ b/Engine/Plugins/Runtime/Interchange/Source/Nodes/Public/InterchangeSkeletonFactoryNode.h
@@ -65,7 +65,7 @@ public:
 
 	/** Get the class this node want to create */
 	UFUNCTION(BlueprintCallable, Category = "Interchange | Node | Skeleton")
-	virtual class UClass* GetAssetClass() const override
+	virtual class UClass* GetObjectClass() const override
 	{
 		ensure(bIsNodeClassInitialized);
 #if WITH_ENGINE
diff --git a/Engine/Plugins/Runtime/Interchange/Source/Nodes/Public/InterchangeStaticMeshFactoryNode.h b/Engine/Plugins/Runtime/Interchange/Source/Nodes/Public/InterchangeStaticMeshFactoryNode.h
--- a/Engine/Plugins/Runtime/Interchange/Source/Nodes/Public/InterchangeStaticMeshFactoryNode.h
+++ b/Engine/Plugins/Runtime/Interchange/Source/Nodes/Public/InterchangeStaticMeshFactoryNode.h
@@ -105,7 +105,7 @@ public:
 
 	/** Get the class this node want to create */
 	UFUNCTION(BlueprintCallable, Category = "Interchange | Node | StaticMesh")
-	virtual class UClass* GetAssetClass() const override
+	virtual class UClass* GetObjectClass() const override
 	{
 		ensure(bIsNodeClassInitialized);
 #if WITH_ENGINE
diff --git a/Engine/Plugins/Runtime/Interchange/Source/Nodes/Public/InterchangeTextureFactoryNode.h b/Engine/Plugins/Runtime/Interchange/Source/Nodes/Public/InterchangeTextureFactoryNode.h
--- a/Engine/Plugins/Runtime/Interchange/Source/Nodes/Public/InterchangeTextureFactoryNode.h
+++ b/Engine/Plugins/Runtime/Interchange/Source/Nodes/Public/InterchangeTextureFactoryNode.h
@@ -74,7 +74,7 @@ public:
 
 	/** Get the class this node want to create */
 	UFUNCTION(BlueprintCallable, Category = "Interchange | Node | Texture")
-	virtual class UClass* GetAssetClass() const override
+	virtual class UClass* GetObjectClass() const override
 	{
 		ensure(bIsTextureNodeClassInitialized);
 #if WITH_ENGINE
diff --git a/Engine/Source/Runtime/Interchange/Core/Private/Nodes/InterchangeBaseNode.cpp b/Engine/Source/Runtime/Interchange/Core/Private/Nodes/InterchangeBaseNode.cpp
--- a/Engine/Source/Runtime/Interchange/Core/Private/Nodes/InterchangeBaseNode.cpp
+++ b/Engine/Source/Runtime/Interchange/Core/Private/Nodes/InterchangeBaseNode.cpp
@@ -177,7 +177,7 @@ FGuid UInterchangeBaseNode::GetHash() const
 	return Attributes->GetStorageHash();
 }
 
-class UClass* UInterchangeBaseNode::GetAssetClass() const
+UClass* UInterchangeBaseNode::GetObjectClass() const
 {
 	return nullptr;
 }
diff --git a/Engine/Source/Runtime/Interchange/Core/Public/Nodes/InterchangeBaseNode.h b/Engine/Source/Runtime/Interchange/Core/Public/Nodes/InterchangeBaseNode.h
--- a/Engine/Source/Runtime/Interchange/Core/Public/Nodes/InterchangeBaseNode.h
+++ b/Engine/Source/Runtime/Interchange/Core/Public/Nodes/InterchangeBaseNode.h
@@ -1089,9 +1089,9 @@ public:
 	virtual FGuid GetHash() const;
 
 	/**
-	 * Optional, Any node that can import/export an asset should return the UClass of the asset so we can find factory/writer
+	 * Optional, Any node that can import/export an object should return the UClass of the object so we can find factory/writer
 	 */
-	virtual class UClass* GetAssetClass() const;
+	virtual class UClass* GetObjectClass() const;
 
 	/**
 	 * Optional, Any node that can import/export an asset should set the proper name we will give to the asset.
diff --git a/Engine/Source/Runtime/Interchange/Engine/Private/Tasks/InterchangeTaskCreateAsset.cpp b/Engine/Source/Runtime/Interchange/Engine/Private/Tasks/InterchangeTaskCreateAsset.cpp
--- a/Engine/Source/Runtime/Interchange/Engine/Private/Tasks/InterchangeTaskCreateAsset.cpp
+++ b/Engine/Source/Runtime/Interchange/Engine/Private/Tasks/InterchangeTaskCreateAsset.cpp
@@ -90,7 +90,7 @@ void UE::Interchange::FTaskCreatePackage::DoTask(ENamedThreads::Type CurrentThre
 			UInterchangeResultError_Generic* Message = Factory->AddMessage<UInterchangeResultError_Generic>();
 			Message->SourceAssetName = AsyncHelper->SourceDatas[SourceIndex]->GetFilename();
 			Message->DestinationAssetName = AssetName;
-			Message->AssetType = Node->GetAssetClass();
+			Message->AssetType = Node->GetObjectClass();
 			Message->Text = NSLOCTEXT("Interchange", "CannotFindPackageDuringReimport", "Cannot find an existing package.");
 
 			//Skip this asset
@@ -106,7 +106,7 @@ void UE::Interchange::FTaskCreatePackage::DoTask(ENamedThreads::Type CurrentThre
 			UInterchangeResultError_Generic* Message = Factory->AddMessage<UInterchangeResultError_Generic>();
 			Message->SourceAssetName = AsyncHelper->SourceDatas[SourceIndex]->GetFilename();
 			Message->DestinationAssetName = AssetName;
-			Message->AssetType = Node->GetAssetClass();
+			Message->AssetType = Node->GetObjectClass();
 			Message->Text = NSLOCTEXT("Interchange", "MapExistsWithSameName", "You cannot create an asset with this name, as there is already a map file with the same name in this folder.");
 
 			//Skip this asset
@@ -119,7 +119,7 @@ void UE::Interchange::FTaskCreatePackage::DoTask(ENamedThreads::Type CurrentThre
 			UInterchangeResultError_Generic* Message = Factory->AddMessage<UInterchangeResultError_Generic>();
 			Message->SourceAssetName = AsyncHelper->SourceDatas[SourceIndex]->GetFilename();
 			Message->DestinationAssetName = AssetName;
-			Message->AssetType = Node->GetAssetClass();
+			Message->AssetType = Node->GetObjectClass();
 			Message->Text = FText::Format(NSLOCTEXT("Interchange", "CouldntCreatePackage", "It was not possible to create a package named '{0}'; the asset will not be imported."), FText::FromString(PackageName));
 
 			//Skip this asset
@@ -222,7 +222,7 @@ void UE::Interchange::FTaskCreateAsset::DoTask(ENamedThreads::Type CurrentThread
 			UInterchangeResultError_Generic* Message = Factory->AddMessage<UInterchangeResultError_Generic>();
 			Message->SourceAssetName = AsyncHelper->SourceDatas[SourceIndex]->GetFilename();
 			Message->DestinationAssetName = AssetName;
-			Message->AssetType = Node->GetAssetClass();
+			Message->AssetType = Node->GetObjectClass();
 			Message->Text = NSLOCTEXT("Interchange", "BadPackage", "It was not possible to create the asset as its package was not created correctly.");
 
 			return;
@@ -232,7 +232,7 @@ void UE::Interchange::FTaskCreateAsset::DoTask(ENamedThreads::Type CurrentThread
 		{
 			UInterchangeResultError_Generic* Message = Factory->AddMessage<UInterchangeResultError_Generic>();
 			Message->DestinationAssetName = AssetName;
-			Message->AssetType = Node->GetAssetClass();
+			Message->AssetType = Node->GetObjectClass();
 			Message->Text = NSLOCTEXT("Interchange", "SourceDataOrTranslatorInvalid", "It was not possible to create the asset as its translator was not created correctly.");
 
 			return;
diff --git a/Engine/Source/Runtime/Interchange/Engine/Private/Tasks/InterchangeTaskParsing.cpp b/Engine/Source/Runtime/Interchange/Engine/Private/Tasks/InterchangeTaskParsing.cpp
--- a/Engine/Source/Runtime/Interchange/Engine/Private/Tasks/InterchangeTaskParsing.cpp
+++ b/Engine/Source/Runtime/Interchange/Engine/Private/Tasks/InterchangeTaskParsing.cpp
@@ -68,9 +68,9 @@ void UE::Interchange::FTaskParsing::DoTask(ENamedThreads::Type CurrentThread, co
 					return;
 				}
 
-				if (Node->GetAssetClass() != nullptr)
+				if (Node->GetObjectClass() != nullptr)
 				{
-					const UClass* RegisteredFactoryClass = InterchangeManager->GetRegisteredFactoryClass(Node->GetAssetClass());
+					const UClass* RegisteredFactoryClass = InterchangeManager->GetRegisteredFactoryClass(Node->GetObjectClass());
 					if (!RegisteredFactoryClass)
 					{
 						//nothing we can import from this element
