Make sure creating a blueprint pipeline is possible and have access to everything it need

diff --git a/Engine/Source/Runtime/Interchange/Core/Public/InterchangePipelineBase.h b/Engine/Source/Runtime/Interchange/Core/Public/InterchangePipelineBase.h
--- a/Engine/Source/Runtime/Interchange/Core/Public/InterchangePipelineBase.h
+++ b/Engine/Source/Runtime/Interchange/Core/Public/InterchangePipelineBase.h
@@ -32,7 +32,7 @@ public:
 	 * by the pipeline. Each factory node should be send to a a interchange factory to create an unreal asset.
 	 * @note - the Interchange manager is calling this function not the virtual one that is call by the default implementation.
 	 */
-	UFUNCTION(BlueprintNativeEvent, BlueprintCallable, Category = "Interchange | Translator")
+	UFUNCTION(BlueprintNativeEvent, BlueprintCallable, Category = "Interchange | Pipeline")
 	bool ScriptedExecutePreImportPipeline(UInterchangeBaseNodeContainer* BaseNodeContainer, const TArray<UInterchangeSourceData*>& SourceDatas);
 	/** The default implementation (call if the blueprint do not have any implementation) will call the virtual ExecuteImportPipeline */
 	bool ScriptedExecutePreImportPipeline_Implementation(UInterchangeBaseNodeContainer* BaseNodeContainer, const TArray<UInterchangeSourceData*>& SourceDatas)
@@ -48,7 +48,7 @@ public:
 	 * @example - PhysicsAsset need skeletal mesh render data to be build properly.
 	 * @note - the Interchange manager is calling this function not the virtual one that is call by the default implementation.
 	 */
-	UFUNCTION(BlueprintNativeEvent, BlueprintCallable, Category = "Interchange | Translator")
+	UFUNCTION(BlueprintNativeEvent, BlueprintCallable, Category = "Interchange | Pipeline")
 	bool ScriptedExecutePostImportPipeline(const UInterchangeBaseNodeContainer* BaseNodeContainer, const FString& NodeKey, UObject* CreatedAsset);
 	/** The default implementation (call if the blueprint do not have any implementation) will call the virtual ExecuteImportPipeline */
 	bool ScriptedExecutePostImportPipeline_Implementation(const UInterchangeBaseNodeContainer* BaseNodeContainer, const FString& NodeKey, UObject* CreatedAsset)
@@ -61,7 +61,7 @@ public:
 	 * Non virtual helper to allow blueprint to implement event base function let the interchange know if it can run asynchronously.
 	 * the Interchange manager is calling this function not the virtual one that is call by the default implementation.
 	 */
-	UFUNCTION(BlueprintCallable, BlueprintNativeEvent, Category = "Interchange | Translator")
+	UFUNCTION(BlueprintCallable, BlueprintNativeEvent, Category = "Interchange | Pipeline")
 	bool ScriptedExecuteExportPipeline(UInterchangeBaseNodeContainer* BaseNodeContainer);
 
 	/** The default implementation (call if the blueprint do not have any implementation) will call the virtual ExecuteExportPipeline */
@@ -74,7 +74,7 @@ public:
 	 * Non virtual helper to allow blueprint to implement event base function let the interchange know if it can run asynchronously.
 	 * the Interchange manager is calling this function not the virtual one that is call by the default implementation.
 	 */
-	UFUNCTION(BlueprintNativeEvent, BlueprintCallable, Category = "Interchange | Translator")
+	UFUNCTION(BlueprintNativeEvent, BlueprintCallable, Category = "Interchange | Pipeline")
 	bool ScriptedCanExecuteOnAnyThread(EInterchangePipelineTask PipelineTask);
 
 	/** The default implementation (call if the blueprint do not have any implementation) will call the virtual CanExecuteAsync */
diff --git a/Engine/Source/Runtime/Interchange/Core/Public/Nodes/InterchangeBaseNode.h b/Engine/Source/Runtime/Interchange/Core/Public/Nodes/InterchangeBaseNode.h
--- a/Engine/Source/Runtime/Interchange/Core/Public/Nodes/InterchangeBaseNode.h
+++ b/Engine/Source/Runtime/Interchange/Core/Public/Nodes/InterchangeBaseNode.h
@@ -83,7 +83,7 @@ namespace InterchangePrivateNodeBase
  * @param EnumType - Optional, specify it only if the AssetType member is an enum so we can type cast it in the apply function (we use uint8 to store the enum value)"
  */
 #define IMPLEMENT_NODE_ATTRIBUTE_KEY(AttributeName)																		\
-const UE::Interchange::FAttributeKey Macro_Custom##AttributeName##Key = UE::Interchange::FAttributeKey(TEXT(#AttributeName));	\
+const UE::Interchange::FAttributeKey Macro_Custom##AttributeName##Key = UE::Interchange::FAttributeKey(TEXT(#AttributeName));
 
 #if WITH_ENGINE
 #define IMPLEMENT_NODE_ATTRIBUTE_APPLY_UOBJECT(AttributeName, AttributeType, AssetType, EnumType)	\
@@ -164,6 +164,17 @@ bool FillCustom##AttributeName##FromAsset(UObject* Asset)											\
 
 #endif //#if WITH_ENGINE
 
+#define INTERCHANGE_BASE_NODE_ADD_ATTRIBUTE(ValueType)																									\
+UE::Interchange::FAttributeStorage::TAttributeHandle<ValueType> Handle = RegisterAttribute<ValueType>(UE::Interchange::FAttributeKey(NodeAttributeKey), Value);	\
+return Handle.IsValid();
+
+#define INTERCHANGE_BASE_NODE_GET_ATTRIBUTE(ValueType)																							\
+UE::Interchange::FAttributeStorage::TAttributeHandle<ValueType> Handle = GetAttributeHandle<ValueType>(UE::Interchange::FAttributeKey(NodeAttributeKey));	\
+if (Handle.IsValid())																															\
+{																																				\
+	return (Handle.Get(OutValue) == UE::Interchange::EAttributeStorageResult::Operation_Success);													\
+}																																				\
+return false;
 
 //Interchange namespace
 namespace UE
@@ -968,6 +979,88 @@ public:
 		Attributes->GetAttributeKeys(AttributeKeys);
 	}
 
+	/**
+	 * Remove any attribute from this node. Return false if we cannot remove it. If the attribute do not exist it will return true.
+	 */
+	UFUNCTION(BlueprintCallable, Category = "Interchange | Node")
+	bool RemoveAttribute(const FString& NodeAttributeKey)
+	{
+		Attributes->UnregisterAttribute(UE::Interchange::FAttributeKey(NodeAttributeKey));
+		return !HasAttribute(UE::Interchange::FAttributeKey(NodeAttributeKey));
+	}
+
+	/**
+	 * Add a boolean attribute to this node. Return false if the attribute do not exist or if we cannot add it
+	 */
+	UFUNCTION(BlueprintCallable, Category = "Interchange | Node")
+	bool AddBooleanAttribute(const FString& NodeAttributeKey, const bool& Value)
+	{
+		INTERCHANGE_BASE_NODE_ADD_ATTRIBUTE(bool);
+	}
+	
+	/**
+	 * Get a boolean attribute from this node. Return false if the attribute do not exist
+	 */
+	UFUNCTION(BlueprintCallable, Category = "Interchange | Node")
+	bool GetBooleanAttribute(const FString& NodeAttributeKey, bool& OutValue)
+	{
+		INTERCHANGE_BASE_NODE_GET_ATTRIBUTE(bool);
+	}
+
+	/**
+	 * Add a int32 attribute to this node. Return false if the attribute do not exist or if we cannot add it
+	 */
+	UFUNCTION(BlueprintCallable, Category = "Interchange | Node")
+	bool AddInt32Attribute(const FString& NodeAttributeKey, const int32& Value)
+	{
+		INTERCHANGE_BASE_NODE_ADD_ATTRIBUTE(int32);
+	}
+	
+	/**
+	 * Get a int32 attribute from this node. Return false if the attribute do not exist
+	 */
+	UFUNCTION(BlueprintCallable, Category = "Interchange | Node")
+	bool GetInt32Attribute(const FString& NodeAttributeKey, int32& OutValue)
+	{
+		INTERCHANGE_BASE_NODE_GET_ATTRIBUTE(int32);
+	}
+
+	/**
+	 * Add a float attribute to this node. Return false if the attribute do not exist or if we cannot add it
+	 */
+	UFUNCTION(BlueprintCallable, Category = "Interchange | Node")
+	bool AddFloatAttribute(const FString& NodeAttributeKey, const float& Value)
+	{
+		INTERCHANGE_BASE_NODE_ADD_ATTRIBUTE(float);
+	}
+	
+	/**
+	 * Get a float attribute from this node. Return false if the attribute do not exist
+	 */
+	UFUNCTION(BlueprintCallable, Category = "Interchange | Node")
+	bool GetFloatAttribute(const FString& NodeAttributeKey, float& OutValue)
+	{
+		INTERCHANGE_BASE_NODE_GET_ATTRIBUTE(float);
+	}
+
+	/**
+	 * Add a string attribute to this node. Return false if the attribute do not exist or if we cannot add it
+	 */
+	UFUNCTION(BlueprintCallable, Category = "Interchange | Node")
+	bool AddStringAttribute(const FString& NodeAttributeKey, const FString& Value)
+	{
+		INTERCHANGE_BASE_NODE_ADD_ATTRIBUTE(FString);
+	}
+	
+	/**
+	 * Get a string attribute from this node. Return false if the attribute do not exist
+	 */
+	UFUNCTION(BlueprintCallable, Category = "Interchange | Node")
+	bool GetStringAttribute(const FString& NodeAttributeKey, FString& OutValue)
+	{
+		INTERCHANGE_BASE_NODE_GET_ATTRIBUTE(FString);
+	}
+
 	/**
 	 * Return the unique id pass in the constructor.
 	 *
