Interchange Level import
- Added a new project setting to use Interchange for "import into level"
- It currently only spawns empty actors, in hierarchy, with the correct transforms.

diff --git a/Engine/Plugins/Enterprise/DatasmithImporter/Source/DatasmithImporter/Private/DatasmithConsumer.cpp b/Engine/Plugins/Enterprise/DatasmithImporter/Source/DatasmithImporter/Private/DatasmithConsumer.cpp
--- a/Engine/Plugins/Enterprise/DatasmithImporter/Source/DatasmithImporter/Private/DatasmithConsumer.cpp
+++ b/Engine/Plugins/Enterprise/DatasmithImporter/Source/DatasmithImporter/Private/DatasmithConsumer.cpp
@@ -1637,8 +1637,8 @@ namespace DatasmithConsumerUtils
 					TSharedRef< IDatasmithTextureElement > TextureElement = FDatasmithSceneFactory::CreateTexture( *AssetTag );
 					TextureElement->SetLabel( *Texture->GetName() );
 
-					UE::Interchange::FAssetImportResultRef& AssetImportResults = ImportContext.ImportedTextures.Add( TextureElement, MakeShared< UE::Interchange::FAssetImportResult, ESPMode::ThreadSafe >() );
-					AssetImportResults->AddImportedAsset( Texture );
+					UE::Interchange::FAssetImportResultRef& AssetImportResults = ImportContext.ImportedTextures.Add( TextureElement, MakeShared< UE::Interchange::FImportResult, ESPMode::ThreadSafe >() );
+					AssetImportResults->AddImportedObject( Texture );
 					AssetImportResults->SetDone();
 
 					ImportContext.Scene->AddTexture( TextureElement );
diff --git a/Engine/Plugins/Enterprise/DatasmithImporter/Source/DatasmithImporter/Private/DatasmithImporter.cpp b/Engine/Plugins/Enterprise/DatasmithImporter/Source/DatasmithImporter/Private/DatasmithImporter.cpp
--- a/Engine/Plugins/Enterprise/DatasmithImporter/Source/DatasmithImporter/Private/DatasmithImporter.cpp
+++ b/Engine/Plugins/Enterprise/DatasmithImporter/Source/DatasmithImporter/Private/DatasmithImporter.cpp
@@ -531,7 +531,7 @@ void FDatasmithImporter::ImportTextures( FDatasmithImportContext& ImportContext
 				if ( FutureTexture->IsValid() )
 				{
 					FutureTexture->OnDone(
-						[ TextureElement = FilteredTextureElements[TextureIndex].ToSharedRef(), &ImportContext ]( UE::Interchange::FAssetImportResult& AssetImportResults )
+						[ TextureElement = FilteredTextureElements[TextureIndex].ToSharedRef(), &ImportContext ]( UE::Interchange::FImportResult& AssetImportResults )
 						{
 							ImportMetaDataForObject( ImportContext, TextureElement, AssetImportResults.GetFirstAssetOfClass( UTexture::StaticClass() ) );
 						}
@@ -644,8 +644,8 @@ UTexture* FDatasmithImporter::ImportTexture( FDatasmithImportContext& ImportCont
 		return nullptr;
 	}
 
-	UE::Interchange::FAssetImportResultRef& AssetImportResults = ImportContext.ImportedTextures.Add( TextureElement, MakeShared< UE::Interchange::FAssetImportResult, ESPMode::ThreadSafe >() );
-	AssetImportResults->AddImportedAsset( ImportedTexture );
+	UE::Interchange::FAssetImportResultRef& AssetImportResults = ImportContext.ImportedTextures.Add( TextureElement, MakeShared< UE::Interchange::FImportResult, ESPMode::ThreadSafe >() );
+	AssetImportResults->AddImportedObject( ImportedTexture );
 	AssetImportResults->SetDone();
 
 	ImportMetaDataForObject( ImportContext, TextureElement, ImportedTexture );
diff --git a/Engine/Plugins/Enterprise/DatasmithImporter/Source/DatasmithImporter/Private/DatasmithTextureImporter.cpp b/Engine/Plugins/Enterprise/DatasmithImporter/Source/DatasmithImporter/Private/DatasmithTextureImporter.cpp
--- a/Engine/Plugins/Enterprise/DatasmithImporter/Source/DatasmithImporter/Private/DatasmithTextureImporter.cpp
+++ b/Engine/Plugins/Enterprise/DatasmithImporter/Source/DatasmithImporter/Private/DatasmithTextureImporter.cpp
@@ -397,7 +397,7 @@ UE::Interchange::FAssetImportResultRef FDatasmithTextureImporter::CreateTextureA
 
 	if ( !UInterchangeManager::GetInterchangeManager().CanTranslateSourceData( ScopedSourceData.GetSourceData() ) )
 	{
-		return MakeShared< UE::Interchange::FAssetImportResult, ESPMode::ThreadSafe >();
+		return MakeShared< UE::Interchange::FImportResult, ESPMode::ThreadSafe >();
 	}
 
 	const FString ContentPath = ImportContext.AssetsContext.TexturesImportPackage->GetPathName();
diff --git a/Engine/Plugins/Runtime/Interchange/Source/Import/InterchangeImport.Build.cs b/Engine/Plugins/Runtime/Interchange/Source/Import/InterchangeImport.Build.cs
--- a/Engine/Plugins/Runtime/Interchange/Source/Import/InterchangeImport.Build.cs
+++ b/Engine/Plugins/Runtime/Interchange/Source/Import/InterchangeImport.Build.cs
@@ -34,6 +34,16 @@ namespace UnrealBuildTool.Rules
 					"TextureUtilitiesCommon",
 				}
 			);
+
+			if (Target.bBuildEditor)
+			{
+				PrivateDependencyModuleNames.AddRange(
+					new string[]
+					{
+						"UnrealEd",
+					}
+				);
+			}
 		}
 	}
 }
diff --git a/Engine/Plugins/Runtime/Interchange/Source/Import/Private/InterchangeImportModule.cpp b/Engine/Plugins/Runtime/Interchange/Source/Import/Private/InterchangeImportModule.cpp
--- a/Engine/Plugins/Runtime/Interchange/Source/Import/Private/InterchangeImportModule.cpp
+++ b/Engine/Plugins/Runtime/Interchange/Source/Import/Private/InterchangeImportModule.cpp
@@ -13,6 +13,7 @@
 #include "Mesh/InterchangeStaticMeshFactory.h"
 #include "Misc/CoreDelegates.h"
 #include "Modules/ModuleManager.h"
+#include "Scene/InterchangeActorFactory.h"
 #include "Texture/InterchangeBMPTranslator.h"
 #include "Texture/InterchangeDDSTranslator.h"
 #include "Texture/InterchangeEXRTranslator.h"
@@ -72,6 +73,7 @@ void FInterchangeImportModule::StartupModule()
 		InterchangeManager.RegisterFactory(UInterchangeSkeletalMeshFactory::StaticClass());
 		InterchangeManager.RegisterFactory(UInterchangeStaticMeshFactory::StaticClass());
 		InterchangeManager.RegisterFactory(UInterchangePhysicsAssetFactory::StaticClass());
+		InterchangeManager.RegisterFactory(UInterchangeActorFactory::StaticClass());
 	};
 
 	if (GEngine)
diff --git a/Engine/Plugins/Runtime/Interchange/Source/Import/Private/Scene/InterchangeActorFactory.cpp b/Engine/Plugins/Runtime/Interchange/Source/Import/Private/Scene/InterchangeActorFactory.cpp
new file mode 100644
--- /dev/null
+++ b/Engine/Plugins/Runtime/Interchange/Source/Import/Private/Scene/InterchangeActorFactory.cpp
@@ -0,0 +1,159 @@
+// Copyright Epic Games, Inc. All Rights Reserved. 
+#include "Scene/InterchangeActorFactory.h"
+
+#include "InterchangeActorFactoryNode.h"
+#include "InterchangeImportCommon.h"
+#include "InterchangeImportLog.h"
+#include "InterchangeSourceData.h"
+
+#include "Components/SceneComponent.h"
+#include "Engine/Engine.h"
+#include "Engine/World.h"
+#include "GameFramework/Actor.h"
+
+#include "Nodes/InterchangeBaseNode.h"
+#include "Nodes/InterchangeBaseNodeContainer.h"
+
+#if WITH_EDITOR
+#include "Editor/EditorEngine.h"
+#endif
+
+namespace UE
+{
+	namespace Interchange
+	{
+		namespace Private
+		{
+			AActor* SpawnActor(const FActorSpawnParameters& ActorSpawnParameters, const UInterchangeActorFactoryNode* ActorFactoryNode, AActor* ParentActor)
+			{
+				ensure(ActorFactoryNode);
+
+				if (!ActorFactoryNode)
+				{
+					return nullptr;
+				}
+
+				UWorld* const World = [&ActorSpawnParameters, &ParentActor]()
+				{
+					UWorld* DefaultWorld = nullptr;
+
+					if (ParentActor)
+					{
+						DefaultWorld = ParentActor->GetWorld();
+					}
+#if WITH_EDITOR
+					if (DefaultWorld == nullptr)
+					{
+						UEditorEngine* EditorEngine = Cast<UEditorEngine>(GEngine);
+						if (GIsEditor && EditorEngine != nullptr)
+						{
+							DefaultWorld = EditorEngine->GetEditorWorldContext().World();
+						}
+					}
+#endif
+					if (DefaultWorld == nullptr && GEngine)
+					{
+						DefaultWorld = GEngine->GetWorld();
+					}
+
+					return ActorSpawnParameters.OverrideLevel ? ActorSpawnParameters.OverrideLevel->GetWorld() : DefaultWorld;
+				}();
+
+				if (!World)
+				{
+					return nullptr;
+				}
+
+				FTransform Transform = FTransform::Identity;
+				ActorFactoryNode->GetCustomGlobalTransform(Transform);
+
+				AActor* SpawnedActor = World->SpawnActor<AActor>(ActorFactoryNode->GetObjectClass(), Transform, ActorSpawnParameters);
+
+				if (SpawnedActor)
+				{
+#if WITH_EDITOR
+					SpawnedActor->SetActorLabel(ActorSpawnParameters.Name.ToString());
+#endif
+					if (!SpawnedActor->GetRootComponent())
+					{
+						USceneComponent* RootComponent = NewObject<USceneComponent>(SpawnedActor, USceneComponent::GetDefaultSceneRootVariableName(), RF_Transactional);
+						RootComponent->Mobility = EComponentMobility::Static; //todo: this should be an attribute on a factory node
+#if WITH_EDITORONLY_DATA
+						RootComponent->bVisualizeComponent = true;
+#endif
+						RootComponent->SetWorldTransform(Transform);
+
+						SpawnedActor->SetRootComponent(RootComponent);
+						SpawnedActor->AddInstanceComponent(RootComponent);
+
+						RootComponent->RegisterComponent();
+					}
+				}
+
+				if (ParentActor)
+				{
+					SpawnedActor->AttachToActor(ParentActor, FAttachmentTransformRules::KeepWorldTransform);
+				}
+
+				return SpawnedActor;
+			}
+
+			void RecursivelySpawnActors(TMap<FString, UObject*>& OutSpawnedActors,
+				const UInterchangeBaseNodeContainer* NodeContainer, const UInterchangeActorFactoryNode* ActorNode, const FActorSpawnParameters& SpawnParameters,
+				AActor* ParentActor, const bool bSpawnChildren)
+			{
+				if (!ActorNode || !NodeContainer)
+				{
+					return;
+				}
+
+				const FString ActorNodeUniqueID = ActorNode->GetUniqueID();
+
+				AActor* SpawnedActor = SpawnActor(SpawnParameters, Cast<UInterchangeActorFactoryNode>(ActorNode), ParentActor);
+				OutSpawnedActors.Add(ActorNodeUniqueID, SpawnedActor);
+
+				if (bSpawnChildren)
+				{
+					for (int32 ChildIndex = 0; ChildIndex < NodeContainer->GetNodeChildrenCount(ActorNodeUniqueID); ++ChildIndex)
+					{
+						if (const UInterchangeBaseNode* ChildNode = NodeContainer->GetNodeChildren(ActorNodeUniqueID, ChildIndex))
+						{
+							if (const UInterchangeActorFactoryNode* ChildActorFactoryNode = Cast<const UInterchangeActorFactoryNode>(ChildNode))
+							{
+								FActorSpawnParameters ChildActorSpawnParameters;
+								ChildActorSpawnParameters.Name = FName(ChildActorFactoryNode->GetDisplayLabel());
+								ChildActorSpawnParameters.NameMode = FActorSpawnParameters::ESpawnActorNameMode::Requested;
+								ChildActorSpawnParameters.OverrideLevel = SpawnParameters.OverrideLevel;
+								ChildActorSpawnParameters.SpawnCollisionHandlingOverride = ESpawnActorCollisionHandlingMethod::AlwaysSpawn;
+
+								RecursivelySpawnActors(OutSpawnedActors, NodeContainer, ChildActorFactoryNode, ChildActorSpawnParameters, SpawnedActor, bSpawnChildren);
+							}
+						}
+					}
+				}
+			};
+		}
+	}
+}
+
+UClass* UInterchangeActorFactory::GetFactoryClass() const
+{
+	return AActor::StaticClass();
+}
+
+TMap<FString, UObject*> UInterchangeActorFactory::CreateSceneObjects(const UInterchangeFactoryBase::FCreateSceneObjectsParams& CreateSceneObjectsParams)
+{
+	FActorSpawnParameters ActorSpawnParameters;
+	ActorSpawnParameters.Name = FName(*CreateSceneObjectsParams.ObjectName);
+	ActorSpawnParameters.NameMode = FActorSpawnParameters::ESpawnActorNameMode::Requested;
+	ActorSpawnParameters.OverrideLevel = CreateSceneObjectsParams.Level;
+	ActorSpawnParameters.SpawnCollisionHandlingOverride = ESpawnActorCollisionHandlingMethod::AlwaysSpawn;
+
+	TMap<FString, UObject*> SpawnedActors;
+	constexpr AActor* ParentActor = nullptr;
+
+	UE::Interchange::Private::RecursivelySpawnActors(SpawnedActors, CreateSceneObjectsParams.NodeContainer, Cast<const UInterchangeActorFactoryNode>(CreateSceneObjectsParams.ObjectNode),
+		ActorSpawnParameters, ParentActor, CreateSceneObjectsParams.bCreateSceneObjectsForChildren);
+
+	return SpawnedActors;
+}
\ No newline at end of file
diff --git a/Engine/Plugins/Runtime/Interchange/Source/Import/Public/Scene/InterchangeActorFactory.h b/Engine/Plugins/Runtime/Interchange/Source/Import/Public/Scene/InterchangeActorFactory.h
new file mode 100644
--- /dev/null
+++ b/Engine/Plugins/Runtime/Interchange/Source/Import/Public/Scene/InterchangeActorFactory.h
@@ -0,0 +1,34 @@
+// Copyright Epic Games, Inc. All Rights Reserved.
+
+#pragma once
+
+#include "CoreMinimal.h"
+#include "InterchangeFactoryBase.h"
+#include "UObject/Object.h"
+#include "UObject/ObjectMacros.h"
+
+#include "InterchangeActorFactory.generated.h"
+
+UCLASS(BlueprintType)
+class INTERCHANGEIMPORT_API UInterchangeActorFactory : public UInterchangeFactoryBase
+{
+	GENERATED_BODY()
+public:
+
+	//////////////////////////////////////////////////////////////////////////
+	// Interchange factory base interface begin
+
+	virtual UClass* GetFactoryClass() const override;
+
+	virtual TMap<FString, UObject*> CreateSceneObjects(const UInterchangeFactoryBase::FCreateSceneObjectsParams& CreateSceneObjectsParams) override;
+
+	virtual bool CanExecuteOnAnyThread() const
+	{
+		return false;
+	}
+
+	// Interchange factory base interface end
+	//////////////////////////////////////////////////////////////////////////
+};
+
+
diff --git a/Engine/Plugins/Runtime/Interchange/Source/Nodes/Private/InterchangeActorFactoryNode.cpp b/Engine/Plugins/Runtime/Interchange/Source/Nodes/Private/InterchangeActorFactoryNode.cpp
new file mode 100644
--- /dev/null
+++ b/Engine/Plugins/Runtime/Interchange/Source/Nodes/Private/InterchangeActorFactoryNode.cpp
@@ -0,0 +1,16 @@
+// Copyright Epic Games, Inc. All Rights Reserved.
+
+#include "InterchangeActorFactoryNode.h"
+
+#if WITH_ENGINE
+	#include "GameFramework/Actor.h"
+#endif
+
+UClass* UInterchangeActorFactoryNode::GetObjectClass() const
+{
+#if WITH_ENGINE
+	return AActor::StaticClass();
+#else
+	return nullptr;
+#endif
+}
\ No newline at end of file
diff --git a/Engine/Plugins/Runtime/Interchange/Source/Nodes/Public/InterchangeActorFactoryNode.h b/Engine/Plugins/Runtime/Interchange/Source/Nodes/Public/InterchangeActorFactoryNode.h
new file mode 100644
--- /dev/null
+++ b/Engine/Plugins/Runtime/Interchange/Source/Nodes/Public/InterchangeActorFactoryNode.h
@@ -0,0 +1,31 @@
+// Copyright Epic Games, Inc. All Rights Reserved.
+
+#pragma once
+
+#include "Nodes/InterchangeBaseNode.h"
+
+#include "InterchangeActorFactoryNode.generated.h"
+
+UCLASS(BlueprintType)
+class INTERCHANGENODES_API UInterchangeActorFactoryNode : public UInterchangeBaseNode
+{
+	GENERATED_BODY()
+
+public:
+	virtual UClass* GetObjectClass() const override;
+
+	UFUNCTION(BlueprintCallable, Category = "Interchange | Node | ActorFactory")
+	bool GetCustomGlobalTransform(FTransform& AttributeValue) const
+	{
+		IMPLEMENT_NODE_ATTRIBUTE_GETTER(GlobalTransform, FTransform);
+	}
+
+	UFUNCTION(BlueprintCallable, Category = "Interchange | Node | ActorFactory")
+	bool SetCustomGlobalTransform(const FTransform& AttributeValue)
+	{
+		IMPLEMENT_NODE_ATTRIBUTE_SETTER_NODELEGATE(GlobalTransform, FTransform);
+	}
+
+private:
+	const UE::Interchange::FAttributeKey Macro_CustomGlobalTransformKey = UE::Interchange::FAttributeKey(TEXT("GlobalTransform"));
+};
\ No newline at end of file
diff --git a/Engine/Plugins/Runtime/Interchange/Source/Pipelines/Private/InterchangeGenericScenesPipeline.cpp b/Engine/Plugins/Runtime/Interchange/Source/Pipelines/Private/InterchangeGenericScenesPipeline.cpp
new file mode 100644
--- /dev/null
+++ b/Engine/Plugins/Runtime/Interchange/Source/Pipelines/Private/InterchangeGenericScenesPipeline.cpp
@@ -0,0 +1,72 @@
+// Copyright Epic Games, Inc. All Rights Reserved. 
+
+#include "InterchangeGenericScenesPipeline.h"
+
+#include "InterchangeActorFactoryNode.h"
+#include "InterchangePipelineLog.h"
+#include "InterchangeSceneNode.h"
+
+bool UInterchangeGenericLevelPipeline::ExecutePreImportPipeline(UInterchangeBaseNodeContainer* InBaseNodeContainer, const TArray<UInterchangeSourceData*>& InSourceDatas)
+{
+	if (!InBaseNodeContainer)
+	{
+		UE_LOG(LogInterchangePipeline, Warning, TEXT("UInterchangeGenericAssetsPipeline: Cannot execute pre-import pipeline because InBaseNodeContrainer is null"));
+		return false;
+	}
+
+	TArray<UInterchangeSceneNode*> SceneNodes;
+
+	//Find all translated node we need for this pipeline
+	InBaseNodeContainer->IterateNodes([&SceneNodes](const FString& NodeUid, UInterchangeBaseNode* Node)
+	{
+		switch(Node->GetnodeContainerType())
+		{
+		case EInterchangeNodeContainerType::NodeContainerType_TranslatedScene:
+		{
+			if (UInterchangeSceneNode* SceneNode = Cast<UInterchangeSceneNode>(Node))
+			{
+				SceneNodes.Add(SceneNode);
+			}
+		}
+		break;
+		}
+	});
+
+	for (const UInterchangeSceneNode* SceneNode : SceneNodes)
+	{
+		CreateActorFactoryNode(SceneNode, InBaseNodeContainer);
+	}
+
+	return true;
+}
+
+void UInterchangeGenericLevelPipeline::CreateActorFactoryNode(const UInterchangeSceneNode* SceneNode, UInterchangeBaseNodeContainer* FactoryNodeContainer)
+{
+	if (!SceneNode)
+	{
+		return;
+	}
+
+	UInterchangeActorFactoryNode* ActorFactoryNode = NewObject<UInterchangeActorFactoryNode>(FactoryNodeContainer, NAME_None);
+
+	if (!ensure(ActorFactoryNode))
+	{
+		return;
+	}
+
+	ActorFactoryNode->InitializeNode(TEXT("Factory_") + SceneNode->GetUniqueID(), SceneNode->GetDisplayLabel(), EInterchangeNodeContainerType::NodeContainerType_FactoryData);
+
+	if (!SceneNode->GetParentUid().IsEmpty())
+	{
+		ActorFactoryNode->SetParentUid(TEXT("Factory_") + SceneNode->GetParentUid());
+	}
+
+	FTransform GlobalTransform;
+	if (SceneNode->GetCustomGlobalTransform(GlobalTransform))
+	{
+		ActorFactoryNode->SetCustomGlobalTransform(GlobalTransform);
+	}
+	
+	FactoryNodeContainer->AddNode(ActorFactoryNode);
+}
+
diff --git a/Engine/Plugins/Runtime/Interchange/Source/Pipelines/Public/InterchangeGenericScenesPipeline.h b/Engine/Plugins/Runtime/Interchange/Source/Pipelines/Public/InterchangeGenericScenesPipeline.h
new file mode 100644
--- /dev/null
+++ b/Engine/Plugins/Runtime/Interchange/Source/Pipelines/Public/InterchangeGenericScenesPipeline.h
@@ -0,0 +1,38 @@
+// Copyright Epic Games, Inc. All Rights Reserved.
+
+#pragma once
+
+#include "CoreMinimal.h"
+#include "InterchangePipelineBase.h"
+#include "InterchangeSourceData.h"
+#include "Nodes/InterchangeBaseNodeContainer.h"
+#include "UObject/Object.h"
+#include "UObject/ObjectMacros.h"
+
+#include "InterchangeGenericScenesPipeline.generated.h"
+
+class UInterchangeActorFactoryNode;
+class UInterchangeSceneNode;
+
+
+UCLASS(BlueprintType)
+class INTERCHANGEPIPELINES_API UInterchangeGenericLevelPipeline : public UInterchangePipelineBase
+{
+	GENERATED_BODY()
+
+public:
+
+protected:
+
+	virtual bool ExecutePreImportPipeline(UInterchangeBaseNodeContainer* InBaseNodeContainer, const TArray<UInterchangeSourceData*>& InSourceDatas) override;
+
+	virtual bool CanExecuteOnAnyThread(EInterchangePipelineTask PipelineTask) override
+	{
+		return true;
+	}
+
+	void CreateActorFactoryNode(const UInterchangeSceneNode* SceneNode, UInterchangeBaseNodeContainer* FactoryNodeContainer);
+	
+};
+
+
diff --git a/Engine/Source/Developer/AssetTools/Private/AssetTools.cpp b/Engine/Source/Developer/AssetTools/Private/AssetTools.cpp
--- a/Engine/Source/Developer/AssetTools/Private/AssetTools.cpp
+++ b/Engine/Source/Developer/AssetTools/Private/AssetTools.cpp
@@ -148,6 +148,7 @@
 #include "Misc/ConfigCacheIni.h"
 #include "Misc/BlacklistNames.h"
 #include "InterchangeManager.h"
+#include "InterchangeProjectSettings.h"
 #include "Engine/World.h"
 #include "Engine/Level.h"
 
@@ -1757,13 +1758,12 @@ TArray<UObject*> UAssetToolsImpl::ImportAssetsInternal(const TArray<FString>& Fi
 	UInterchangeManager& InterchangeManager = UInterchangeManager::GetInterchangeManager();
 #if WITH_EDITOR
 	bUseInterchangeFramework = GetDefault<UEditorExperimentalSettings>()->bEnableInterchangeFramework;
-	//Avoid using interchange for scene factory, we did not do the scene import yet
-	//TODO do the scene import and remove this 
+
 	if (bUseInterchangeFramework && Params.SpecifiedFactory)
 	{
 		if (Params.SpecifiedFactory->GetClass()->IsChildOf(USceneImportFactory::StaticClass()))
 		{
-			bUseInterchangeFramework = false;
+			bUseInterchangeFramework = GetDefault<UInterchangeProjectSettings>()->bUseInterchangeWhenImportingIntoLevel;
 		}
 	}
 #endif
@@ -1985,19 +1985,39 @@ TArray<UObject*> UAssetToolsImpl::ImportAssetsInternal(const TArray<FString>& Fi
 				ImportAssetParameters.bIsAutomated = bAutomatedImport;
 				ImportAssetParameters.ReimportAsset = nullptr;
 
-				UE::Interchange::FAssetImportResultRef InterchangeResult = (InterchangeManager.ImportAssetAsync(DestinationPath, ScopedSourceData.GetSourceData(), ImportAssetParameters));
-				InterchangeResult->OnDone(
+				TFunction<void(UE::Interchange::FImportResult&)> AppendImportResult =
 					// Note: ImportStatus captured by value so that the lambda keeps the shared ptr alive
-					[ImportStatus](UE::Interchange::FAssetImportResult& Result)
+					[ImportStatus](UE::Interchange::FImportResult& Result)
 					{
 						ImportStatus->InterchangeResultsContainer->Append(Result.GetResults());
+					};
+
+				TFunction<void(UE::Interchange::FImportResult&)> AppendAndBroadcastImportResultIfNeeded =
+					// Note: ImportStatus captured by value so that the lambda keeps the shared ptr alive
+					[ImportStatus, AppendImportResult](UE::Interchange::FImportResult& Result)
+					{
+						AppendImportResult(Result);
 						if (--ImportStatus->ImportCount == 0)
 						{
 							UInterchangeManager& InterchangeManager = UInterchangeManager::GetInterchangeManager();
 							InterchangeManager.OnBatchImportComplete.Broadcast(ImportStatus->InterchangeResultsContainer);
 						}
-					});
+					};
+
+				if (Params.SpecifiedFactory && Params.SpecifiedFactory->GetClass()->IsChildOf(USceneImportFactory::StaticClass()))
+				{
+					TPair<UE::Interchange::FAssetImportResultRef, UE::Interchange::FSceneImportResultRef> InterchangeResults =
+						InterchangeManager.ImportSceneAsync(DestinationPath, ScopedSourceData.GetSourceData(), ImportAssetParameters);
 
+					InterchangeResults.Key->OnDone(AppendImportResult);;
+					InterchangeResults.Value->OnDone(AppendAndBroadcastImportResultIfNeeded);
+				}
+				else
+				{
+					UE::Interchange::FAssetImportResultRef InterchangeResult = (InterchangeManager.ImportAssetAsync(DestinationPath, ScopedSourceData.GetSourceData(), ImportAssetParameters));
+					InterchangeResult->OnDone(AppendAndBroadcastImportResultIfNeeded);
+				}
+				
 				InterchangeManager.OnAssetPostImport.Remove(PostImportHandle);
 				//Import done, iterate the next file and destination
 				
diff --git a/Engine/Source/Runtime/Interchange/Core/Public/InterchangeFactoryBase.h b/Engine/Source/Runtime/Interchange/Core/Public/InterchangeFactoryBase.h
--- a/Engine/Source/Runtime/Interchange/Core/Public/InterchangeFactoryBase.h
+++ b/Engine/Source/Runtime/Interchange/Core/Public/InterchangeFactoryBase.h
@@ -10,6 +10,7 @@
 
 #include "InterchangeFactoryBase.generated.h"
 
+class ULevel;
 class UInterchangeBaseNode;
 class UInterchangeBaseNodeContainer;
 class UInterchangePipelineBase;
@@ -100,6 +101,39 @@ public:
 		return nullptr;
 	}
 
+	/**
+	 * Parameters to pass to SpawnActor function
+	 */
+	struct FCreateSceneObjectsParams
+	{
+		/** The level in which to create the scene objects */
+		ULevel* Level = nullptr;
+
+		/** The name we want to give to the actor that we will create */
+		FString ObjectName;
+
+		/** The base node that describe how to create the asset */
+		UInterchangeBaseNode* ObjectNode = nullptr;
+
+		/** The node container associated with the current source index */
+		const UInterchangeBaseNodeContainer* NodeContainer = nullptr;
+
+		/** Whether to create the scene objects for the child nodes or not */
+		bool bCreateSceneObjectsForChildren = false;
+	};
+
+	/**
+	 * Creates the scene object from a Scene Node data.
+	 * If FCreateSceneObjectsParams::bCreateSceneObjectsForChildren is true, will also create the scene objects for our children.
+	 *
+	 * @param Arguments - The structure containing all necessary arguments, see the structure definition for the documentation.
+	 * @return The node uids and the scene objects that we created from them.
+	 */
+	virtual TMap<FString, UObject*> CreateSceneObjects(const FCreateSceneObjectsParams& Arguments)
+	{
+		return {};
+	}
+
 	/** Return true if the factory can create the asset asynchronously on any thread, false if it need to be on the main thread */
 	virtual bool CanExecuteOnAnyThread() const
 	{
diff --git a/Engine/Source/Runtime/Interchange/Engine/Private/InterchangeManager.cpp b/Engine/Source/Runtime/Interchange/Engine/Private/InterchangeManager.cpp
--- a/Engine/Source/Runtime/Interchange/Engine/Private/InterchangeManager.cpp
+++ b/Engine/Source/Runtime/Interchange/Engine/Private/InterchangeManager.cpp
@@ -71,7 +71,8 @@ UInterchangeTranslatorBase* UE::Interchange::FScopedTranslator::GetTranslator()
 }
 
 UE::Interchange::FImportAsyncHelper::FImportAsyncHelper()
-	: AssetImportResult(new FAssetImportResult())
+	: AssetImportResult(MakeShared<FImportResult>())
+	, SceneImportResult(MakeShared<FImportResult>())
 {
 	bCancel = false;
 }
@@ -107,30 +108,19 @@ void UE::Interchange::FImportAsyncHelper::CancelAndWaitUntilDoneSynchronously()
 
 	FGraphEventArray TasksToComplete;
 
-	if (TranslatorTasks.Num())
-	{
-		TasksToComplete.Append(TranslatorTasks);
-	}
-	if (PipelinePreImportTasks.Num())
-	{
-		TasksToComplete.Append(PipelinePreImportTasks);
-	}
+	TasksToComplete.Append(TranslatorTasks);
+	TasksToComplete.Append(PipelinePreImportTasks);
+	
 	if (ParsingTask.GetReference())
 	{
 		TasksToComplete.Add(ParsingTask);
 	}
-	if (CreatePackageTasks.Num())
-	{
-		TasksToComplete.Append(CreatePackageTasks);
-	}
-	if (CreateAssetTasks.Num())
-	{
-		TasksToComplete.Append(CreateAssetTasks);
-	}
-	if (PipelinePostImportTasks.Num())
-	{
-		TasksToComplete.Append(PipelinePostImportTasks);
-	}
+
+	TasksToComplete.Append(CreatePackageTasks);
+	TasksToComplete.Append(CreateAssetTasks);
+	TasksToComplete.Append(SceneTasks);
+	TasksToComplete.Append(PipelinePostImportTasks);
+
 	if (PreAsyncCompletionTask.GetReference())
 	{
 		TasksToComplete.Add(PreAsyncCompletionTask);
@@ -151,11 +141,8 @@ void UE::Interchange::FImportAsyncHelper::CancelAndWaitUntilDoneSynchronously()
 		FTaskGraphInterface::Get().WaitUntilTasksComplete(TasksToComplete, ENamedThreads::GameThread);
 	}
 
-	//Async import result in a null object when we cancel
-	if (AssetImportResult->GetStatus() != FAssetImportResult::EStatus::Done)
-	{
-		AssetImportResult->SetDone();
-	}
+	AssetImportResult->SetDone();
+	SceneImportResult->SetDone();
 }
 
 void UE::Interchange::FImportAsyncHelper::CleanUp()
@@ -205,23 +192,23 @@ void UE::Interchange::FImportAsyncHelper::CleanUp()
 	CreatedFactories.Empty();
 }
 
-UE::Interchange::FAssetImportResult::FAssetImportResult()
+UE::Interchange::FImportResult::FImportResult()
 	: ImportStatus(EStatus::Invalid)
 {
 	Results = NewObject<UInterchangeResultsContainer>(GetTransientPackage());
 }
 
-UE::Interchange::FAssetImportResult::EStatus UE::Interchange::FAssetImportResult::GetStatus() const
+UE::Interchange::FImportResult::EStatus UE::Interchange::FImportResult::GetStatus() const
 {
 	return ImportStatus;
 }
 
-bool UE::Interchange::FAssetImportResult::IsValid() const
+bool UE::Interchange::FImportResult::IsValid() const
 {
 	return GetStatus() != EStatus::Invalid;
 }
 
-void UE::Interchange::FAssetImportResult::SetInProgress()
+void UE::Interchange::FImportResult::SetInProgress()
 {
 	EStatus ExpectedStatus = EStatus::Invalid;
 	if (ImportStatus.compare_exchange_strong(ExpectedStatus, EStatus::InProgress))
@@ -230,7 +217,7 @@ void UE::Interchange::FAssetImportResult::SetInProgress()
 	}
 }
 
-void UE::Interchange::FAssetImportResult::SetDone()
+void UE::Interchange::FImportResult::SetDone()
 {
 	SetInProgress(); // Make sure we always pass through the InProgress state
 
@@ -246,7 +233,7 @@ void UE::Interchange::FAssetImportResult::SetDone()
 	}
 }
 
-void UE::Interchange::FAssetImportResult::WaitUntilDone()
+void UE::Interchange::FImportResult::WaitUntilDone()
 {
 	if (ImportStatus == EStatus::InProgress)
 	{
@@ -254,18 +241,18 @@ void UE::Interchange::FAssetImportResult::WaitUntilDone()
 	}
 }
 
-const TArray< UObject* >& UE::Interchange::FAssetImportResult::GetImportedAssets() const
+const TArray< UObject* >& UE::Interchange::FImportResult::GetImportedObjects() const
 {
-	FReadScopeLock ReadScopeLock(ImportedAssetsRWLock);
-	return ImportedAssets;
+	FReadScopeLock ReadScopeLock(ImportedObjectsRWLock);
+	return ImportedObjects;
 }
 
-UObject* UE::Interchange::FAssetImportResult::GetFirstAssetOfClass(UClass* InClass) const
+UObject* UE::Interchange::FImportResult::GetFirstAssetOfClass(UClass* InClass) const
 {
 	UObject* Asset = nullptr;
 
-	FReadScopeLock ReadScopeLock(ImportedAssetsRWLock);
-	for (UObject* ImportedAsset : ImportedAssets)
+	FReadScopeLock ReadScopeLock(ImportedObjectsRWLock);
+	for (UObject* ImportedAsset : ImportedObjects)
 	{
 		if (ImportedAsset->IsA(InClass))
 		{
@@ -277,21 +264,21 @@ UObject* UE::Interchange::FAssetImportResult::GetFirstAssetOfClass(UClass* InCla
 	return Asset;
 }
 
-void UE::Interchange::FAssetImportResult::AddImportedAsset(UObject* ImportedAsset)
+void UE::Interchange::FImportResult::AddImportedObject(UObject* ImportedObject)
 {
-	FWriteScopeLock WriteScopeLock(ImportedAssetsRWLock);
-	ImportedAssets.Add(ImportedAsset);
+	FWriteScopeLock WriteScopeLock(ImportedObjectsRWLock);
+	ImportedObjects.Add(ImportedObject);
 }
 
-void UE::Interchange::FAssetImportResult::OnDone(TFunction< void(FAssetImportResult&) > Callback)
+void UE::Interchange::FImportResult::OnDone(TFunction< void(FImportResult&) > Callback)
 {
 	DoneCallback = Callback;
 }
 
-void UE::Interchange::FAssetImportResult::AddReferencedObjects(FReferenceCollector& Collector)
+void UE::Interchange::FImportResult::AddReferencedObjects(FReferenceCollector& Collector)
 {
-	FReadScopeLock ReadScopeLock(ImportedAssetsRWLock);
-	Collector.AddReferencedObjects(ImportedAssets);
+	FReadScopeLock ReadScopeLock(ImportedObjectsRWLock);
+	Collector.AddReferencedObjects(ImportedObjects);
 	Collector.AddReferencedObject(Results);
 }
 
@@ -526,11 +513,31 @@ bool UInterchangeManager::ImportAsset(const FString& ContentPath, const UInterch
 }
 
 UE::Interchange::FAssetImportResultRef UInterchangeManager::ImportAssetAsync(const FString& ContentPath, const UInterchangeSourceData* SourceData, const FImportAssetParameters& ImportAssetParameters)
+{
+	return ImportInternal(ContentPath, SourceData, ImportAssetParameters, UE::Interchange::EImportType::ImportType_Asset).Get<0>();
+}
+
+bool UInterchangeManager::ImportScene(const FString& ContentPath, const UInterchangeSourceData* SourceData, const FImportAssetParameters& ImportAssetParameters)
+{
+	using namespace UE::Interchange;
+	TTuple<FAssetImportResultRef, FSceneImportResultRef> ImportResults = ImportInternal(ContentPath, SourceData, ImportAssetParameters, UE::Interchange::EImportType::ImportType_Scene);
+
+	return ImportResults.Get<0>()->IsValid() && ImportResults.Get<1>()->IsValid();
+}
+
+TTuple<UE::Interchange::FAssetImportResultRef, UE::Interchange::FSceneImportResultRef>
+UInterchangeManager::ImportSceneAsync(const FString& ContentPath, const UInterchangeSourceData* SourceData, const FImportAssetParameters& ImportAssetParameters)
+{
+	return ImportInternal(ContentPath, SourceData, ImportAssetParameters, UE::Interchange::EImportType::ImportType_Scene);
+}
+
+TTuple<UE::Interchange::FAssetImportResultRef, UE::Interchange::FSceneImportResultRef>
+UInterchangeManager::ImportInternal(const FString& ContentPath, const UInterchangeSourceData* SourceData, const FImportAssetParameters& ImportAssetParameters, const UE::Interchange::EImportType ImportType)
 {
 	if (!ensure(IsInGameThread()))
 	{
 		//Import process can be started only in the game thread
-		return MakeShared< UE::Interchange::FAssetImportResult, ESPMode::ThreadSafe >();
+		return TTuple<UE::Interchange::FAssetImportResultRef, UE::Interchange::FSceneImportResultRef>{ MakeShared< UE::Interchange::FImportResult, ESPMode::ThreadSafe >(), MakeShared< UE::Interchange::FImportResult, ESPMode::ThreadSafe >() };
 	}
 	UInterchangeAssetImportData* OriginalAssetImportData = nullptr;
 	FString PackageBasePath = ContentPath;
@@ -553,12 +560,11 @@ UE::Interchange::FAssetImportResultRef UInterchangeManager::ImportAssetAsync(con
 		}
 	}
 	bool bImportCancel = false;
-	bool bCanShowDialog = !ImportAssetParameters.bIsAutomated && !IsAttended();
 
 	//Create a task for every source data
 	UE::Interchange::FImportAsyncHelperData TaskData;
 	TaskData.bIsAutomated = ImportAssetParameters.bIsAutomated;
-	TaskData.ImportType = UE::Interchange::EImportType::ImportType_Asset;
+	TaskData.ImportType = ImportType;
 	TaskData.ReimportObject = ImportAssetParameters.ReimportAsset;
 	TSharedPtr<UE::Interchange::FImportAsyncHelper, ESPMode::ThreadSafe> AsyncHelper = CreateAsyncHelper(TaskData);
 	check(AsyncHelper.IsValid());
@@ -602,9 +608,20 @@ UE::Interchange::FAssetImportResultRef UInterchangeManager::ImportAssetAsync(con
 															&& InterchangeProjectSettings->bShowPipelineStacksConfigurationDialog
 															&& !bImportAllWithDefault;
 
-		const FName DefaultPipelineStackName = InterchangeProjectSettings->DefaultPipelineStack;
+		const FName DefaultPipelineStackName = [ImportType, &InterchangeProjectSettings]()
+		{
+			if (ImportType == UE::Interchange::EImportType::ImportType_Scene)
+			{
+				return InterchangeProjectSettings->DefaultScenePipelineStack;
+			}
+			else
+			{
+				return InterchangeProjectSettings->DefaultPipelineStack;
+			}
+		}();
+			
 		const TMap<FName, FInterchangePipelineStack>& DefaultPipelineStacks = InterchangeProjectSettings->PipelineStacks;
-		
+
 		//If we reimport we want to load the original pipeline and the original pipeline settings
 		if (OriginalAssetImportData && OriginalAssetImportData->Pipelines.Num() > 0)
 		{
@@ -649,7 +666,6 @@ UE::Interchange::FAssetImportResultRef UInterchangeManager::ImportAssetAsync(con
 				{
 					bImportAllWithDefault = true;
 				}
-				PipelineStackName = InterchangeProjectSettings->DefaultPipelineStack;
 			}
 			if (!bImportCancel)
 			{
@@ -724,12 +740,7 @@ UE::Interchange::FAssetImportResultRef UInterchangeManager::ImportAssetAsync(con
 
 	StartQueuedTasks();
 
-	return AsyncHelper->AssetImportResult;
-}
-
-bool UInterchangeManager::ImportScene(const FString& ImportContext, const UInterchangeSourceData* SourceData, bool bIsReimport, bool bIsAutomated)
-{
-	return false;
+	return TTuple<UE::Interchange::FAssetImportResultRef, UE::Interchange::FSceneImportResultRef>{ AsyncHelper->AssetImportResult, AsyncHelper->SceneImportResult };
 }
 
 bool UInterchangeManager::ExportAsset(const UObject* Asset, bool bIsAutomated)
diff --git a/Engine/Source/Runtime/Interchange/Engine/Private/Tasks/InterchangeTaskCompletion.cpp b/Engine/Source/Runtime/Interchange/Engine/Private/Tasks/InterchangeTaskCompletion.cpp
--- a/Engine/Source/Runtime/Interchange/Engine/Private/Tasks/InterchangeTaskCompletion.cpp
+++ b/Engine/Source/Runtime/Interchange/Engine/Private/Tasks/InterchangeTaskCompletion.cpp
@@ -4,6 +4,8 @@
 #include "AssetRegistryModule.h"
 #include "Async/TaskGraphInterfaces.h"
 #include "CoreMinimal.h"
+#include "Engine/World.h"
+#include "GameFramework/Actor.h"
 #include "InterchangeEngineLogPrivate.h"
 #include "InterchangeFactoryBase.h"
 #include "InterchangeManager.h"
@@ -36,77 +38,93 @@ void UE::Interchange::FTaskPreCompletion::DoTask(ENamedThreads::Type CurrentThre
 
 	UInterchangeResultsContainer* Results = AsyncHelper->AssetImportResult->GetResults();
 
-	for (TPair<int32, TArray<FImportAsyncHelper::FImportedAssetInfo>>& AssetInfosPerSourceIndexPair : AsyncHelper->ImportedAssetsPerSourceIndex)
+	auto ForEachImportedObjectForEachSource = [&AsyncHelper, &Results](TMap<int32, TArray<FImportAsyncHelper::FImportedObjectInfo>>& ImportedObjectsPerSourceIndex, bool bIsAsset)
 	{
-		//Verify if the task was cancel
-		if (AsyncHelper->bCancel)
+		for (TPair<int32, TArray<FImportAsyncHelper::FImportedObjectInfo>>& ObjectInfosPerSourceIndexPair : ImportedObjectsPerSourceIndex)
 		{
-			break;
-		}
-
-		const int32 SourceIndex = AssetInfosPerSourceIndexPair.Key;
-		const bool bCallPostImportGameThreadCallback = ensure(AsyncHelper->SourceDatas.IsValidIndex(SourceIndex));
-
-		for (const FImportAsyncHelper::FImportedAssetInfo& AssetInfo : AssetInfosPerSourceIndexPair.Value)
-		{
-			UObject* Asset = AssetInfo.ImportAsset;
-			//In case Some factory code cannot run outside of the main thread we offer this callback to finish the work before calling post edit change (building the asset)
-			if (bCallPostImportGameThreadCallback && AssetInfo.Factory)
+			//Verify if the task was cancel
+			if (AsyncHelper->bCancel)
 			{
-				UInterchangeFactoryBase::FImportPreCompletedCallbackParams Arguments;
-				Arguments.ImportedObject = Asset;
-				Arguments.SourceData = AsyncHelper->SourceDatas[SourceIndex];
-				Arguments.FactoryNode = AssetInfo.FactoryNode;
-				// Should we assert if there is no factory node?
-				Arguments.NodeUniqueID = AssetInfo.FactoryNode ? AssetInfo.FactoryNode->GetUniqueID() : FString();
-				Arguments.NodeContainer = AsyncHelper->BaseNodeContainers[SourceIndex].Get();
-				Arguments.Pipelines = AsyncHelper->Pipelines;
-				Arguments.bIsReimport = AssetInfo.bIsReimport;
-				AssetInfo.Factory->PreImportPreCompletedCallback(Arguments);
+				break;
 			}
 
-			if (Asset == nullptr)
+			const int32 SourceIndex = ObjectInfosPerSourceIndexPair.Key;
+			const bool bCallPostImportGameThreadCallback = ensure(AsyncHelper->SourceDatas.IsValidIndex(SourceIndex));
+
+			for (const FImportAsyncHelper::FImportedObjectInfo& ObjectInfo : ObjectInfosPerSourceIndexPair.Value)
 			{
-				continue;
-			}
+				UObject* ImportedObject = ObjectInfo.ImportedObject;
+				//In case Some factory code cannot run outside of the main thread we offer this callback to finish the work before calling post edit change (building the asset)
+				if (bCallPostImportGameThreadCallback && ObjectInfo.Factory)
+				{
+					UInterchangeFactoryBase::FImportPreCompletedCallbackParams Arguments;
+					Arguments.ImportedObject = ImportedObject;
+					Arguments.SourceData = AsyncHelper->SourceDatas[SourceIndex];
+					Arguments.FactoryNode = ObjectInfo.FactoryNode;
+					// Should we assert if there is no factory node?
+					Arguments.NodeUniqueID = ObjectInfo.FactoryNode ? ObjectInfo.FactoryNode->GetUniqueID() : FString();
+					Arguments.NodeContainer = AsyncHelper->BaseNodeContainers[SourceIndex].Get();
+					Arguments.Pipelines = AsyncHelper->Pipelines;
+					ObjectInfo.Factory->PreImportPreCompletedCallback(Arguments);
+				}
 
-			UInterchangeResultSuccess* Message = Results->Add<UInterchangeResultSuccess>();
-			Message->SourceAssetName = AsyncHelper->SourceDatas[SourceIndex]->GetFilename();
-			Message->DestinationAssetName = Asset->GetPathName();
-			Message->AssetType = Asset->GetClass();
+				if (ImportedObject == nullptr)
+				{
+					continue;
+				}
 
-			//Clear any async flag from the created asset
-			const EInternalObjectFlags AsyncFlags = EInternalObjectFlags::Async | EInternalObjectFlags::AsyncLoading;
-			Asset->ClearInternalFlags(AsyncFlags);
-			//Make sure the package is dirty
-			Asset->MarkPackageDirty();
+				UInterchangeResultSuccess* Message = Results->Add<UInterchangeResultSuccess>();
+				Message->SourceAssetName = AsyncHelper->SourceDatas[SourceIndex]->GetFilename();
+				Message->DestinationAssetName = ImportedObject->GetPathName();
+				Message->AssetType = ImportedObject->GetClass();
+
+				if (bIsAsset)
+				{
+					//Clear any async flag from the created asset
+					const EInternalObjectFlags AsyncFlags = EInternalObjectFlags::Async | EInternalObjectFlags::AsyncLoading;
+					ImportedObject->ClearInternalFlags(AsyncFlags);
+					//Make sure the package is dirty
+					ImportedObject->MarkPackageDirty();
 #if WITH_EDITOR
-			//Make sure the asset is built correctly
-			Asset->PostEditChange();
+					//Make sure the asset is built correctly
+					ImportedObject->PostEditChange();
 #endif
-			//Post import broadcast
-			if (!AsyncHelper->TaskData.ReimportObject)
-			{
-				//Notify the asset registry, only when we have created the asset
-				FAssetRegistryModule::AssetCreated(Asset);
-			}
-			AsyncHelper->AssetImportResult->AddImportedAsset(Asset);
-			//In case Some factory code cannot run outside of the main thread we offer this callback to finish the work after calling post edit change (building the asset)
-			//Its possible the build of the asset to be asynchronous, the factory must handle is own asset correctly
-			if (bCallPostImportGameThreadCallback && AssetInfo.Factory)
-			{
-				UInterchangeFactoryBase::FImportPreCompletedCallbackParams Arguments;
-				Arguments.ImportedObject = Asset;
-				Arguments.SourceData = AsyncHelper->SourceDatas[SourceIndex];
-				Arguments.FactoryNode = AssetInfo.FactoryNode;
-				Arguments.NodeUniqueID = AssetInfo.FactoryNode ? AssetInfo.FactoryNode->GetUniqueID() : FString();
-				Arguments.NodeContainer = AsyncHelper->BaseNodeContainers[SourceIndex].Get();
-				Arguments.Pipelines = AsyncHelper->Pipelines;
-				Arguments.bIsReimport = AssetInfo.bIsReimport;
-				AssetInfo.Factory->PostImportPreCompletedCallback(Arguments);
+					//Post import broadcast
+					if (!AsyncHelper->TaskData.ReimportObject)
+					{
+						//Notify the asset registry, only when we have created the asset
+						FAssetRegistryModule::AssetCreated(ImportedObject);
+					}
+					AsyncHelper->AssetImportResult->AddImportedObject(ImportedObject);
+				}
+				else
+				{
+					AsyncHelper->SceneImportResult->AddImportedObject(ImportedObject);
+				}
+
+				//In case Some factory code cannot run outside of the main thread we offer this callback to finish the work after calling post edit change (building the asset)
+				//Its possible the build of the asset to be asynchronous, the factory must handle is own asset correctly
+				if (bCallPostImportGameThreadCallback && ObjectInfo.Factory)
+				{
+					UInterchangeFactoryBase::FImportPreCompletedCallbackParams Arguments;
+					Arguments.ImportedObject = ImportedObject;
+					Arguments.SourceData = AsyncHelper->SourceDatas[SourceIndex];
+					Arguments.FactoryNode = ObjectInfo.FactoryNode;
+					Arguments.NodeUniqueID = ObjectInfo.FactoryNode ? ObjectInfo.FactoryNode->GetUniqueID() : FString();
+					Arguments.NodeContainer = AsyncHelper->BaseNodeContainers[SourceIndex].Get();
+					Arguments.Pipelines = AsyncHelper->Pipelines;
+					Arguments.bIsReimport = ObjectInfo.bIsReimport;
+					ObjectInfo.Factory->PostImportPreCompletedCallback(Arguments);
+				}
 			}
 		}
-	}
+	};
+
+	bool bIsAsset = true;
+	ForEachImportedObjectForEachSource(AsyncHelper->ImportedAssetsPerSourceIndex, bIsAsset);
+
+	bIsAsset = false;
+	ForEachImportedObjectForEachSource(AsyncHelper->ImportedSceneObjectsPerSourceIndex, bIsAsset);
 }
 
 
@@ -120,7 +138,9 @@ void UE::Interchange::FTaskCompletion::DoTask(ENamedThreads::Type CurrentThread,
 
 	//No need anymore of the translators sources
 	AsyncHelper->ReleaseTranslatorsSource();
-	for(TPair<int32, TArray<FImportAsyncHelper::FImportedAssetInfo>>& AssetInfosPerSourceIndexPair : AsyncHelper->ImportedAssetsPerSourceIndex)
+
+	//Broadcast OnAssetPostImport/OnAssetPostReimport for each imported asset
+	for(TPair<int32, TArray<FImportAsyncHelper::FImportedObjectInfo>>& AssetInfosPerSourceIndexPair : AsyncHelper->ImportedAssetsPerSourceIndex)
 	{
 		//Verify if the task was cancel
 		if (AsyncHelper->bCancel)
@@ -128,9 +148,9 @@ void UE::Interchange::FTaskCompletion::DoTask(ENamedThreads::Type CurrentThread,
 			break;
 		}
 		const int32 SourceIndex = AssetInfosPerSourceIndexPair.Key;
-		for (const FImportAsyncHelper::FImportedAssetInfo& AssetInfo : AssetInfosPerSourceIndexPair.Value)
+		for (const FImportAsyncHelper::FImportedObjectInfo& AssetInfo : AssetInfosPerSourceIndexPair.Value)
 		{
-			UObject* Asset = AssetInfo.ImportAsset;
+			UObject* Asset = AssetInfo.ImportedObject;
 			if (AsyncHelper->TaskData.ReimportObject)
 			{
 				InterchangeManager->OnAssetPostReimport.Broadcast(Asset);
@@ -142,15 +162,15 @@ void UE::Interchange::FTaskCompletion::DoTask(ENamedThreads::Type CurrentThread,
 		}
 	}
 
-	//If task is cancel, delete all created assets by this task
 	if (AsyncHelper->bCancel)
 	{
-		for (TPair<int32, TArray<FImportAsyncHelper::FImportedAssetInfo>>& AssetInfosPerSourceIndexPair : AsyncHelper->ImportedAssetsPerSourceIndex)
+		//If task is canceled, delete all created assets by this task
+		for (TPair<int32, TArray<FImportAsyncHelper::FImportedObjectInfo>>& AssetInfosPerSourceIndexPair : AsyncHelper->ImportedAssetsPerSourceIndex)
 		{
 			const int32 SourceIndex = AssetInfosPerSourceIndexPair.Key;
-			for (const FImportAsyncHelper::FImportedAssetInfo& AssetInfo : AssetInfosPerSourceIndexPair.Value)
+			for (const FImportAsyncHelper::FImportedObjectInfo& AssetInfo : AssetInfosPerSourceIndexPair.Value)
 			{
-				UObject* Asset = AssetInfo.ImportAsset;
+				UObject* Asset = AssetInfo.ImportedObject;
 				if (Asset)
 				{
 					//Make any created asset go away
@@ -161,9 +181,27 @@ void UE::Interchange::FTaskCompletion::DoTask(ENamedThreads::Type CurrentThread,
 				}
 			}
 		}
+
+		//If task is canceled, remove all actors from their world
+		for (TPair<int32, TArray<FImportAsyncHelper::FImportedObjectInfo>>& SceneObjectInfosPerSourceIndexPair : AsyncHelper->ImportedSceneObjectsPerSourceIndex)
+		{
+			const int32 SourceIndex = SceneObjectInfosPerSourceIndexPair.Key;
+			for (const FImportAsyncHelper::FImportedObjectInfo& SceneObjectInfo : SceneObjectInfosPerSourceIndexPair.Value)
+			{
+				if (AActor* Actor = Cast<AActor>(SceneObjectInfo.ImportedObject))
+				{
+					if (UWorld* ActorWorld = Actor->GetWorld())
+					{
+						const bool bModifyLevel = false; //This isn't undoable
+						ActorWorld->RemoveActor(Actor, bModifyLevel);
+					}
+				}
+			}
+		}
 	}
 
 	AsyncHelper->AssetImportResult->SetDone();
+	AsyncHelper->SceneImportResult->SetDone();
 
 	//Release the async helper
 	AsyncHelper = nullptr;
diff --git a/Engine/Source/Runtime/Interchange/Engine/Private/Tasks/InterchangeTaskCreateAsset.cpp b/Engine/Source/Runtime/Interchange/Engine/Private/Tasks/InterchangeTaskCreateAsset.cpp
--- a/Engine/Source/Runtime/Interchange/Engine/Private/Tasks/InterchangeTaskCreateAsset.cpp
+++ b/Engine/Source/Runtime/Interchange/Engine/Private/Tasks/InterchangeTaskCreateAsset.cpp
@@ -149,9 +149,9 @@ void UE::Interchange::FTaskCreatePackage::DoTask(ENamedThreads::Type CurrentThre
 				NodeAsset->SetInternalFlags(EInternalObjectFlags::Async);
 			}
 			FScopeLock Lock(&AsyncHelper->ImportedAssetsPerSourceIndexLock);
-			TArray<UE::Interchange::FImportAsyncHelper::FImportedAssetInfo>& ImportedInfos = AsyncHelper->ImportedAssetsPerSourceIndex.FindOrAdd(SourceIndex);
-			UE::Interchange::FImportAsyncHelper::FImportedAssetInfo& AssetInfo = ImportedInfos.AddDefaulted_GetRef();
-			AssetInfo.ImportAsset = NodeAsset;
+			TArray<UE::Interchange::FImportAsyncHelper::FImportedObjectInfo>& ImportedInfos = AsyncHelper->ImportedAssetsPerSourceIndex.FindOrAdd(SourceIndex);
+			UE::Interchange::FImportAsyncHelper::FImportedObjectInfo& AssetInfo = ImportedInfos.AddDefaulted_GetRef();
+			AssetInfo.ImportedObject = NodeAsset;
 			AssetInfo.Factory = Factory;
 			AssetInfo.FactoryNode = Node;
 			AssetInfo.bIsReimport = bool(AsyncHelper->TaskData.ReimportObject);
@@ -274,16 +274,16 @@ void UE::Interchange::FTaskCreateAsset::DoTask(ENamedThreads::Type CurrentThread
 		if (!bSkipAsset)
 		{
 			FScopeLock Lock(&AsyncHelper->ImportedAssetsPerSourceIndexLock);
-			TArray<UE::Interchange::FImportAsyncHelper::FImportedAssetInfo>& ImportedInfos = AsyncHelper->ImportedAssetsPerSourceIndex.FindOrAdd(SourceIndex);
-			UE::Interchange::FImportAsyncHelper::FImportedAssetInfo* AssetInfoPtr = ImportedInfos.FindByPredicate([NodeAsset](const UE::Interchange::FImportAsyncHelper::FImportedAssetInfo& CurInfo)
+			TArray<UE::Interchange::FImportAsyncHelper::FImportedObjectInfo>& ImportedInfos = AsyncHelper->ImportedAssetsPerSourceIndex.FindOrAdd(SourceIndex);
+			UE::Interchange::FImportAsyncHelper::FImportedObjectInfo* AssetInfoPtr = ImportedInfos.FindByPredicate([NodeAsset](const UE::Interchange::FImportAsyncHelper::FImportedObjectInfo& CurInfo)
 			{
-				return CurInfo.ImportAsset == NodeAsset;
+				return CurInfo.ImportedObject == NodeAsset;
 			});
 
 			if (!AssetInfoPtr)
 			{
-				UE::Interchange::FImportAsyncHelper::FImportedAssetInfo& AssetInfo = ImportedInfos.AddDefaulted_GetRef();
-				AssetInfo.ImportAsset = NodeAsset;
+				UE::Interchange::FImportAsyncHelper::FImportedObjectInfo& AssetInfo = ImportedInfos.AddDefaulted_GetRef();
+				AssetInfo.ImportedObject = NodeAsset;
 				AssetInfo.Factory = Factory;
 				AssetInfo.FactoryNode = Node;
 				AssetInfo.bIsReimport = bool(AsyncHelper->TaskData.ReimportObject);
diff --git a/Engine/Source/Runtime/Interchange/Engine/Private/Tasks/InterchangeTaskCreateSceneObjects.cpp b/Engine/Source/Runtime/Interchange/Engine/Private/Tasks/InterchangeTaskCreateSceneObjects.cpp
new file mode 100644
--- /dev/null
+++ b/Engine/Source/Runtime/Interchange/Engine/Private/Tasks/InterchangeTaskCreateSceneObjects.cpp
@@ -0,0 +1,100 @@
+// Copyright Epic Games, Inc. All Rights Reserved.
+#include "InterchangeTaskCreateSceneObjects.h"
+
+#include "CoreMinimal.h"
+
+#include "Async/TaskGraphInterfaces.h"
+#include "Engine/World.h"
+#include "InterchangeEngineLogPrivate.h"
+#include "InterchangeFactoryBase.h"
+#include "InterchangeManager.h"
+#include "InterchangeSourceData.h"
+#include "InterchangeTranslatorBase.h"
+#include "Misc/Paths.h"
+#include "Nodes/InterchangeBaseNode.h"
+#include "PackageUtils/PackageUtils.h"
+#include "Stats/Stats.h"
+#include "Templates/SharedPointer.h"
+#include "UObject/Object.h"
+#include "UObject/UObjectGlobals.h"
+#include "UObject/ObjectMacros.h"
+#include "UObject/WeakObjectPtrTemplates.h"
+
+UE::Interchange::FTaskCreateSceneObjects::FTaskCreateSceneObjects(const FString& InPackageBasePath, const int32 InSourceIndex, TWeakPtr<FImportAsyncHelper, ESPMode::ThreadSafe> InAsyncHelper,
+	TArrayView<UInterchangeBaseNode*> InNodes, UInterchangeFactoryBase* InFactory, bool bInCreateSceneObjectsForChildren)
+	: PackageBasePath(InPackageBasePath)
+	, SourceIndex(InSourceIndex)
+	, WeakAsyncHelper(InAsyncHelper)
+	, Nodes(InNodes)
+	, Factory(InFactory)
+	, bCreateSceneObjectsForChildren(bInCreateSceneObjectsForChildren)
+{
+	check(Factory);
+}
+
+ENamedThreads::Type UE::Interchange::FTaskCreateSceneObjects::GetDesiredThread()
+{
+	return Factory->CanExecuteOnAnyThread() ? ENamedThreads::AnyBackgroundThreadNormalTask : ENamedThreads::GameThread;
+}
+
+void UE::Interchange::FTaskCreateSceneObjects::DoTask(ENamedThreads::Type CurrentThread, const FGraphEventRef& MyCompletionGraphEvent)
+{
+#if INTERCHANGE_TRACE_ASYNCHRONOUS_TASK_ENABLED
+	INTERCHANGE_TRACE_ASYNCHRONOUS_TASK(SpawnActor)
+#endif
+	TSharedPtr<UE::Interchange::FImportAsyncHelper, ESPMode::ThreadSafe> AsyncHelper = WeakAsyncHelper.Pin();
+	check(WeakAsyncHelper.IsValid());
+
+	//Verify if the task was canceled
+	if (AsyncHelper->bCancel)
+	{
+		return;
+	}
+
+	for (UInterchangeBaseNode* Node : Nodes)
+	{
+		FString NodeDisplayName = Node->GetDisplayLabel();
+		SanitizeObjectName(NodeDisplayName);
+
+		UInterchangeFactoryBase::FCreateSceneObjectsParams CreateSceneObjectsParams;
+		CreateSceneObjectsParams.ObjectName = NodeDisplayName;
+		CreateSceneObjectsParams.ObjectNode = Node;
+		CreateSceneObjectsParams.Level = GWorld->GetCurrentLevel();
+		CreateSceneObjectsParams.bCreateSceneObjectsForChildren = bCreateSceneObjectsForChildren;
+
+		if (AsyncHelper->BaseNodeContainers.IsValidIndex(SourceIndex))
+		{
+			CreateSceneObjectsParams.NodeContainer = AsyncHelper->BaseNodeContainers[SourceIndex].Get();
+		}
+
+		TMap<FString, UObject*> SceneObjects = Factory->CreateSceneObjects(CreateSceneObjectsParams);
+
+		for (const TPair<FString, UObject*>& SceneObject : SceneObjects)
+		{
+			if (UObject* NodeObject = SceneObject.Value)
+			{
+				FScopeLock Lock(&AsyncHelper->ImportedSceneObjectsPerSourceIndexLock);
+				TArray<UE::Interchange::FImportAsyncHelper::FImportedObjectInfo>& ImportedInfos = AsyncHelper->ImportedSceneObjectsPerSourceIndex.FindOrAdd(SourceIndex);
+				UE::Interchange::FImportAsyncHelper::FImportedObjectInfo* ImportedInfoPtr = ImportedInfos.FindByPredicate([NodeObject](const UE::Interchange::FImportAsyncHelper::FImportedObjectInfo& CurInfo)
+					{
+						return CurInfo.ImportedObject == NodeObject;
+					});
+
+				if (!ImportedInfoPtr)
+				{
+					UE::Interchange::FImportAsyncHelper::FImportedObjectInfo& ObjectInfo = ImportedInfos.AddDefaulted_GetRef();
+					ObjectInfo.ImportedObject = NodeObject;
+					ObjectInfo.Factory = Factory;
+				}
+
+				if (CreateSceneObjectsParams.NodeContainer)
+				{
+					if (const UInterchangeBaseNode* ActorNode = CreateSceneObjectsParams.NodeContainer->GetNode(SceneObject.Key))
+					{
+						ActorNode->ReferenceObject = FSoftObjectPath(SceneObject.Value);
+					}
+				}
+			}
+		}
+	}
+}
\ No newline at end of file
diff --git a/Engine/Source/Runtime/Interchange/Engine/Private/Tasks/InterchangeTaskCreateSceneObjects.h b/Engine/Source/Runtime/Interchange/Engine/Private/Tasks/InterchangeTaskCreateSceneObjects.h
new file mode 100644
--- /dev/null
+++ b/Engine/Source/Runtime/Interchange/Engine/Private/Tasks/InterchangeTaskCreateSceneObjects.h
@@ -0,0 +1,51 @@
+// Copyright Epic Games, Inc. All Rights Reserved.
+
+#pragma once
+
+#include "Async/TaskGraphInterfaces.h"
+#include "CoreMinimal.h"
+#include "Stats/Stats.h"
+#include "UObject/WeakObjectPtrTemplates.h"
+#include "Nodes/InterchangeBaseNode.h"
+
+class UInterchangeBaseNode;
+class UInterchangeFactoryBase;
+
+namespace UE
+{
+	namespace Interchange
+	{
+		class FImportAsyncHelper;
+
+		class FTaskCreateSceneObjects
+		{
+		private:
+			FString PackageBasePath;
+			int32 SourceIndex;
+			TWeakPtr<FImportAsyncHelper, ESPMode::ThreadSafe> WeakAsyncHelper;
+			TArray<UInterchangeBaseNode*> Nodes;
+			UInterchangeFactoryBase* Factory;
+			bool bCreateSceneObjectsForChildren = false;
+
+		public:
+			explicit FTaskCreateSceneObjects(const FString& InPackageBasePath, const int32 InSourceIndex, TWeakPtr<FImportAsyncHelper, ESPMode::ThreadSafe> InAsyncHelper, TArrayView<UInterchangeBaseNode*> InNodes,
+				UInterchangeFactoryBase* InFactory, bool bInCreateSceneObjectsForChildren);
+
+			ENamedThreads::Type GetDesiredThread();
+
+			static ESubsequentsMode::Type GetSubsequentsMode()
+			{
+				return ESubsequentsMode::TrackSubsequents;
+			}
+
+			TStatId GetStatId() const
+			{
+				RETURN_QUICK_DECLARE_CYCLE_STAT(FTaskCreateAsset, STATGROUP_TaskGraphTasks);
+			}
+
+			void DoTask(ENamedThreads::Type CurrentThread, const FGraphEventRef& MyCompletionGraphEvent);
+		};
+
+
+	} //ns Interchange
+}//ns UE
diff --git a/Engine/Source/Runtime/Interchange/Engine/Private/Tasks/InterchangeTaskParsing.cpp b/Engine/Source/Runtime/Interchange/Engine/Private/Tasks/InterchangeTaskParsing.cpp
--- a/Engine/Source/Runtime/Interchange/Engine/Private/Tasks/InterchangeTaskParsing.cpp
+++ b/Engine/Source/Runtime/Interchange/Engine/Private/Tasks/InterchangeTaskParsing.cpp
@@ -3,6 +3,7 @@
 
 #include "Async/TaskGraphInterfaces.h"
 #include "CoreMinimal.h"
+#include "GameFramework/Actor.h"
 #include "InterchangeEngineLogPrivate.h"
 #include "InterchangeFactoryBase.h"
 #include "InterchangeManager.h"
@@ -10,8 +11,11 @@
 #include "InterchangeTaskCompletion.h"
 #include "InterchangeTaskCreateAsset.h"
 #include "InterchangeTaskPipeline.h"
+#include "InterchangeTaskCreateSceneObjects.h"
 #include "InterchangeTranslatorBase.h"
 #include "Misc/Paths.h"
+#include "Nodes/InterchangeBaseNode.h"
+#include "Nodes/InterchangeBaseNodeContainer.h"
 #include "PackageUtils/PackageUtils.h"
 #include "Stats/Stats.h"
 #include "Templates/SharedPointer.h"
@@ -19,8 +23,7 @@
 #include "UObject/UObjectGlobals.h"
 #include "UObject/ObjectMacros.h"
 #include "UObject/WeakObjectPtrTemplates.h"
-#include "Nodes/InterchangeBaseNode.h"
-#include "Nodes/InterchangeBaseNodeContainer.h"
+
 
 void UE::Interchange::FTaskParsing::DoTask(ENamedThreads::Type CurrentThread, const FGraphEventRef& MyCompletionGraphEvent)
 {
@@ -34,19 +37,19 @@ void UE::Interchange::FTaskParsing::DoTask(ENamedThreads::Type CurrentThread, co
 	struct FTaskData
 	{
 		FString UniqueID;
-		int32 SourceIndex;
-		UInterchangeBaseNode* Node;
+		int32 SourceIndex = INDEX_NONE;
 		TArray<FString> Dependencies;
 		FGraphEventRef GraphEventRef;
-		FGraphEventArray Prerequistes;
+		FGraphEventArray Prerequisites;
 		const UClass* FactoryClass;
-	};
 
+		TArray<UInterchangeBaseNode*, TInlineAllocator<1>> Nodes; // For scenes, we can group multiple nodes into a single task as they are usually very light
+	};
 
-	
-	TArray<FTaskData> TaskDatas;
+	TArray<FTaskData> AssetTaskDatas;
+	TArray<FTaskData> SceneTaskDatas;
 
-	//Avoid creating asset if the asynchronous import is cancel, just create the completion task
+	//Avoid creating asset if the asynchronous import is canceled, just create the completion task
 	if (!AsyncHelper->bCancel)
 	{
 		for (int32 SourceIndex = 0; SourceIndex < AsyncHelper->SourceDatas.Num(); ++SourceIndex)
@@ -64,23 +67,33 @@ void UE::Interchange::FTaskParsing::DoTask(ENamedThreads::Type CurrentThread, co
 			{
 				if (!Node->IsEnabled())
 				{
-					//Do not call factory for a disable node
+					//Do not call factory for a disabled node
 					return;
 				}
 
 				if (Node->GetObjectClass() != nullptr)
 				{
 					const UClass* RegisteredFactoryClass = InterchangeManager->GetRegisteredFactoryClass(Node->GetObjectClass());
-					if (!RegisteredFactoryClass)
+
+					const bool bIsAsset = !(Node->GetObjectClass()->IsChildOf<AActor>() || Node->GetObjectClass()->IsChildOf<UActorComponent>());
+					const bool bCanImportSceneNode = [&AsyncHelper, &Node]()
+					{
+						bool bCanImport = AsyncHelper->TaskData.ImportType == EImportType::ImportType_Scene;
+						bCanImport = bCanImport && Node->GetParentUid().IsEmpty(); // We only import root scene nodes since we ask the factory to spawn the children nodes
+
+						return bCanImport;
+					}();
+
+					if (!RegisteredFactoryClass || (!bIsAsset && !bCanImportSceneNode))
 					{
 						//nothing we can import from this element
 						return;
 					}
 
-					FTaskData& NodeTaskData = TaskDatas.AddDefaulted_GetRef();
+					FTaskData& NodeTaskData = bIsAsset ? AssetTaskDatas.AddDefaulted_GetRef() : SceneTaskDatas.AddDefaulted_GetRef();
 					NodeTaskData.UniqueID = Node->GetUniqueID();
 					NodeTaskData.SourceIndex = SourceIndex;
-					NodeTaskData.Node = Node;
+					NodeTaskData.Nodes.Add(Node);
 					Node->GetFactoryDependencies(NodeTaskData.Dependencies);
 					NodeTaskData.FactoryClass = RegisteredFactoryClass;
 				}
@@ -88,75 +101,121 @@ void UE::Interchange::FTaskParsing::DoTask(ENamedThreads::Type CurrentThread, co
 		}
 	}
 	//Sort per dependencies
-	TaskDatas.Sort([](const FTaskData& A, const FTaskData& B)
-	{
-		//if A is a dependency of B then return true to do A before B
-		if (B.Dependencies.Contains(A.UniqueID))
-		{
-			return true;
-		}
-		if (A.Dependencies.Contains(B.UniqueID))
+	auto SortByDependencies =
+		[](const FTaskData& A, const FTaskData& B)
 		{
-			return false;
-		}
-		return A.Dependencies.Num() <= B.Dependencies.Num();
-	});
+			//if A is a dependency of B then return true to do A before B
+			if (B.Dependencies.Contains(A.UniqueID))
+			{
+				return true;
+			}
+			if (A.Dependencies.Contains(B.UniqueID))
+			{
+				return false;
+			}
+			return A.Dependencies.Num() <= B.Dependencies.Num();
+		};
 
-	FGraphEventArray CompletionPrerequistes;
+	AssetTaskDatas.Sort(SortByDependencies);
+	SceneTaskDatas.Sort(SortByDependencies);
 
-	for (int32 TaskIndex = 0; TaskIndex < TaskDatas.Num(); ++TaskIndex)
+	auto CreateTasksForEachTaskData = [](TArray<FTaskData>& TaskDatas, TFunction<FGraphEventRef(FTaskData&)> CreateTasksFunc) -> FGraphEventArray
 	{
-		if (TaskDatas[TaskIndex].Dependencies.Num() > 0)
+		FGraphEventArray GraphEvents;
+
+		for (int32 TaskIndex = 0; TaskIndex < TaskDatas.Num(); ++TaskIndex)
 		{
-			//Search the previous node to find the dependence
-			for (int32 DepTaskIndex = 0; DepTaskIndex < TaskIndex; ++DepTaskIndex)
+			FTaskData& TaskData = TaskDatas[TaskIndex];
+
+			if (TaskData.Dependencies.Num() > 0)
 			{
-				if (TaskDatas[TaskIndex].Dependencies.Contains(TaskDatas[DepTaskIndex].UniqueID))
+				//Search the previous node to find the dependence
+				for (int32 DepTaskIndex = 0; DepTaskIndex < TaskIndex; ++DepTaskIndex)
 				{
-					//Add has prerequisite
-					TaskDatas[TaskIndex].Prerequistes.Add(TaskDatas[DepTaskIndex].GraphEventRef);
+					if (TaskData.Dependencies.Contains(TaskDatas[DepTaskIndex].UniqueID))
+					{
+						//Add has prerequisite
+						TaskData.Prerequisites.Add(TaskDatas[DepTaskIndex].GraphEventRef);
+					}
 				}
 			}
+
+			TaskData.GraphEventRef = CreateTasksFunc(TaskData);
+			GraphEvents.Add(TaskData.GraphEventRef);
 		}
 
-		const int32 SourceIndex = TaskDatas[TaskIndex].SourceIndex;
-		const UClass* const FactoryClass = TaskDatas[TaskIndex].FactoryClass;
-		UInterchangeBaseNode* const FactoryNode = TaskDatas[TaskIndex].Node;
-		const bool bFactoryCanRunOnAnyThread = FactoryClass->GetDefaultObject<UInterchangeFactoryBase>()->CanExecuteOnAnyThread();
-
-		//Add create package task has a prerequisite of FTaskCreateAsset. Create package task is a game thread task
-		FGraphEventArray CreatePackagePrerequistes;
-		int32 CreatePackageTaskIndex = AsyncHelper->CreatePackageTasks.Add(
-			TGraphTask<FTaskCreatePackage>::CreateTask(&(TaskDatas[TaskIndex].Prerequistes)).ConstructAndDispatchWhenReady(PackageBasePath, SourceIndex, WeakAsyncHelper, FactoryNode, FactoryClass)
-		);
-		CreatePackagePrerequistes.Add(AsyncHelper->CreatePackageTasks[CreatePackageTaskIndex]);
-
-		FGraphEventArray CreateAssetPrerequistes;
-		int32 CreateTaskIndex = AsyncHelper->CreateAssetTasks.Add(
-			TGraphTask<FTaskCreateAsset>::CreateTask(&(CreatePackagePrerequistes)).ConstructAndDispatchWhenReady(PackageBasePath, SourceIndex, WeakAsyncHelper, FactoryNode, bFactoryCanRunOnAnyThread)
-		);
-		CreateAssetPrerequistes.Add(AsyncHelper->CreateAssetTasks[CreateTaskIndex]);
-		
-		TaskDatas[TaskIndex].GraphEventRef = AsyncHelper->CreateAssetTasks[CreateTaskIndex];
-		CompletionPrerequistes.Add(TaskDatas[TaskIndex].GraphEventRef);
+		return GraphEvents;
+	};
+
+	//Assets
+	FGraphEventArray AssetsCompletionPrerequistes;
+	{
+		TFunction<FGraphEventRef(FTaskData&)> CreateTasksForAsset = [this, &AsyncHelper](FTaskData& TaskData)
+		{
+			check(TaskData.Nodes.Num() == 1); //We expect 1 node per asset task
+
+			const int32 SourceIndex = TaskData.SourceIndex;
+			const UClass* const FactoryClass = TaskData.FactoryClass;
+			UInterchangeBaseNode* const FactoryNode = TaskData.Nodes[0];
+			const bool bFactoryCanRunOnAnyThread = FactoryClass->GetDefaultObject<UInterchangeFactoryBase>()->CanExecuteOnAnyThread();
+
+			//Add create package task has a prerequisite of FTaskCreateAsset. Create package task is a game thread task
+			FGraphEventArray CreatePackagePrerequistes;
+			int32 CreatePackageTaskIndex = AsyncHelper->CreatePackageTasks.Add(
+				TGraphTask<FTaskCreatePackage>::CreateTask(&(TaskData.Prerequisites)).ConstructAndDispatchWhenReady(PackageBasePath, SourceIndex, WeakAsyncHelper, FactoryNode, FactoryClass)
+			);
+			CreatePackagePrerequistes.Add(AsyncHelper->CreatePackageTasks[CreatePackageTaskIndex]);
+
+			FGraphEventArray CreateAssetPrerequistes;
+			int32 CreateTaskIndex = AsyncHelper->CreateAssetTasks.Add(
+				TGraphTask<FTaskCreateAsset>::CreateTask(&(CreatePackagePrerequistes)).ConstructAndDispatchWhenReady(PackageBasePath, SourceIndex, WeakAsyncHelper, FactoryNode, bFactoryCanRunOnAnyThread)
+			);
+			CreateAssetPrerequistes.Add(AsyncHelper->CreateAssetTasks[CreateTaskIndex]);
+
+			return AsyncHelper->CreateAssetTasks[CreateTaskIndex];
+		};
+
+		AssetsCompletionPrerequistes = CreateTasksForEachTaskData(AssetTaskDatas, CreateTasksForAsset);
 	}
 
-	//Add a async task for pre completion
-	
+	//Scenes
+	//Note: Scene tasks are delayed until all asset tasks are completed
+	FGraphEventArray ScenesCompletionPrerequistes;
+	{
+		TFunction<FGraphEventRef(FTaskData&)> CreateTasksForSceneObject = [this, &AsyncHelper, &AssetsCompletionPrerequistes](FTaskData& TaskData)
+		{
+			const bool bSpawnChildren = true;
+			const int32 SourceIndex = TaskData.SourceIndex;
+			const UClass* const FactoryClass = TaskData.FactoryClass;
+
+			return AsyncHelper->SceneTasks.Add_GetRef(
+				TGraphTask<FTaskCreateSceneObjects>::CreateTask(&AssetsCompletionPrerequistes)
+				.ConstructAndDispatchWhenReady(PackageBasePath, SourceIndex, WeakAsyncHelper, TaskData.Nodes, FactoryClass->GetDefaultObject<UInterchangeFactoryBase>(), bSpawnChildren));
+
+		};
 
+		ScenesCompletionPrerequistes = CreateTasksForEachTaskData(SceneTaskDatas, CreateTasksForSceneObject);
+	}
+
+	FGraphEventArray CompletionPrerequistes;
+	CompletionPrerequistes.Append(AssetsCompletionPrerequistes);
+	CompletionPrerequistes.Append(ScenesCompletionPrerequistes);
+
+	//Add an async task for pre completion
+	
 	FGraphEventArray PreCompletionPrerequistes;
 	AsyncHelper->PreCompletionTask = TGraphTask<FTaskPreCompletion>::CreateTask(&CompletionPrerequistes).ConstructAndDispatchWhenReady(InterchangeManager, WeakAsyncHelper);
 	PreCompletionPrerequistes.Add(AsyncHelper->PreCompletionTask);
 
 	//Start the Post pipeline task
-	for (int32 TaskIndex = 0; TaskIndex < TaskDatas.Num(); ++TaskIndex)
+	for (int32 SourceIndex = 0; SourceIndex < AsyncHelper->SourceDatas.Num(); ++SourceIndex)
 	{
 		for (int32 GraphPipelineIndex = 0; GraphPipelineIndex < AsyncHelper->Pipelines.Num(); ++GraphPipelineIndex)
 		{
 			int32 GraphPipelineTaskIndex = AsyncHelper->PipelinePostImportTasks.Add(
-				TGraphTask<FTaskPipelinePostImport>::CreateTask(&(PreCompletionPrerequistes)).ConstructAndDispatchWhenReady(TaskDatas[TaskIndex].SourceIndex, GraphPipelineIndex, WeakAsyncHelper)
+				TGraphTask<FTaskPipelinePostImport>::CreateTask(&(PreCompletionPrerequistes)).ConstructAndDispatchWhenReady(SourceIndex, GraphPipelineIndex, WeakAsyncHelper)
 			);
-			//Ensure we run the pipeline in the same order we create the task, since pipeline modify the node container, its important that its not process in parallel, Adding the one we start to the prerequisites
+			//Ensure we run the pipeline in the same order we create the task, since the pipeline modifies the node container, its important that its not processed in parallel, Adding the one we start to the prerequisites
 			//is the way to go here
 			PreCompletionPrerequistes.Add(AsyncHelper->PipelinePostImportTasks[GraphPipelineTaskIndex]);
 		}
diff --git a/Engine/Source/Runtime/Interchange/Engine/Private/Tasks/InterchangeTaskPipeline.cpp b/Engine/Source/Runtime/Interchange/Engine/Private/Tasks/InterchangeTaskPipeline.cpp
--- a/Engine/Source/Runtime/Interchange/Engine/Private/Tasks/InterchangeTaskPipeline.cpp
+++ b/Engine/Source/Runtime/Interchange/Engine/Private/Tasks/InterchangeTaskPipeline.cpp
@@ -58,24 +58,29 @@ void UE::Interchange::FTaskPipelinePostImport::DoTask(ENamedThreads::Type Curren
 	}
 	UInterchangePipelineBase* PipelineBase = AsyncHelper->Pipelines[PipelineIndex];
 	TArray<FString> NodeUniqueIDs;
-	TArray<UObject*> ImportAssets;
-	//Create a lock scope to read the Imported asset infos map
-	{
-		FScopeLock Lock(&AsyncHelper->ImportedAssetsPerSourceIndexLock);
-		if (AsyncHelper->ImportedAssetsPerSourceIndex.Contains(SourceIndex))
+	TArray<UObject*> ImportedObjects;
+
+	auto FillImportedObjectsFromSource =
+		[&NodeUniqueIDs, &ImportedObjects, this](FCriticalSection& CriticalSection, const TMap<int32, TArray<UE::Interchange::FImportAsyncHelper::FImportedObjectInfo>>& ImportedInfosPerSource)
 		{
-			TArray<UE::Interchange::FImportAsyncHelper::FImportedAssetInfo>& ImportedInfos = AsyncHelper->ImportedAssetsPerSourceIndex.FindChecked(SourceIndex);
-			NodeUniqueIDs.Reserve(ImportedInfos.Num());
-			ImportAssets.Reserve(ImportedInfos.Num());
-			for (UE::Interchange::FImportAsyncHelper::FImportedAssetInfo& ImportedInfo : ImportedInfos)
+			FScopeLock Lock(&CriticalSection);
+			if (ImportedInfosPerSource.Contains(SourceIndex))
 			{
-				NodeUniqueIDs.Add(ImportedInfo.FactoryNode->GetUniqueID());
-				ImportAssets.Add(ImportedInfo.ImportAsset);
+				const TArray<UE::Interchange::FImportAsyncHelper::FImportedObjectInfo>& ImportedInfos = ImportedInfosPerSource.FindChecked(SourceIndex);
+				NodeUniqueIDs.Reserve(ImportedInfos.Num());
+				ImportedObjects.Reserve(ImportedInfos.Num());
+				for (const UE::Interchange::FImportAsyncHelper::FImportedObjectInfo& ImportedInfo : ImportedInfos)
+				{
+					NodeUniqueIDs.Add(ImportedInfo.FactoryNode->GetUniqueID());
+					ImportedObjects.Add(ImportedInfo.ImportedObject);
+				}
 			}
-		}
-	}
+		};
+
+	FillImportedObjectsFromSource(AsyncHelper->ImportedAssetsPerSourceIndexLock, AsyncHelper->ImportedAssetsPerSourceIndex);
+	FillImportedObjectsFromSource(AsyncHelper->ImportedSceneObjectsPerSourceIndexLock, AsyncHelper->ImportedSceneObjectsPerSourceIndex);
 
-	if (!ensure(NodeUniqueIDs.Num() == ImportAssets.Num()))
+	if (!ensure(NodeUniqueIDs.Num() == ImportedObjects.Num()))
 	{
 		//We do not execute the script if we cannot give proper parameter
 		return;
@@ -91,8 +96,8 @@ void UE::Interchange::FTaskPipelinePostImport::DoTask(ENamedThreads::Type Curren
 	Pipeline->SetResultsContainer(AsyncHelper->AssetImportResult->GetResults());
 
 	//Call the pipeline outside of the lock, we do this in case the pipeline take a long time. We call it for each asset created by this import
-	for (int32 AssetIndex = 0; AssetIndex < ImportAssets.Num(); ++AssetIndex)
+	for (int32 ObjectIndex = 0; ObjectIndex < ImportedObjects.Num(); ++ObjectIndex)
 	{
-		Pipeline->ScriptedExecutePostImportPipeline(NodeContainer, NodeUniqueIDs[AssetIndex], ImportAssets[AssetIndex]);
+		Pipeline->ScriptedExecutePostImportPipeline(NodeContainer, NodeUniqueIDs[ObjectIndex], ImportedObjects[ObjectIndex]);
 	}
 }
diff --git a/Engine/Source/Runtime/Interchange/Engine/Public/InterchangeEngineFwd.h b/Engine/Source/Runtime/Interchange/Engine/Public/InterchangeEngineFwd.h
--- a/Engine/Source/Runtime/Interchange/Engine/Public/InterchangeEngineFwd.h
+++ b/Engine/Source/Runtime/Interchange/Engine/Public/InterchangeEngineFwd.h
@@ -8,8 +8,9 @@ namespace UE
 {
 	namespace Interchange
 	{
-		class FAssetImportResult;
+		class FImportResult;
 
-		using FAssetImportResultRef = TSharedRef< FAssetImportResult, ESPMode::ThreadSafe >;
+		using FAssetImportResultRef = TSharedRef< FImportResult, ESPMode::ThreadSafe >;
+		using FSceneImportResultRef = TSharedRef< FImportResult, ESPMode::ThreadSafe >;
 	}
 }
diff --git a/Engine/Source/Runtime/Interchange/Engine/Public/InterchangeManager.h b/Engine/Source/Runtime/Interchange/Engine/Public/InterchangeManager.h
--- a/Engine/Source/Runtime/Interchange/Engine/Public/InterchangeManager.h
+++ b/Engine/Source/Runtime/Interchange/Engine/Public/InterchangeManager.h
@@ -3,9 +3,11 @@
 #pragma once
 
 #include <atomic>
+
+#include "CoreMinimal.h"
+
 #include "Async/TaskGraphInterfaces.h"
 #include "Containers/Queue.h"
-#include "CoreMinimal.h"
 #include "Delegates/DelegateCombinations.h"
 #include "HAL/Thread.h"
 #include "HAL/ThreadSafeBool.h"
@@ -15,11 +17,11 @@
 #include "InterchangeTranslatorBase.h"
 #include "InterchangeWriterBase.h"
 #include "Nodes/InterchangeBaseNodeContainer.h"
+#include "Templates/Tuple.h"
 #include "UObject/Package.h"
 #include "UObject/Object.h"
 #include "UObject/ObjectMacros.h"
 #include "UObject/StrongObjectPtr.h"
-#include "Containers/Queue.h"
 
 #include "InterchangeManager.generated.h"
 
@@ -69,18 +71,18 @@ namespace UE
 			UObject* ReimportObject = nullptr;
 		};
 
-		class INTERCHANGEENGINE_API FAssetImportResult : protected FGCObject
+		class INTERCHANGEENGINE_API FImportResult : protected FGCObject
 		{
 		public:
-			FAssetImportResult();
+			FImportResult();
 
-			FAssetImportResult(FAssetImportResult&&) = delete;
-			FAssetImportResult& operator=(FAssetImportResult&&) = delete;
+			FImportResult(FImportResult&&) = delete;
+			FImportResult& operator=(FImportResult&&) = delete;
 
-			FAssetImportResult(const FAssetImportResult&) = delete;
-			FAssetImportResult& operator=(const FAssetImportResult&) = delete;
+			FImportResult(const FImportResult&) = delete;
+			FImportResult& operator=(const FImportResult&) = delete;
 
-			virtual ~FAssetImportResult() = default;
+			virtual ~FImportResult() = default;
 
 		public:
 			enum class EStatus
@@ -101,7 +103,7 @@ namespace UE
 			// Assets are only made available once they have been completely imported (passed through the entire import pipeline)
 			// While the status isn't EStatus::Done, the list can grow between subsequent calls.
 			// FAssetImportResult holds a reference to the assets so that they aren't garbage collected.
-			const TArray< UObject* >& GetImportedAssets() const;
+			const TArray< UObject* >& GetImportedObjects() const;
 
 			// Helper to get the first asset of a certain class. Use when expecting a single asset of that class to be imported since the order isn't deterministic.
 			UObject* GetFirstAssetOfClass(UClass* InClass) const;
@@ -110,10 +112,10 @@ namespace UE
 			UInterchangeResultsContainer* GetResults() const { return Results; }
 
 			// Adds an asset to the list of imported assets.
-			void AddImportedAsset(UObject* ImportedAsset);
+			void AddImportedObject(UObject* ImportedObject);
 
 			// Callback when the status switches to done.
-			void OnDone(TFunction< void(FAssetImportResult&) > Callback);
+			void OnDone(TFunction< void(FImportResult&) > Callback);
 
 		protected:
 			/* FGCObject interface */
@@ -122,16 +124,17 @@ namespace UE
 		private:
 			std::atomic< EStatus > ImportStatus;
 
-			TArray< UObject* > ImportedAssets;
-			mutable FRWLock ImportedAssetsRWLock;
+			TArray< UObject* > ImportedObjects;
+			mutable FRWLock ImportedObjectsRWLock;
 			UInterchangeResultsContainer* Results;
 
 			FGraphEventRef GraphEvent; // WaitUntilDone waits for this event to be triggered.
 
-			TFunction< void(FAssetImportResult&) > DoneCallback;
+			TFunction< void(FImportResult&) > DoneCallback;
 		};
 
-		using FAssetImportResultRef = TSharedRef< FAssetImportResult, ESPMode::ThreadSafe >;
+		using FAssetImportResultRef = TSharedRef< FImportResult, ESPMode::ThreadSafe >;
+		using FSceneImportResultRef = TSharedRef< FImportResult, ESPMode::ThreadSafe >;
 
 		class FImportAsyncHelper : protected FGCObject
 		{
@@ -160,6 +163,7 @@ namespace UE
 			FGraphEventRef ParsingTask;
 			TArray<FGraphEventRef> CreatePackageTasks;
 			TArray<FGraphEventRef> CreateAssetTasks;
+			TArray<FGraphEventRef> SceneTasks;
 
 			FGraphEventRef PreAsyncCompletionTask;
 			FGraphEventRef PreCompletionTask;
@@ -173,20 +177,24 @@ namespace UE
 			FCriticalSection CreatedFactoriesLock;
 			TMap<FString, UInterchangeFactoryBase*> CreatedFactories;
 
-			struct FImportedAssetInfo
+			struct FImportedObjectInfo
 			{
-				UObject* ImportAsset;
-				UInterchangeFactoryBase* Factory;
-				UInterchangeBaseNode* FactoryNode;
+				UObject* ImportedObject = nullptr; // The object that was imported
+				UInterchangeFactoryBase* Factory = nullptr; //The factory that created the imported object
+				UInterchangeBaseNode* FactoryNode; //The node that describes the object
 				bool bIsReimport;
 			};
 
 			FCriticalSection ImportedAssetsPerSourceIndexLock;
-			TMap<int32, TArray<FImportedAssetInfo>> ImportedAssetsPerSourceIndex;
+			TMap<int32, TArray<FImportedObjectInfo>> ImportedAssetsPerSourceIndex;
+
+			FCriticalSection ImportedSceneObjectsPerSourceIndexLock;
+			TMap<int32, TArray<FImportedObjectInfo>> ImportedSceneObjectsPerSourceIndex;
 
 			FImportAsyncHelperData TaskData;
 
 			FAssetImportResultRef AssetImportResult;
+			FSceneImportResultRef SceneImportResult;
 			
 			//If we cancel the tasks, we set this boolean to true
 			std::atomic<bool> bCancel;
@@ -314,7 +322,10 @@ public:
 	 * @return true if the import succeed, false otherwise.
 	 */
 	UFUNCTION(BlueprintCallable, Category = "Interchange | Import Manager")
-	bool ImportScene(const FString& ContentPath, const UInterchangeSourceData* SourceData, bool bIsReimport = false, bool bIsAutomated = false);
+	bool ImportScene(const FString& ContentPath, const UInterchangeSourceData* SourceData, const FImportAssetParameters& ImportAssetParameters);
+
+	TTuple<UE::Interchange::FAssetImportResultRef, UE::Interchange::FSceneImportResultRef>
+	ImportSceneAsync(const FString& ContentPath, const UInterchangeSourceData* SourceData, const FImportAssetParameters& ImportAssetParameters);
 
 	/**
 	 * Call this to start an export asset process, the caller must specify a source data.
@@ -426,6 +437,12 @@ protected:
 	 */
 	void StartQueuedTasks(bool bCancelAllTasks = false);
 
+	/**
+	 * Called by the public Import functions
+	 */
+	TTuple<UE::Interchange::FAssetImportResultRef, UE::Interchange::FSceneImportResultRef>
+	ImportInternal(const FString& ContentPath, const UInterchangeSourceData* SourceData, const FImportAssetParameters& ImportAssetParameters, const UE::Interchange::EImportType ImportType);
+
 private:
 	struct FQueuedTaskData
 	{
diff --git a/Engine/Source/Runtime/Interchange/Engine/Public/InterchangeProjectSettings.h b/Engine/Source/Runtime/Interchange/Engine/Public/InterchangeProjectSettings.h
--- a/Engine/Source/Runtime/Interchange/Engine/Public/InterchangeProjectSettings.h
+++ b/Engine/Source/Runtime/Interchange/Engine/Public/InterchangeProjectSettings.h
@@ -30,7 +30,7 @@ public:
 	UPROPERTY(EditAnywhere, config, Category = Interchange)
 	TMap<FName, FInterchangePipelineStack> PipelineStacks;
 
-	/** This tell interchange which pipeline to select when importing.*/
+	/** This tell interchange which pipeline to select when importing assets.*/
 	UPROPERTY(EditAnywhere, config, Category = Interchange)
 	FName DefaultPipelineStack;
 
@@ -41,4 +41,12 @@ public:
 	/** If enabled, the pipeline stacks configuration dialog will show every time interchange must choose a pipeline to import or re-import. If disabled interchange will use the DefaultPipelineStack.*/
 	UPROPERTY(EditAnywhere, config, Category = Interchange)
 	bool bShowPipelineStacksConfigurationDialog;
+
+	/** If checked, will use Interchange when importing into level.*/
+	UPROPERTY(EditAnywhere, config, Category = "Interchange (Experimental)")
+	bool bUseInterchangeWhenImportingIntoLevel;
+
+	/** This tell interchange which pipeline to select when importing scenes.*/
+	UPROPERTY(EditAnywhere, config, Category = "Interchange (Experimental)")
+	FName DefaultScenePipelineStack;
 };
\ No newline at end of file
