Changed the way we aquire the GC lock when running interchange. The old implementation could cause some dead lock in the game thread.

diff --git a/Engine/Source/Runtime/Interchange/Engine/Private/InterchangeManager.cpp b/Engine/Source/Runtime/Interchange/Engine/Private/InterchangeManager.cpp
--- a/Engine/Source/Runtime/Interchange/Engine/Private/InterchangeManager.cpp
+++ b/Engine/Source/Runtime/Interchange/Engine/Private/InterchangeManager.cpp
@@ -1046,27 +1046,10 @@ void UInterchangeManager::SetActiveMode(bool IsActive)
 			}
 			return true;
 		});
-
-		//Block GC in a different thread then game thread
-		FString ThreadName = FString(TEXT("InterchangeGCGuard"));
-		GcGuardThread = FThread(*ThreadName, [this]()
-		{
-			FGCScopeGuard GCScopeGuard;
-			while (bIsActive && ImportTasks.Num() > 0)
-			{
-				FPlatformProcess::Sleep(0.01f);
-			}
-		});
 	}
 	else
 	{
 		FTSTicker::GetCoreTicker().RemoveTicker(NotificationTickHandle);
 		NotificationTickHandle.Reset();
-
-		if (GcGuardThread.IsJoinable())
-		{
-			//Finish the thread
-			GcGuardThread.Join();
-		}
 	}
 }
\ No newline at end of file
diff --git a/Engine/Source/Runtime/Interchange/Engine/Private/Tasks/InterchangeTaskCompletion.cpp b/Engine/Source/Runtime/Interchange/Engine/Private/Tasks/InterchangeTaskCompletion.cpp
--- a/Engine/Source/Runtime/Interchange/Engine/Private/Tasks/InterchangeTaskCompletion.cpp
+++ b/Engine/Source/Runtime/Interchange/Engine/Private/Tasks/InterchangeTaskCompletion.cpp
@@ -12,6 +12,7 @@
 #include "InterchangeResultsContainer.h"
 #include "Stats/Stats.h"
 #include "Templates/SharedPointer.h"
+#include "UObject/GarbageCollection.h"
 #include "UObject/Object.h"
 #include "UObject/ObjectMacros.h"
 #include "UObject/WeakObjectPtrTemplates.h"
@@ -21,6 +22,8 @@ void UE::Interchange::FTaskPreAsyncCompletion::DoTask(ENamedThreads::Type Curren
 #if INTERCHANGE_TRACE_ASYNCHRONOUS_TASK_ENABLED
 	INTERCHANGE_TRACE_ASYNCHRONOUS_TASK(PreAsyncCompletion)
 #endif
+	FGCScopeGuard GCScopeGuard;
+
 	TSharedPtr<FImportAsyncHelper, ESPMode::ThreadSafe> AsyncHelper = WeakAsyncHelper.Pin();
 	check(AsyncHelper.IsValid());
 
diff --git a/Engine/Source/Runtime/Interchange/Engine/Private/Tasks/InterchangeTaskCreateSceneObjects.cpp b/Engine/Source/Runtime/Interchange/Engine/Private/Tasks/InterchangeTaskCreateSceneObjects.cpp
--- a/Engine/Source/Runtime/Interchange/Engine/Private/Tasks/InterchangeTaskCreateSceneObjects.cpp
+++ b/Engine/Source/Runtime/Interchange/Engine/Private/Tasks/InterchangeTaskCreateSceneObjects.cpp
@@ -15,9 +15,10 @@
 #include "PackageUtils/PackageUtils.h"
 #include "Stats/Stats.h"
 #include "Templates/SharedPointer.h"
+#include "UObject/GarbageCollection.h"
 #include "UObject/Object.h"
-#include "UObject/UObjectGlobals.h"
 #include "UObject/ObjectMacros.h"
+#include "UObject/UObjectGlobals.h"
 #include "UObject/WeakObjectPtrTemplates.h"
 
 UE::Interchange::FTaskCreateSceneObjects::FTaskCreateSceneObjects(const FString& InPackageBasePath, const int32 InSourceIndex, TWeakPtr<FImportAsyncHelper, ESPMode::ThreadSafe> InAsyncHelper,
@@ -42,6 +43,13 @@ void UE::Interchange::FTaskCreateSceneObjects::DoTask(ENamedThreads::Type Curren
 #if INTERCHANGE_TRACE_ASYNCHRONOUS_TASK_ENABLED
 	INTERCHANGE_TRACE_ASYNCHRONOUS_TASK(SpawnActor)
 #endif
+
+	TOptional<FGCScopeGuard> GCScopeGuard;
+	if (!IsInGameThread())
+	{
+		GCScopeGuard.Emplace();
+	}
+
 	TSharedPtr<UE::Interchange::FImportAsyncHelper, ESPMode::ThreadSafe> AsyncHelper = WeakAsyncHelper.Pin();
 	check(WeakAsyncHelper.IsValid());
 
diff --git a/Engine/Source/Runtime/Interchange/Engine/Private/Tasks/InterchangeTaskParsing.cpp b/Engine/Source/Runtime/Interchange/Engine/Private/Tasks/InterchangeTaskParsing.cpp
--- a/Engine/Source/Runtime/Interchange/Engine/Private/Tasks/InterchangeTaskParsing.cpp
+++ b/Engine/Source/Runtime/Interchange/Engine/Private/Tasks/InterchangeTaskParsing.cpp
@@ -10,8 +10,8 @@
 #include "InterchangeSourceData.h"
 #include "InterchangeTaskCompletion.h"
 #include "InterchangeTaskCreateAsset.h"
-#include "InterchangeTaskPipeline.h"
 #include "InterchangeTaskCreateSceneObjects.h"
+#include "InterchangeTaskPipeline.h"
 #include "InterchangeTranslatorBase.h"
 #include "Misc/Paths.h"
 #include "Nodes/InterchangeBaseNode.h"
@@ -19,9 +19,10 @@
 #include "PackageUtils/PackageUtils.h"
 #include "Stats/Stats.h"
 #include "Templates/SharedPointer.h"
+#include "UObject/GarbageCollection.h"
 #include "UObject/Object.h"
-#include "UObject/UObjectGlobals.h"
 #include "UObject/ObjectMacros.h"
+#include "UObject/UObjectGlobals.h"
 #include "UObject/WeakObjectPtrTemplates.h"
 
 
@@ -30,6 +31,8 @@ void UE::Interchange::FTaskParsing::DoTask(ENamedThreads::Type CurrentThread, co
 #if INTERCHANGE_TRACE_ASYNCHRONOUS_TASK_ENABLED
 	INTERCHANGE_TRACE_ASYNCHRONOUS_TASK(ParsingGraph)
 #endif
+	FGCScopeGuard GCScopeGuard;
+
 	TSharedPtr<FImportAsyncHelper, ESPMode::ThreadSafe> AsyncHelper = WeakAsyncHelper.Pin();
 	check(AsyncHelper.IsValid());
 
diff --git a/Engine/Source/Runtime/Interchange/Engine/Private/Tasks/InterchangeTaskPipeline.cpp b/Engine/Source/Runtime/Interchange/Engine/Private/Tasks/InterchangeTaskPipeline.cpp
--- a/Engine/Source/Runtime/Interchange/Engine/Private/Tasks/InterchangeTaskPipeline.cpp
+++ b/Engine/Source/Runtime/Interchange/Engine/Private/Tasks/InterchangeTaskPipeline.cpp
@@ -6,10 +6,11 @@
 #include "InterchangeEngineLogPrivate.h"
 #include "InterchangeManager.h"
 #include "InterchangePipelineBase.h"
+#include "Nodes/InterchangeBaseNodeContainer.h"
 #include "Stats/Stats.h"
 #include "Templates/SharedPointer.h"
+#include "UObject/GarbageCollection.h"
 #include "UObject/WeakObjectPtrTemplates.h"
-#include "Nodes/InterchangeBaseNodeContainer.h"
 
 
 
@@ -18,6 +19,13 @@ void UE::Interchange::FTaskPipelinePreImport::DoTask(ENamedThreads::Type Current
 #if INTERCHANGE_TRACE_ASYNCHRONOUS_TASK_ENABLED
 	INTERCHANGE_TRACE_ASYNCHRONOUS_TASK(PipelinePreImport)
 #endif
+
+	TOptional<FGCScopeGuard> GCScopeGuard;
+	if (!IsInGameThread())
+	{
+		GCScopeGuard.Emplace();
+	}
+
 	TSharedPtr<FImportAsyncHelper, ESPMode::ThreadSafe> AsyncHelper = WeakAsyncHelper.Pin();
 	check(AsyncHelper.IsValid());
 
@@ -46,6 +54,13 @@ void UE::Interchange::FTaskPipelinePostImport::DoTask(ENamedThreads::Type Curren
 #if INTERCHANGE_TRACE_ASYNCHRONOUS_TASK_ENABLED
 	INTERCHANGE_TRACE_ASYNCHRONOUS_TASK(PipelinePostImport)
 #endif
+
+	TOptional<FGCScopeGuard> GCScopeGuard;
+	if (!IsInGameThread())
+	{
+		GCScopeGuard.Emplace();
+	}
+
 	TSharedPtr<FImportAsyncHelper, ESPMode::ThreadSafe> AsyncHelper = WeakAsyncHelper.Pin();
 	if (!ensure(AsyncHelper.IsValid()) || AsyncHelper->bCancel)
 	{
diff --git a/Engine/Source/Runtime/Interchange/Engine/Private/Tasks/InterchangeTaskTranslator.cpp b/Engine/Source/Runtime/Interchange/Engine/Private/Tasks/InterchangeTaskTranslator.cpp
--- a/Engine/Source/Runtime/Interchange/Engine/Private/Tasks/InterchangeTaskTranslator.cpp
+++ b/Engine/Source/Runtime/Interchange/Engine/Private/Tasks/InterchangeTaskTranslator.cpp
@@ -7,16 +7,19 @@
 #include "InterchangeManager.h"
 #include "InterchangeSourceData.h"
 #include "InterchangeTranslatorBase.h"
+#include "Nodes/InterchangeBaseNodeContainer.h"
 #include "Stats/Stats.h"
 #include "Templates/SharedPointer.h"
+#include "UObject/GarbageCollection.h"
 #include "UObject/WeakObjectPtrTemplates.h"
-#include "Nodes/InterchangeBaseNodeContainer.h"
 
 void UE::Interchange::FTaskTranslator::DoTask(ENamedThreads::Type CurrentThread, const FGraphEventRef& MyCompletionGraphEvent)
 {
 #if INTERCHANGE_TRACE_ASYNCHRONOUS_TASK_ENABLED
 	INTERCHANGE_TRACE_ASYNCHRONOUS_TASK(Translator)
 #endif
+	FGCScopeGuard GCScopeGuard;
+
 	TSharedPtr<FImportAsyncHelper, ESPMode::ThreadSafe> AsyncHelper = WeakAsyncHelper.Pin();
 	check(AsyncHelper.IsValid());
 
diff --git a/Engine/Source/Runtime/Interchange/Engine/Public/InterchangeManager.h b/Engine/Source/Runtime/Interchange/Engine/Public/InterchangeManager.h
--- a/Engine/Source/Runtime/Interchange/Engine/Public/InterchangeManager.h
+++ b/Engine/Source/Runtime/Interchange/Engine/Public/InterchangeManager.h
@@ -484,7 +484,6 @@ private:
 
 	//If interchange is currently importing we have a timer to watch the cancel and we block GC 
 	FThreadSafeBool bIsActive = false;
-	FThread GcGuardThread;
 
 	//If the user want to use the same import pipeline stack for all the queue task
 	//This boolean is reset to false when the ImportTasks array is empty.
