Allow better scripting of interchange for the automation by adding some callback when a asset or a scene object is created and when the import is completed.

Updated the PinHiddenByDefault uproperty metadata to work in blueprint struct ouside of animation blueprint.

diff --git a/Engine/Source/Developer/AssetTools/Private/AssetTools.cpp b/Engine/Source/Developer/AssetTools/Private/AssetTools.cpp
--- a/Engine/Source/Developer/AssetTools/Private/AssetTools.cpp
+++ b/Engine/Source/Developer/AssetTools/Private/AssetTools.cpp
@@ -1915,11 +1915,10 @@ TArray<UObject*> UAssetToolsImpl::ImportAssetsInternal(const TArray<FString>& Fi
 	{
 		for (int32 FileIdx = 0; FileIdx < FilesAndDestinations.Num(); ++FileIdx)
 		{
-			// Filename and DestinationPath will need to get santized before we create an asset out of them as they
+			// Filename will need to get sanitized before we create an asset out of them as they
 			// can be created out of sources that contain spaces and other invalid characters. Filename cannot be sanitized
 			// until other checks are done that rely on looking at the actual source file so sanitation is delayed.
 			const FString& Filename = FilesAndDestinations[FileIdx].Key;
-			const FString DestinationPath = ObjectTools::SanitizeObjectPath(FilesAndDestinations[FileIdx].Value);
 			{
 				UE::Interchange::FScopedSourceData ScopedSourceData(Filename);
 				if (!InterchangeManager.CanTranslateSourceData(ScopedSourceData.GetSourceData()))
diff --git a/Engine/Source/Editor/BlueprintGraph/Classes/K2Node_StructOperation.h b/Engine/Source/Editor/BlueprintGraph/Classes/K2Node_StructOperation.h
--- a/Engine/Source/Editor/BlueprintGraph/Classes/K2Node_StructOperation.h
+++ b/Engine/Source/Editor/BlueprintGraph/Classes/K2Node_StructOperation.h
@@ -39,8 +39,18 @@ protected:
 		virtual void GetRecordDefaults(FProperty* TestProperty, FOptionalPinFromProperty& Record) const override
 		{
 			Record.bCanToggleVisibility = true;
-			UStruct* OwnerStruct = TestProperty ? TestProperty->GetOwnerStruct() : nullptr;
-			Record.bShowPin = OwnerStruct ? !OwnerStruct->HasMetaData(TEXT("HiddenByDefault")) : true;
+			Record.bShowPin = true;
+			if (TestProperty)
+			{
+				Record.bShowPin = !TestProperty->HasMetaData(TEXT("PinHiddenByDefault"));
+				if (Record.bShowPin)
+				{
+					if (UStruct* OwnerStruct = TestProperty->GetOwnerStruct())
+					{
+						Record.bShowPin = !OwnerStruct->HasMetaData(TEXT("HiddenByDefault"));
+					}
+				}
+			}
 		}
 
 		virtual void CustomizePinData(UEdGraphPin* Pin, FName SourcePropertyName, int32 ArrayIndex, FProperty* Property) const override;
diff --git a/Engine/Source/Runtime/CoreUObject/Public/UObject/ObjectMacros.h b/Engine/Source/Runtime/CoreUObject/Public/UObject/ObjectMacros.h
--- a/Engine/Source/Runtime/CoreUObject/Public/UObject/ObjectMacros.h
+++ b/Engine/Source/Runtime/CoreUObject/Public/UObject/ObjectMacros.h
@@ -1336,6 +1336,9 @@ namespace UM
 		/// UFUNCTION()
 		/// TArray<FString> FuncName() const; // Always return string array even if FName property.
 		GetOptions,
+
+		/// [PropertyMetadata] The property can be exposed as a data pin, but is hidden by default.
+		PinHiddenByDefault,
 	};
 
 	// Metadata usable in UPROPERTY for customizing the behavior of Persona and UMG
@@ -1345,16 +1348,13 @@ namespace UM
 		/// [PropertyMetadata] The property is not exposed as a data pin and is only be editable in the details panel. Applicable only to properties that will be displayed in Persona and UMG.
 		NeverAsPin, 
 
-		/// [PropertyMetadata] The property can be exposed as a data pin, but is hidden by default. Applicable only to properties that will be displayed in Persona and UMG.
-		PinHiddenByDefault, 
-
 		/// [PropertyMetadata] The property can be exposed as a data pin and is visible by default. Applicable only to properties that will be displayed in Persona and UMG.
 		PinShownByDefault, 
 
 		/// [PropertyMetadata] The property is always exposed as a data pin. Applicable only to properties that will be displayed in Persona and UMG.
 		AlwaysAsPin, 
 
-		/// [PropertyMetadata] Indicates that the property has custom code to display and should not generate a standard property widget int he details panel. Applicable only to properties that will be displayed in Persona.
+		/// [PropertyMetadata] Indicates that the property has custom code to display and should not generate a standard property widget in the details panel. Applicable only to properties that will be displayed in Persona.
 		CustomizeProperty,
 	};
 
diff --git a/Engine/Source/Runtime/Interchange/Engine/Private/InterchangeManager.cpp b/Engine/Source/Runtime/Interchange/Engine/Private/InterchangeManager.cpp
--- a/Engine/Source/Runtime/Interchange/Engine/Private/InterchangeManager.cpp
+++ b/Engine/Source/Runtime/Interchange/Engine/Private/InterchangeManager.cpp
@@ -229,6 +229,42 @@ void UE::Interchange::FImportResult::SetDone()
 			DoneCallback(*this);
 		}
 
+		TArray<UObject*> Objects = GetImportedObjects();
+
+		if (IsInGameThread())
+		{
+			OnImportDoneNative.ExecuteIfBound(Objects);
+			OnImportDone.ExecuteIfBound(Objects);
+		}
+		else
+		{
+			TArray<TWeakObjectPtr<UObject>> WeakObjects;
+			WeakObjects.Reserve(Objects.Num());
+			for (UObject* Object : Objects)
+			{
+				WeakObjects.Emplace(Object);
+			}
+
+			// call the callbacks on the game thread
+			Async(EAsyncExecution::TaskGraphMainThread, [InWeakObjects = MoveTemp(WeakObjects), ImportDoneNative = OnImportDoneNative, ImportDone = OnImportDone]()
+			{
+				TArray<UObject*> ValidObjects;
+				ValidObjects.Reserve(InWeakObjects.Num());
+
+				for (const TWeakObjectPtr<UObject>& WeakObject : InWeakObjects)
+				{
+					if (UObject* ValidObject = WeakObject.Get())
+					{
+						ValidObjects.Add(ValidObject);
+					}
+				}
+
+				ImportDoneNative.ExecuteIfBound(ValidObjects);
+				ImportDone.ExecuteIfBound(ValidObjects);
+			});
+		}
+
+
 		GraphEvent->DispatchSubsequents();
 	}
 }
@@ -266,8 +302,28 @@ UObject* UE::Interchange::FImportResult::GetFirstAssetOfClass(UClass* InClass) c
 
 void UE::Interchange::FImportResult::AddImportedObject(UObject* ImportedObject)
 {
-	FWriteScopeLock WriteScopeLock(ImportedObjectsRWLock);
-	ImportedObjects.Add(ImportedObject);
+	{
+		FWriteScopeLock WriteScopeLock(ImportedObjectsRWLock);
+		ImportedObjects.Add(ImportedObject);
+	}
+
+	if (IsInGameThread())
+	{
+		OnObjectDoneNative.ExecuteIfBound(ImportedObject);
+		OnObjectDone.ExecuteIfBound(ImportedObject);
+	}
+	else
+	{
+		// call the callbacks on the game thread
+		Async(EAsyncExecution::TaskGraphMainThread, [WeakImportedObject = TWeakObjectPtr<UObject>(ImportedObject), ObjectDoneNative = OnObjectDoneNative, ObjectDone = OnObjectDone] ()
+			{
+				if (UObject* ImportedObjectInGameThread = WeakImportedObject.Get())
+				{
+					ObjectDoneNative.ExecuteIfBound(ImportedObjectInGameThread);
+					ObjectDone.ExecuteIfBound(ImportedObjectInGameThread);
+				}
+			});
+	}
 }
 
 void UE::Interchange::FImportResult::OnDone(TFunction< void(FImportResult&) > Callback)
@@ -566,8 +622,7 @@ UInterchangeManager::ImportInternal(const FString& ContentPath, const UInterchan
 	TaskData.bIsAutomated = ImportAssetParameters.bIsAutomated;
 	TaskData.ImportType = ImportType;
 	TaskData.ReimportObject = ImportAssetParameters.ReimportAsset;
-	TSharedPtr<UE::Interchange::FImportAsyncHelper, ESPMode::ThreadSafe> AsyncHelper = CreateAsyncHelper(TaskData);
-	check(AsyncHelper.IsValid());
+	TSharedRef<UE::Interchange::FImportAsyncHelper, ESPMode::ThreadSafe> AsyncHelper = CreateAsyncHelper(TaskData, ImportAssetParameters);
 
 	//Create a duplicate of the source data, we need to be multithread safe so we copy it to control the life cycle. The async helper will hold it and delete it when the import task will be completed.
 	UInterchangeSourceData* DuplicateSourceData = Cast<UInterchangeSourceData>(StaticDuplicateObject(SourceData, GetTransientPackage()));
@@ -783,12 +838,22 @@ const UClass* UInterchangeManager::GetRegisteredFactoryClass(const UClass* Class
 	return Result;
 }
 
-TSharedPtr<UE::Interchange::FImportAsyncHelper, ESPMode::ThreadSafe> UInterchangeManager::CreateAsyncHelper(const UE::Interchange::FImportAsyncHelperData& Data)
+TSharedRef<UE::Interchange::FImportAsyncHelper, ESPMode::ThreadSafe> UInterchangeManager::CreateAsyncHelper(const UE::Interchange::FImportAsyncHelperData& Data, const FImportAssetParameters& ImportAssetParameters)
 {
-	TSharedPtr<UE::Interchange::FImportAsyncHelper, ESPMode::ThreadSafe> AsyncHelper = MakeShared<UE::Interchange::FImportAsyncHelper, ESPMode::ThreadSafe>();
+	TSharedRef<UE::Interchange::FImportAsyncHelper, ESPMode::ThreadSafe> AsyncHelper = MakeShared<UE::Interchange::FImportAsyncHelper, ESPMode::ThreadSafe>();
 	//Copy the task data
 	AsyncHelper->TaskData = Data;
 	
+	AsyncHelper->AssetImportResult->OnObjectDone = ImportAssetParameters.OnAssetDone;
+	AsyncHelper->AssetImportResult->OnObjectDoneNative = ImportAssetParameters.OnAssetDoneNative;
+	AsyncHelper->AssetImportResult->OnImportDone = ImportAssetParameters.OnAssetsImportDone;
+	AsyncHelper->AssetImportResult->OnImportDoneNative = ImportAssetParameters.OnAssetsImportDoneNative;
+
+	AsyncHelper->SceneImportResult->OnObjectDone = ImportAssetParameters.OnSceneObjectDone;
+	AsyncHelper->SceneImportResult->OnObjectDoneNative = ImportAssetParameters.OnSceneObjectDoneNative;
+	AsyncHelper->SceneImportResult->OnImportDone = ImportAssetParameters.OnSceneImportDone;
+	AsyncHelper->SceneImportResult->OnImportDoneNative = ImportAssetParameters.OnSceneImportDoneNative;
+
 	AsyncHelper->AssetImportResult->SetInProgress();
 
 	return AsyncHelper;
diff --git a/Engine/Source/Runtime/Interchange/Engine/Public/InterchangeManager.h b/Engine/Source/Runtime/Interchange/Engine/Public/InterchangeManager.h
--- a/Engine/Source/Runtime/Interchange/Engine/Public/InterchangeManager.h
+++ b/Engine/Source/Runtime/Interchange/Engine/Public/InterchangeManager.h
@@ -30,6 +30,13 @@ class FAsyncTaskNotification;
 class UInterchangeFactoryBase;
 class UInterchangePipelineBase;
 
+/** Some utilities delegates for the automation of interchange */
+DECLARE_DYNAMIC_DELEGATE_OneParam(FOnObjectImportDoneDynamic, UObject*, Object);
+DECLARE_DELEGATE_OneParam(FOnObjectImportDoneNative, UObject*);
+
+DECLARE_DYNAMIC_DELEGATE_OneParam(FOnImportDoneDynamic, const TArray<UObject*>&, Objects);
+DECLARE_DELEGATE_OneParam(FOnImportDoneNative, const TArray<UObject*>&);
+
 namespace UE
 {
 	namespace Interchange
@@ -118,6 +125,13 @@ namespace UE
 			// Callback when the status switches to done.
 			void OnDone(TFunction< void(FImportResult&) > Callback);
 
+			// Internal delegates (use the FImportAssetParameters when calling the interchange import functions to set those)
+			FOnObjectImportDoneDynamic OnObjectDone;
+			FOnObjectImportDoneNative OnObjectDoneNative;
+
+			FOnImportDoneDynamic OnImportDone;
+			FOnImportDoneNative OnImportDoneNative;
+
 		protected:
 			/* FGCObject interface */
 			virtual void AddReferencedObjects(FReferenceCollector& Collector) override;
@@ -219,14 +233,39 @@ struct FImportAssetParameters
 {
 	GENERATED_USTRUCT_BODY()
 
+	// If the import is a reimport for a specific asset set the asset to reimport here
 	UPROPERTY(EditAnywhere, BlueprintReadWrite, Category = "Interchange|ImportAsset")
 	TObjectPtr<UObject> ReimportAsset = nullptr;
 
+	// Tell interchange that import is automated and it shouldn't present a model window
 	UPROPERTY(EditAnywhere, BlueprintReadWrite, Category = "Interchange|ImportAsset")
 	bool bIsAutomated = false;
 
+	// Adding some override will tell interchange to use the specific custom set pipelines instead of letting the user or the system chose
 	UPROPERTY(EditAnywhere, BlueprintReadWrite, Category = "Interchange|ImportAsset")
 	TArray<TObjectPtr<UInterchangePipelineBase>> OverridePipelines;
+
+	/* Delegates used track the imported objects */
+
+	// This is called each time an asset is imported or reimported from the import call
+	UPROPERTY(EditAnywhere, BlueprintReadWrite, Category = "Interchange|ImportAsset", meta=(PinHiddenByDefault))
+	FOnObjectImportDoneDynamic OnAssetDone;
+	FOnObjectImportDoneNative OnAssetDoneNative;
+
+	// This is called when all the assets where imported from the source data
+	UPROPERTY(EditAnywhere, BlueprintReadWrite, Category = "Interchange|ImportAsset", meta=(PinHiddenByDefault))
+	FOnImportDoneDynamic OnAssetsImportDone;
+	FOnImportDoneNative OnAssetsImportDoneNative;
+
+	// This is called each time an object in the scene is imported or reimported from the import call
+	UPROPERTY(EditAnywhere, BlueprintReadWrite, Category = "Interchange|ImportAsset", meta=(PinHiddenByDefault))
+	FOnObjectImportDoneDynamic OnSceneObjectDone;
+	FOnObjectImportDoneNative OnSceneObjectDoneNative;
+
+	// This is called when all the scene object where imported from the source data
+	UPROPERTY(EditAnywhere, BlueprintReadWrite, Category = "Interchange|ImportAsset", meta=(PinHiddenByDefault))
+	FOnImportDoneDynamic OnSceneImportDone;
+	FOnImportDoneNative OnSceneImportDoneNative;
 };
 
 UCLASS(Transient, BlueprintType)
@@ -354,7 +393,7 @@ public:
 	* @return: A new UInterchangeSourceData.
 	*/
 	UFUNCTION(BlueprintCallable, Category = "Interchange | Import Manager")
-	UInterchangeSourceData* CreateSourceData(const FString& InFileName);
+	static UInterchangeSourceData* CreateSourceData(const FString& InFileName);
 
 	/**
 	* Script helper to get a registered factory for a specified class
@@ -368,7 +407,7 @@ public:
 	 * Return an FImportAsynHelper pointer. The pointer is deleted when ReleaseAsyncHelper is call.
 	 * @param Data - The data we want to pass to the different import tasks
 	 */
-	TSharedPtr<UE::Interchange::FImportAsyncHelper, ESPMode::ThreadSafe> CreateAsyncHelper(const UE::Interchange::FImportAsyncHelperData& Data);
+	TSharedRef<UE::Interchange::FImportAsyncHelper, ESPMode::ThreadSafe> CreateAsyncHelper(const UE::Interchange::FImportAsyncHelperData& Data, const FImportAssetParameters& ImportAssetParameters);
 
 	/** Delete the specified AsyncHelper and remove it from the array that was holding it. */
 	void ReleaseAsyncHelper(TWeakPtr<UE::Interchange::FImportAsyncHelper, ESPMode::ThreadSafe> AsyncHelper);
