Interchange scene import
- Support different types of lights.
- Set the scene components mobility in the pipeline.
- Made a generic macro to apply an attribute to a FProperty where the attribute and property names can be different. Supports nested script structs.
- Renamed ApplyAllCustomAttributeToAsset to ApplyAllCustomAttributeToObject.
- Added support for Camera physical attributes (focal length, sensor width & height) from FBX.

diff --git a/Engine/Plugins/Runtime/Interchange/Source/Import/InterchangeImport.Build.cs b/Engine/Plugins/Runtime/Interchange/Source/Import/InterchangeImport.Build.cs
--- a/Engine/Plugins/Runtime/Interchange/Source/Import/InterchangeImport.Build.cs
+++ b/Engine/Plugins/Runtime/Interchange/Source/Import/InterchangeImport.Build.cs
@@ -25,6 +25,7 @@ namespace UnrealBuildTool.Rules
 				new string[]
 				{
 					"AssetRegistry",
+					"CinematicCamera",
 					"IESFile",
 					"ImageWrapper",
 					"InterchangeDispatcher",
diff --git a/Engine/Plugins/Runtime/Interchange/Source/Import/Private/InterchangeImportCommon.cpp b/Engine/Plugins/Runtime/Interchange/Source/Import/Private/InterchangeImportCommon.cpp
--- a/Engine/Plugins/Runtime/Interchange/Source/Import/Private/InterchangeImportCommon.cpp
+++ b/Engine/Plugins/Runtime/Interchange/Source/Import/Private/InterchangeImportCommon.cpp
@@ -199,7 +199,7 @@ namespace UE
 				case EReimportStrategyFlags::ApplyPipelineProperties:
 				{
 					//Directly apply pipeline node attribute to the asset
-					PipelineAssetNode->ApplyAllCustomAttributeToAsset(Asset);
+					PipelineAssetNode->ApplyAllCustomAttributeToObject(Asset);
 					break;
 				}
 				
@@ -213,7 +213,7 @@ namespace UE
 					//set all ModifedAttributes from the CurrentAssetNode to the pipeline node. This will put back all user changes
 					UInterchangeBaseNode::CopyStorageAttributes(CurrentAssetNode, PipelineAssetNode, ModifiedAttributes);
 					//Now apply the pipeline node attribute to the asset
-					PipelineAssetNode->ApplyAllCustomAttributeToAsset(Asset);
+					PipelineAssetNode->ApplyAllCustomAttributeToObject(Asset);
 					break;
 				}
 			}
diff --git a/Engine/Plugins/Runtime/Interchange/Source/Import/Private/Material/InterchangeMaterialFactory.cpp b/Engine/Plugins/Runtime/Interchange/Source/Import/Private/Material/InterchangeMaterialFactory.cpp
--- a/Engine/Plugins/Runtime/Interchange/Source/Import/Private/Material/InterchangeMaterialFactory.cpp
+++ b/Engine/Plugins/Runtime/Interchange/Source/Import/Private/Material/InterchangeMaterialFactory.cpp
@@ -281,7 +281,7 @@ UObject* UInterchangeMaterialFactory::CreateAsset(const FCreateAssetParams& Argu
 			}
 
 			/** Apply all MaterialNode custom attributes to the material asset */
-			MaterialFactoryNode->ApplyAllCustomAttributeToAsset(Material);
+			MaterialFactoryNode->ApplyAllCustomAttributeToObject(Material);
 
 			//TODO support material instance
 		}
diff --git a/Engine/Plugins/Runtime/Interchange/Source/Import/Private/Mesh/InterchangeSkeletalMeshFactory.cpp b/Engine/Plugins/Runtime/Interchange/Source/Import/Private/Mesh/InterchangeSkeletalMeshFactory.cpp
--- a/Engine/Plugins/Runtime/Interchange/Source/Import/Private/Mesh/InterchangeSkeletalMeshFactory.cpp
+++ b/Engine/Plugins/Runtime/Interchange/Source/Import/Private/Mesh/InterchangeSkeletalMeshFactory.cpp
@@ -1019,8 +1019,8 @@ UObject* UInterchangeSkeletalMeshFactory::CreateAsset(const FCreateAssetParams&
 
 			SkeletalMesh->CalculateInvRefMatrices();
 
-			/** Apply all SkeletalMeshFactoryNode custom attributes to the material asset */
-			SkeletalMeshFactoryNode->ApplyAllCustomAttributeToAsset(SkeletalMesh);
+			/** Apply all SkeletalMeshFactoryNode custom attributes to the skeletal mesh asset */
+			SkeletalMeshFactoryNode->ApplyAllCustomAttributeToObject(SkeletalMesh);
 
 			bool bCreatePhysicsAsset = false;
 			SkeletalMeshFactoryNode->GetCustomCreatePhysicsAsset(bCreatePhysicsAsset);
diff --git a/Engine/Plugins/Runtime/Interchange/Source/Import/Private/Mesh/InterchangeStaticMeshFactory.cpp b/Engine/Plugins/Runtime/Interchange/Source/Import/Private/Mesh/InterchangeStaticMeshFactory.cpp
--- a/Engine/Plugins/Runtime/Interchange/Source/Import/Private/Mesh/InterchangeStaticMeshFactory.cpp
+++ b/Engine/Plugins/Runtime/Interchange/Source/Import/Private/Mesh/InterchangeStaticMeshFactory.cpp
@@ -364,8 +364,8 @@ UObject* UInterchangeStaticMeshFactory::CreateAsset(const FCreateAssetParams& Ar
 				CurrentLodIndex++;
 			}
 
-			// Apply all SkeletalMeshFactoryNode custom attributes to the material asset
-			StaticMeshFactoryNode->ApplyAllCustomAttributeToAsset(StaticMesh);
+			// Apply all StaticMeshFactoryNode custom attributes to the static mesh asset
+			StaticMeshFactoryNode->ApplyAllCustomAttributeToObject(StaticMesh);
 		}
 		
 		// Getting the file Hash will cache it into the source data
diff --git a/Engine/Plugins/Runtime/Interchange/Source/Import/Private/Scene/InterchangeActorFactory.cpp b/Engine/Plugins/Runtime/Interchange/Source/Import/Private/Scene/InterchangeActorFactory.cpp
--- a/Engine/Plugins/Runtime/Interchange/Source/Import/Private/Scene/InterchangeActorFactory.cpp
+++ b/Engine/Plugins/Runtime/Interchange/Source/Import/Private/Scene/InterchangeActorFactory.cpp
@@ -6,6 +6,8 @@
 #include "InterchangeImportLog.h"
 #include "InterchangeSourceData.h"
 
+#include "CineCameraActor.h"
+#include "CineCameraComponent.h"
 #include "Components/SceneComponent.h"
 #include "Components/SkeletalMeshComponent.h"
 #include "Engine/Engine.h"
@@ -140,7 +142,6 @@ namespace UE
 					if (!SpawnedActor->GetRootComponent())
 					{
 						USceneComponent* RootComponent = NewObject<USceneComponent>(SpawnedActor, USceneComponent::GetDefaultSceneRootVariableName(), RF_Transactional);
-						RootComponent->Mobility = EComponentMobility::Static; //todo: this should be an attribute on a factory node
 #if WITH_EDITORONLY_DATA
 						RootComponent->bVisualizeComponent = true;
 #endif
@@ -148,8 +149,23 @@ namespace UE
 
 						SpawnedActor->SetRootComponent(RootComponent);
 						SpawnedActor->AddInstanceComponent(RootComponent);
+					}
+
+					if (USceneComponent* RootComponent = SpawnedActor->GetRootComponent())
+					{
+						uint8 Mobility;
+						if (ActorFactoryNode->GetCustomMobility(Mobility))
+						{
+							//Make sure we don't have a mobility that's more restrictive than our parent mobility, as that wouldn't be a valid setup.
+							EComponentMobility::Type TargetMobility = (EComponentMobility::Type)Mobility;
+
+							if (ParentActor && ParentActor->GetRootComponent())
+							{
+								TargetMobility = (EComponentMobility::Type)FMath::Max((uint8)Mobility, (uint8)ParentActor->GetRootComponent()->Mobility);
+							}
 
-						RootComponent->RegisterComponent();
+							RootComponent->SetMobility(TargetMobility);
+						}
 					}
 				}
 
@@ -188,6 +204,15 @@ namespace UE
 					SetupSkeletalMeshComponent(NodeContainer, ActorNode, *SkeletalMeshComponent);
 				}
 
+				if (ACineCameraActor* CineCameraActor = Cast<ACineCameraActor>(SpawnedActor))
+				{
+					ActorNode->ApplyAllCustomAttributeToObject(CineCameraActor->GetCineCameraComponent());
+				}
+				else
+				{
+					ActorNode->ApplyAllCustomAttributeToObject(SpawnedActor->GetRootComponent());
+				}
+
 				OutSpawnedActors.Add(ActorNodeUniqueID, SpawnedActor);
 
 				if (bSpawnChildren)
diff --git a/Engine/Plugins/Runtime/Interchange/Source/Import/Private/Texture/InterchangeTextureFactory.cpp b/Engine/Plugins/Runtime/Interchange/Source/Import/Private/Texture/InterchangeTextureFactory.cpp
--- a/Engine/Plugins/Runtime/Interchange/Source/Import/Private/Texture/InterchangeTextureFactory.cpp
+++ b/Engine/Plugins/Runtime/Interchange/Source/Import/Private/Texture/InterchangeTextureFactory.cpp
@@ -968,7 +968,7 @@ void UInterchangeTextureFactory::PreImportPreCompletedCallback(const FImportPreC
 		if (!Arguments.bIsReimport)
 		{
 			/** Apply all TextureNode custom attributes to the texture asset */
-			TextureFactoryNode->ApplyAllCustomAttributeToAsset(Texture);
+			TextureFactoryNode->ApplyAllCustomAttributeToObject(Texture);
 		}
 		else
 		{
@@ -981,7 +981,7 @@ void UInterchangeTextureFactory::PreImportPreCompletedCallback(const FImportPreC
 
 			UInterchangeBaseNode* CurrentNode = NewObject<UInterchangeBaseNode>(GetTransientPackage(), GetSupportedFactoryNodeClass(TextureFactoryNode));
 			UInterchangeBaseNode::CopyStorage(TextureFactoryNode, CurrentNode);
-			CurrentNode->FillAllCustomAttributeFromAsset(Texture);
+			CurrentNode->FillAllCustomAttributeFromObject(Texture);
 			//Apply reimport strategy
 			UE::Interchange::FFactoryCommon::ApplyReimportStrategyToAsset(Arguments.ReimportStrategyFlags, Texture, PreviousNode, CurrentNode, TextureFactoryNode);
 		}
diff --git a/Engine/Plugins/Runtime/Interchange/Source/Nodes/InterchangeNodes.Build.cs b/Engine/Plugins/Runtime/Interchange/Source/Nodes/InterchangeNodes.Build.cs
--- a/Engine/Plugins/Runtime/Interchange/Source/Nodes/InterchangeNodes.Build.cs
+++ b/Engine/Plugins/Runtime/Interchange/Source/Nodes/InterchangeNodes.Build.cs
@@ -17,7 +17,13 @@ namespace UnrealBuildTool.Rules
 
 			if(Target.bCompileAgainstEngine)
 			{
-				PublicDependencyModuleNames.Add("Engine");
+				PublicDependencyModuleNames.AddRange(
+					new string[]
+					{
+						"CinematicCamera",
+						"Engine"
+					}
+				);
 			}
 		}
 	}
diff --git a/Engine/Plugins/Runtime/Interchange/Source/Nodes/Public/InterchangeActorFactoryNode.h b/Engine/Plugins/Runtime/Interchange/Source/Nodes/Public/InterchangeActorFactoryNode.h
--- a/Engine/Plugins/Runtime/Interchange/Source/Nodes/Public/InterchangeActorFactoryNode.h
+++ b/Engine/Plugins/Runtime/Interchange/Source/Nodes/Public/InterchangeActorFactoryNode.h
@@ -38,7 +38,20 @@ public:
 		IMPLEMENT_NODE_ATTRIBUTE_SETTER_NODELEGATE(ActorClassName, FString);
 	}
 
+	UFUNCTION(BlueprintCallable, Category = "Interchange | Node | ActorFactory")
+	bool GetCustomMobility(uint8& AttributeValue) const
+	{
+		IMPLEMENT_NODE_ATTRIBUTE_GETTER(Mobility, uint8);
+	}
+
+	UFUNCTION(BlueprintCallable, Category = "Interchange | Node | ActorFactory")
+	bool SetCustomMobility(const uint8& AttributeValue, bool bAddApplyDelegate = true)
+	{
+		IMPLEMENT_NODE_ATTRIBUTE_SETTER_NODELEGATE(Mobility, uint8);
+	}
+
 private:
 	const UE::Interchange::FAttributeKey Macro_CustomGlobalTransformKey = UE::Interchange::FAttributeKey(TEXT("GlobalTransform"));
 	const UE::Interchange::FAttributeKey Macro_CustomActorClassNameKey = UE::Interchange::FAttributeKey(TEXT("ActorClassName"));
+	const UE::Interchange::FAttributeKey Macro_CustomMobilityKey = UE::Interchange::FAttributeKey(TEXT("Mobility"));
 };
\ No newline at end of file
diff --git a/Engine/Plugins/Runtime/Interchange/Source/Nodes/Public/InterchangeCameraNode.h b/Engine/Plugins/Runtime/Interchange/Source/Nodes/Public/InterchangeCameraNode.h
--- a/Engine/Plugins/Runtime/Interchange/Source/Nodes/Public/InterchangeCameraNode.h
+++ b/Engine/Plugins/Runtime/Interchange/Source/Nodes/Public/InterchangeCameraNode.h
@@ -26,4 +26,46 @@ public:
 		const FString TypeName = TEXT("CameraNode");
 		return TypeName;
 	}
+
+public:
+	UFUNCTION(BlueprintCallable, Category = "Interchange | Node | Camera")
+	bool GetCustomFocalLength(float& AttributeValue) const
+	{
+		IMPLEMENT_NODE_ATTRIBUTE_GETTER(FocalLength, float);
+	}
+
+	UFUNCTION(BlueprintCallable, Category = "Interchange | Node | Camera")
+	bool SetCustomFocalLength(const float& AttributeValue)
+	{
+		IMPLEMENT_NODE_ATTRIBUTE_SETTER_NODELEGATE(FocalLength, float);
+	}
+
+	UFUNCTION(BlueprintCallable, Category = "Interchange | Node | Camera")
+	bool GetCustomSensorWidth(float& AttributeValue) const
+	{
+		IMPLEMENT_NODE_ATTRIBUTE_GETTER(SensorWidth, float);
+	}
+
+	UFUNCTION(BlueprintCallable, Category = "Interchange | Node | Camera")
+	bool SetCustomSensorWidth(const float& AttributeValue)
+	{
+		IMPLEMENT_NODE_ATTRIBUTE_SETTER_NODELEGATE(SensorWidth, float);
+	}
+
+	UFUNCTION(BlueprintCallable, Category = "Interchange | Node | Camera")
+		bool GetCustomSensorHeight(float& AttributeValue) const
+	{
+		IMPLEMENT_NODE_ATTRIBUTE_GETTER(SensorHeight, float);
+	}
+
+	UFUNCTION(BlueprintCallable, Category = "Interchange | Node | Camera")
+		bool SetCustomSensorHeight(const float& AttributeValue)
+	{
+		IMPLEMENT_NODE_ATTRIBUTE_SETTER_NODELEGATE(SensorHeight, float);
+	}
+
+private:
+	const UE::Interchange::FAttributeKey Macro_CustomFocalLengthKey = UE::Interchange::FAttributeKey(TEXT("FocalLength"));
+	const UE::Interchange::FAttributeKey Macro_CustomSensorWidthKey = UE::Interchange::FAttributeKey(TEXT("SensorWidth"));
+	const UE::Interchange::FAttributeKey Macro_CustomSensorHeightKey = UE::Interchange::FAttributeKey(TEXT("SensorHeight"));
 };
diff --git a/Engine/Plugins/Runtime/Interchange/Source/Nodes/Public/InterchangeCineCameraFactoryNode.h b/Engine/Plugins/Runtime/Interchange/Source/Nodes/Public/InterchangeCineCameraFactoryNode.h
new file mode 100644
--- /dev/null
+++ b/Engine/Plugins/Runtime/Interchange/Source/Nodes/Public/InterchangeCineCameraFactoryNode.h
@@ -0,0 +1,65 @@
+// Copyright Epic Games, Inc. All Rights Reserved.
+
+#pragma once
+
+#include "CoreMinimal.h"
+#include "InterchangeActorFactoryNode.h"
+
+#if WITH_ENGINE
+	#include "CineCameraComponent.h"
+#endif
+
+#include "InterchangeCineCameraFactoryNode.generated.h"
+
+UCLASS(BlueprintType)
+class INTERCHANGENODES_API UInterchangeCineCameraFactoryNode : public UInterchangeActorFactoryNode
+{
+	GENERATED_BODY()
+
+public:
+	UFUNCTION(BlueprintCallable, Category = "Interchange | Node | CameraFactory")
+	bool GetCustomFocalLength(float& AttributeValue) const
+	{
+		IMPLEMENT_NODE_ATTRIBUTE_GETTER(FocalLength, float);
+	}
+
+	UFUNCTION(BlueprintCallable, Category = "Interchange | Node | CameraFactory")
+	bool SetCustomFocalLength(const float& AttributeValue, bool bAddApplyDelegate = true)
+	{
+		IMPLEMENT_NODE_ATTRIBUTE_SETTER(UInterchangeCineCameraFactoryNode, FocalLength, float, UCineCameraComponent);
+	}
+
+	UFUNCTION(BlueprintCallable, Category = "Interchange | Node | CameraFactory")
+	bool GetCustomSensorWidth(float& AttributeValue) const
+	{
+		IMPLEMENT_NODE_ATTRIBUTE_GETTER(SensorWidth, float);
+	}
+
+	UFUNCTION(BlueprintCallable, Category = "Interchange | Node | CameraFactory")
+	bool SetCustomSensorWidth(const float& AttributeValue, bool bAddApplyDelegate = true)
+	{
+		IMPLEMENT_NODE_ATTRIBUTE_SETTER(UInterchangeCineCameraFactoryNode, SensorWidth, float, UCineCameraComponent);
+	}
+
+	UFUNCTION(BlueprintCallable, Category = "Interchange | Node | CameraFactory")
+	bool GetCustomSensorHeight(float& AttributeValue) const
+	{
+		IMPLEMENT_NODE_ATTRIBUTE_GETTER(SensorHeight, float);
+	}
+
+	UFUNCTION(BlueprintCallable, Category = "Interchange | Node | CameraFactory")
+	bool SetCustomSensorHeight(const float& AttributeValue, bool bAddApplyDelegate = true)
+	{
+		IMPLEMENT_NODE_ATTRIBUTE_SETTER(UInterchangeCineCameraFactoryNode, SensorHeight, float, UCineCameraComponent);
+	}
+
+private:
+	const UE::Interchange::FAttributeKey Macro_CustomFocalLengthKey = UE::Interchange::FAttributeKey(TEXT("FocalLength"));
+	const UE::Interchange::FAttributeKey Macro_CustomSensorWidthKey = UE::Interchange::FAttributeKey(TEXT("SensorWidth"));
+	const UE::Interchange::FAttributeKey Macro_CustomSensorHeightKey = UE::Interchange::FAttributeKey(TEXT("SensorHeight"));
+
+private:
+	IMPLEMENT_NODE_ATTRIBUTE_APPLY_UOBJECT_BYNAME(FocalLength, float, UCineCameraComponent, TEXT("CurrentFocalLength"));
+	IMPLEMENT_NODE_ATTRIBUTE_APPLY_UOBJECT_BYNAME(SensorWidth, float, UCineCameraComponent, TEXT("Filmback.SensorWidth"));
+	IMPLEMENT_NODE_ATTRIBUTE_APPLY_UOBJECT_BYNAME(SensorHeight, float, UCineCameraComponent, TEXT("Filmback.SensorHeight"));
+};
diff --git a/Engine/Plugins/Runtime/Interchange/Source/Nodes/Public/InterchangeLightNode.h b/Engine/Plugins/Runtime/Interchange/Source/Nodes/Public/InterchangeLightNode.h
--- a/Engine/Plugins/Runtime/Interchange/Source/Nodes/Public/InterchangeLightNode.h
+++ b/Engine/Plugins/Runtime/Interchange/Source/Nodes/Public/InterchangeLightNode.h
@@ -28,3 +28,39 @@ public:
 	}
 
 };
+
+UCLASS(BlueprintType)
+class INTERCHANGENODES_API UInterchangePointLightNode : public UInterchangeLightNode
+{
+	GENERATED_BODY()
+
+public:
+
+};
+
+UCLASS(BlueprintType)
+class INTERCHANGENODES_API UInterchangeSpotLightNode : public UInterchangePointLightNode
+{
+	GENERATED_BODY()
+
+public:
+
+};
+
+UCLASS(BlueprintType)
+class INTERCHANGENODES_API UInterchangeRectLightNode : public UInterchangeLightNode
+{
+	GENERATED_BODY()
+
+public:
+
+};
+
+UCLASS(BlueprintType)
+class INTERCHANGENODES_API UInterchangeDirectionalLightNode : public UInterchangeLightNode
+{
+	GENERATED_BODY()
+
+public:
+
+};
diff --git a/Engine/Plugins/Runtime/Interchange/Source/Nodes/Public/InterchangeMaterialFactoryNode.h b/Engine/Plugins/Runtime/Interchange/Source/Nodes/Public/InterchangeMaterialFactoryNode.h
--- a/Engine/Plugins/Runtime/Interchange/Source/Nodes/Public/InterchangeMaterialFactoryNode.h
+++ b/Engine/Plugins/Runtime/Interchange/Source/Nodes/Public/InterchangeMaterialFactoryNode.h
@@ -316,7 +316,7 @@ private:
 	}
 #endif
 
-	IMPLEMENT_NODE_ATTRIBUTE_APPLY_UOBJECT(BlendMode, uint8, UMaterial, TEnumAsByte<enum EBlendMode>);
+	IMPLEMENT_NODE_ATTRIBUTE_APPLY_UOBJECT(BlendMode, uint8, UMaterial);
 
 protected:
 #if WITH_ENGINE
diff --git a/Engine/Plugins/Runtime/Interchange/Source/Nodes/Public/InterchangeTexture2DArrayFactoryNode.h b/Engine/Plugins/Runtime/Interchange/Source/Nodes/Public/InterchangeTexture2DArrayFactoryNode.h
--- a/Engine/Plugins/Runtime/Interchange/Source/Nodes/Public/InterchangeTexture2DArrayFactoryNode.h
+++ b/Engine/Plugins/Runtime/Interchange/Source/Nodes/Public/InterchangeTexture2DArrayFactoryNode.h
@@ -74,8 +74,8 @@ private:
 	IMPLEMENT_NODE_ATTRIBUTE_KEY(AddressZ);
 
 #if WITH_EDITORONLY_DATA
-	IMPLEMENT_NODE_ATTRIBUTE_APPLY_UOBJECT(AddressX, uint8, UTexture2DArray, TEnumAsByte<TextureAddress>);
-	IMPLEMENT_NODE_ATTRIBUTE_APPLY_UOBJECT(AddressY, uint8, UTexture2DArray, TEnumAsByte<TextureAddress>);
-	IMPLEMENT_NODE_ATTRIBUTE_APPLY_UOBJECT(AddressZ, uint8, UTexture2DArray, TEnumAsByte<TextureAddress>);
+	IMPLEMENT_NODE_ATTRIBUTE_APPLY_UOBJECT(AddressX, uint8, UTexture2DArray);
+	IMPLEMENT_NODE_ATTRIBUTE_APPLY_UOBJECT(AddressY, uint8, UTexture2DArray);
+	IMPLEMENT_NODE_ATTRIBUTE_APPLY_UOBJECT(AddressZ, uint8, UTexture2DArray);
 #endif
 };
\ No newline at end of file
diff --git a/Engine/Plugins/Runtime/Interchange/Source/Nodes/Public/InterchangeTextureFactoryNode.h b/Engine/Plugins/Runtime/Interchange/Source/Nodes/Public/InterchangeTextureFactoryNode.h
--- a/Engine/Plugins/Runtime/Interchange/Source/Nodes/Public/InterchangeTextureFactoryNode.h
+++ b/Engine/Plugins/Runtime/Interchange/Source/Nodes/Public/InterchangeTextureFactoryNode.h
@@ -805,50 +805,50 @@ private:
 	const UE::Interchange::FAttributeKey Macro_CustomCompositePowerKey = UE::Interchange::FAttributeKey(TEXT("CompositePower"));
 
 #if WITH_EDITORONLY_DATA
-	IMPLEMENT_NODE_ATTRIBUTE_APPLY_UOBJECT(AdjustBrightness, float, UTexture, );
-	IMPLEMENT_NODE_ATTRIBUTE_APPLY_UOBJECT(AdjustBrightnessCurve, float, UTexture, );
-	IMPLEMENT_NODE_ATTRIBUTE_APPLY_UOBJECT(AdjustVibrance, float, UTexture, );
-	IMPLEMENT_NODE_ATTRIBUTE_APPLY_UOBJECT(AdjustSaturation, float, UTexture, );
-	IMPLEMENT_NODE_ATTRIBUTE_APPLY_UOBJECT(AdjustRGBCurve, float, UTexture, );
-	IMPLEMENT_NODE_ATTRIBUTE_APPLY_UOBJECT(AdjustHue, float, UTexture, );
-	IMPLEMENT_NODE_ATTRIBUTE_APPLY_UOBJECT(AdjustMinAlpha, float, UTexture, );
-	IMPLEMENT_NODE_ATTRIBUTE_APPLY_UOBJECT(AdjustMaxAlpha, float, UTexture, );
-	IMPLEMENT_NODE_ATTRIBUTE_APPLY_UOBJECT(bChromaKeyTexture, bool, UTexture, );
-	IMPLEMENT_NODE_ATTRIBUTE_APPLY_UOBJECT(ChromaKeyThreshold, float, UTexture, );
-	IMPLEMENT_NODE_ATTRIBUTE_APPLY_UOBJECT(ChromaKeyColor, FColor, UTexture, );
+	IMPLEMENT_NODE_ATTRIBUTE_APPLY_UOBJECT(AdjustBrightness, float, UTexture);
+	IMPLEMENT_NODE_ATTRIBUTE_APPLY_UOBJECT(AdjustBrightnessCurve, float, UTexture);
+	IMPLEMENT_NODE_ATTRIBUTE_APPLY_UOBJECT(AdjustVibrance, float, UTexture);
+	IMPLEMENT_NODE_ATTRIBUTE_APPLY_UOBJECT(AdjustSaturation, float, UTexture);
+	IMPLEMENT_NODE_ATTRIBUTE_APPLY_UOBJECT(AdjustRGBCurve, float, UTexture);
+	IMPLEMENT_NODE_ATTRIBUTE_APPLY_UOBJECT(AdjustHue, float, UTexture);
+	IMPLEMENT_NODE_ATTRIBUTE_APPLY_UOBJECT(AdjustMinAlpha, float, UTexture);
+	IMPLEMENT_NODE_ATTRIBUTE_APPLY_UOBJECT(AdjustMaxAlpha, float, UTexture);
+	IMPLEMENT_NODE_ATTRIBUTE_APPLY_UOBJECT(bChromaKeyTexture, bool, UTexture);
+	IMPLEMENT_NODE_ATTRIBUTE_APPLY_UOBJECT(ChromaKeyThreshold, float, UTexture);
+	IMPLEMENT_NODE_ATTRIBUTE_APPLY_UOBJECT(ChromaKeyColor, FColor, UTexture);
 
 	//Texture Compression
-	IMPLEMENT_NODE_ATTRIBUTE_APPLY_UOBJECT(CompressionNoAlpha, bool, UTexture, );
-	IMPLEMENT_NODE_ATTRIBUTE_APPLY_UOBJECT(DeferCompression, bool, UTexture, );
-	IMPLEMENT_NODE_ATTRIBUTE_APPLY_UOBJECT(LossyCompressionAmount, uint8, UTexture, TEnumAsByte<ETextureLossyCompressionAmount>);
-	IMPLEMENT_NODE_ATTRIBUTE_APPLY_UOBJECT(MaxTextureSize, int32, UTexture, );
-	IMPLEMENT_NODE_ATTRIBUTE_APPLY_UOBJECT(CompressionQuality, uint8, UTexture, TEnumAsByte<enum ETextureCompressionQuality>);
+	IMPLEMENT_NODE_ATTRIBUTE_APPLY_UOBJECT(CompressionNoAlpha, bool, UTexture);
+	IMPLEMENT_NODE_ATTRIBUTE_APPLY_UOBJECT(DeferCompression, bool, UTexture);
+	IMPLEMENT_NODE_ATTRIBUTE_APPLY_UOBJECT(LossyCompressionAmount, uint8, UTexture);
+	IMPLEMENT_NODE_ATTRIBUTE_APPLY_UOBJECT(MaxTextureSize, int32, UTexture);
+	IMPLEMENT_NODE_ATTRIBUTE_APPLY_UOBJECT(CompressionQuality, uint8, UTexture);
 #endif
-	IMPLEMENT_NODE_ATTRIBUTE_APPLY_UOBJECT(CompressionSettings, uint8, UTexture, TEnumAsByte<enum TextureCompressionSettings>);
+	IMPLEMENT_NODE_ATTRIBUTE_APPLY_UOBJECT(CompressionSettings, uint8, UTexture);
 #if WITH_EDITORONLY_DATA
 	//Texture general
-	IMPLEMENT_NODE_ATTRIBUTE_APPLY_UOBJECT(bDitherMipMapAlpha, bool, UTexture, );
-	IMPLEMENT_NODE_ATTRIBUTE_APPLY_UOBJECT(AlphaCoverageThresholds, FVector4, UTexture, );
-	IMPLEMENT_NODE_ATTRIBUTE_APPLY_UOBJECT(bFlipGreenChannel, bool, UTexture, );
-	IMPLEMENT_NODE_ATTRIBUTE_APPLY_UOBJECT(PowerOfTwoMode, uint8, UTexture, TEnumAsByte<enum ETexturePowerOfTwoSetting::Type>);
-	IMPLEMENT_NODE_ATTRIBUTE_APPLY_UOBJECT(PaddingColor, FColor, UTexture, );
+	IMPLEMENT_NODE_ATTRIBUTE_APPLY_UOBJECT(bDitherMipMapAlpha, bool, UTexture);
+	IMPLEMENT_NODE_ATTRIBUTE_APPLY_UOBJECT(AlphaCoverageThresholds, FVector4, UTexture);
+	IMPLEMENT_NODE_ATTRIBUTE_APPLY_UOBJECT(bFlipGreenChannel, bool, UTexture);
+	IMPLEMENT_NODE_ATTRIBUTE_APPLY_UOBJECT(PowerOfTwoMode, uint8, UTexture);
+	IMPLEMENT_NODE_ATTRIBUTE_APPLY_UOBJECT(PaddingColor, FColor, UTexture);
 #endif //#if WITH_EDITORONLY_DATA
-	IMPLEMENT_NODE_ATTRIBUTE_APPLY_UOBJECT(Filter, uint8, UTexture, TEnumAsByte<enum TextureFilter>);
-	IMPLEMENT_NODE_ATTRIBUTE_APPLY_UOBJECT(MipLoadOptions, uint8, UTexture, ETextureMipLoadOptions);
-	IMPLEMENT_NODE_ATTRIBUTE_APPLY_UOBJECT(SRGB, bool, UTexture, );
+	IMPLEMENT_NODE_ATTRIBUTE_APPLY_UOBJECT(Filter, uint8, UTexture);
+	IMPLEMENT_NODE_ATTRIBUTE_APPLY_UOBJECT(MipLoadOptions, uint8, UTexture);
+	IMPLEMENT_NODE_ATTRIBUTE_APPLY_UOBJECT(SRGB, bool, UTexture);
 #if WITH_EDITORONLY_DATA
-	IMPLEMENT_NODE_ATTRIBUTE_APPLY_UOBJECT(bUseLegacyGamma, bool, UTexture, );
+	IMPLEMENT_NODE_ATTRIBUTE_APPLY_UOBJECT(bUseLegacyGamma, bool, UTexture);
 #endif //#if WITH_EDITORONLY_DATA
-	IMPLEMENT_NODE_ATTRIBUTE_APPLY_UOBJECT(VirtualTextureStreaming, bool, UTexture, );
-	IMPLEMENT_NODE_ATTRIBUTE_APPLY_UOBJECT(AddressX, uint8, UTexture2D, TEnumAsByte<enum TextureAddress>);
-	IMPLEMENT_NODE_ATTRIBUTE_APPLY_UOBJECT(AddressY, uint8, UTexture2D, TEnumAsByte<enum TextureAddress>);
+	IMPLEMENT_NODE_ATTRIBUTE_APPLY_UOBJECT(VirtualTextureStreaming, bool, UTexture);
+	IMPLEMENT_NODE_ATTRIBUTE_APPLY_UOBJECT(AddressX, uint8, UTexture2D);
+	IMPLEMENT_NODE_ATTRIBUTE_APPLY_UOBJECT(AddressY, uint8, UTexture2D);
 #if WITH_EDITORONLY_DATA
 	//Level of Detail
-	IMPLEMENT_NODE_ATTRIBUTE_APPLY_UOBJECT(bPreserveBorder, bool, UTexture, );
-	IMPLEMENT_NODE_ATTRIBUTE_APPLY_UOBJECT(MipGenSettings, uint8, UTexture, TEnumAsByte<enum TextureMipGenSettings>);
+	IMPLEMENT_NODE_ATTRIBUTE_APPLY_UOBJECT(bPreserveBorder, bool, UTexture);
+	IMPLEMENT_NODE_ATTRIBUTE_APPLY_UOBJECT(MipGenSettings, uint8, UTexture);
 #endif //#if WITH_EDITORONLY_DATA
-	IMPLEMENT_NODE_ATTRIBUTE_APPLY_UOBJECT(LODBias, int32, UTexture, );
-	IMPLEMENT_NODE_ATTRIBUTE_APPLY_UOBJECT(LODGroup, uint8, UTexture, TEnumAsByte<enum TextureGroup>);
+	IMPLEMENT_NODE_ATTRIBUTE_APPLY_UOBJECT(LODBias, int32, UTexture);
+	IMPLEMENT_NODE_ATTRIBUTE_APPLY_UOBJECT(LODGroup, uint8, UTexture);
 
 	//TODO support per platform data in the FAttributeStorage, so we can set different value per platform at the pipeline stage, We set only the default value for now
 	//IMPLEMENT_NODE_ATTRIBUTE_APPLY_UOBJECT(Downscale, float, UTexture, );
@@ -892,11 +892,11 @@ private:
 	}
 #endif //WITH_ENGINE
 
-	IMPLEMENT_NODE_ATTRIBUTE_APPLY_UOBJECT(DownscaleOptions, uint8, UTexture, ETextureDownscaleOptions);
+	IMPLEMENT_NODE_ATTRIBUTE_APPLY_UOBJECT(DownscaleOptions, uint8, UTexture);
 #if WITH_EDITORONLY_DATA
 	//Compositing
-	IMPLEMENT_NODE_ATTRIBUTE_APPLY_UOBJECT(CompositeTextureMode, uint8, UTexture, TEnumAsByte<enum ECompositeTextureMode>);
-	IMPLEMENT_NODE_ATTRIBUTE_APPLY_UOBJECT(CompositePower, float, UTexture, );
+	IMPLEMENT_NODE_ATTRIBUTE_APPLY_UOBJECT(CompositeTextureMode, uint8, UTexture);
+	IMPLEMENT_NODE_ATTRIBUTE_APPLY_UOBJECT(CompositePower, float, UTexture);
 #endif //#if WITH_EDITORONLY_DATA
 protected:
 #if WITH_ENGINE
diff --git a/Engine/Plugins/Runtime/Interchange/Source/Nodes/Public/InterchangeTextureLightProfileFactoryNode.h b/Engine/Plugins/Runtime/Interchange/Source/Nodes/Public/InterchangeTextureLightProfileFactoryNode.h
--- a/Engine/Plugins/Runtime/Interchange/Source/Nodes/Public/InterchangeTextureLightProfileFactoryNode.h
+++ b/Engine/Plugins/Runtime/Interchange/Source/Nodes/Public/InterchangeTextureLightProfileFactoryNode.h
@@ -71,7 +71,7 @@ private:
 	IMPLEMENT_NODE_ATTRIBUTE_KEY(Brightness);
 	IMPLEMENT_NODE_ATTRIBUTE_KEY(TextureMultiplier);
 
-	IMPLEMENT_NODE_ATTRIBUTE_APPLY_UOBJECT(Brightness, float, UTextureLightProfile,);
-	IMPLEMENT_NODE_ATTRIBUTE_APPLY_UOBJECT(TextureMultiplier, float, UTextureLightProfile,);
+	IMPLEMENT_NODE_ATTRIBUTE_APPLY_UOBJECT(Brightness, float, UTextureLightProfile);
+	IMPLEMENT_NODE_ATTRIBUTE_APPLY_UOBJECT(TextureMultiplier, float, UTextureLightProfile);
 
 };
\ No newline at end of file
diff --git a/Engine/Plugins/Runtime/Interchange/Source/Parsers/Fbx/Private/FbxCamera.cpp b/Engine/Plugins/Runtime/Interchange/Source/Parsers/Fbx/Private/FbxCamera.cpp
--- a/Engine/Plugins/Runtime/Interchange/Source/Parsers/Fbx/Private/FbxCamera.cpp
+++ b/Engine/Plugins/Runtime/Interchange/Source/Parsers/Fbx/Private/FbxCamera.cpp
@@ -8,6 +8,7 @@
 #include "FbxInclude.h"
 #include "InterchangeCameraNode.h"
 #include "InterchangeResultsContainer.h"
+#include "Math/UnitConversion.h"
 #include "Nodes/InterchangeBaseNodeContainer.h"
 
 #define LOCTEXT_NAMESPACE "InterchangeFbxCamera"
@@ -18,6 +19,28 @@ namespace UE
 	{
 		namespace Private
 		{
+			void FillCameraNode(UInterchangeCameraNode* CameraNode, FbxCamera& SourceCamera)
+			{
+				if (!CameraNode)
+				{
+					return;
+				}
+
+				float FocalLength;
+				if (SourceCamera.GetApertureMode() == FbxCamera::eFocalLength)
+				{
+					FocalLength = SourceCamera.FocalLength.Get();
+				}
+				else
+				{
+					FocalLength = SourceCamera.ComputeFocalLength(SourceCamera.FieldOfView.Get());
+				}
+
+				CameraNode->SetCustomFocalLength(FocalLength); //Both FBX and UE have their focal length in mm
+				CameraNode->SetCustomSensorHeight(FUnitConversion::Convert(SourceCamera.GetApertureHeight(), EUnit::Inches, EUnit::Millimeters));
+				CameraNode->SetCustomSensorWidth(FUnitConversion::Convert(SourceCamera.GetApertureWidth(), EUnit::Inches, EUnit::Millimeters));
+			}
+
 			UInterchangeCameraNode* FFbxCamera::CreateCameraNode(UInterchangeBaseNodeContainer& NodeContainer, const FString& NodeUid, const FString& NodeName)
 			{
 				UInterchangeCameraNode* CameraNode = NewObject<UInterchangeCameraNode>(&NodeContainer, NAME_None);
@@ -41,7 +64,7 @@ namespace UE
 				{
 					FbxNodeAttribute* NodeAttribute = Node->GetNodeAttributeByIndex(AttributeIndex);
 
-					if (NodeAttribute->GetAttributeType() == FbxNodeAttribute::eCamera)
+					if (NodeAttribute && NodeAttribute->GetAttributeType() == FbxNodeAttribute::eCamera)
 					{
 						FString NodeName = FFbxHelper::GetNodeAttributeName(NodeAttribute, UInterchangeCameraNode::StaticAssetTypeName());
 						FString NodeUid = FFbxHelper::GetNodeAttributeUniqueID(NodeAttribute, UInterchangeCameraNode::StaticAssetTypeName());
@@ -50,7 +73,8 @@ namespace UE
 
 						if (!CameraNode)
 						{
-							CreateCameraNode(NodeContainer, NodeUid, NodeName);
+							CameraNode = CreateCameraNode(NodeContainer, NodeUid, NodeName);
+							FillCameraNode(CameraNode, static_cast<FbxCamera&>(*NodeAttribute));
 						}
 					}
 				}
diff --git a/Engine/Plugins/Runtime/Interchange/Source/Parsers/Fbx/Private/FbxLight.cpp b/Engine/Plugins/Runtime/Interchange/Source/Parsers/Fbx/Private/FbxLight.cpp
--- a/Engine/Plugins/Runtime/Interchange/Source/Parsers/Fbx/Private/FbxLight.cpp
+++ b/Engine/Plugins/Runtime/Interchange/Source/Parsers/Fbx/Private/FbxLight.cpp
@@ -18,9 +18,31 @@ namespace UE
 	{
 		namespace Private
 		{
-			UInterchangeLightNode* FFbxLight::CreateLightNode(UInterchangeBaseNodeContainer& NodeContainer, const FString& NodeUid, const FString& NodeName)
+			UInterchangeLightNode* FFbxLight::CreateLightNode(UInterchangeBaseNodeContainer& NodeContainer, const FString& NodeUid, const FString& NodeName, const FbxLight& LightAttribute)
 			{
-				UInterchangeLightNode* LightNode = NewObject<UInterchangeLightNode>(&NodeContainer, NAME_None);
+				UClass* LightClass;
+
+				switch(LightAttribute.LightType.Get())
+				{
+				case FbxLight::ePoint:
+				case FbxLight::eVolume:
+					LightClass = UInterchangePointLightNode::StaticClass();
+					break;
+				case FbxLight::eDirectional:
+					LightClass = UInterchangeDirectionalLightNode::StaticClass();
+					break;
+				case FbxLight::eSpot:
+					LightClass = UInterchangeSpotLightNode::StaticClass();
+					break;
+				case FbxLight::eArea:
+					LightClass = UInterchangeRectLightNode::StaticClass();
+					break;
+				default:
+					LightClass = UInterchangePointLightNode::StaticClass();
+					break;
+				}
+
+				UInterchangeLightNode* LightNode = NewObject<UInterchangeLightNode>(&NodeContainer, LightClass, NAME_None);
 				if (!ensure(LightNode))
 				{
 					UInterchangeResultError_Generic* Message = Parser.AddMessage<UInterchangeResultError_Generic>();
@@ -41,7 +63,7 @@ namespace UE
 				{
 					FbxNodeAttribute* NodeAttribute = Node->GetNodeAttributeByIndex(AttributeIndex);
 
-					if (NodeAttribute->GetAttributeType() == FbxNodeAttribute::eLight)
+					if (NodeAttribute && NodeAttribute->GetAttributeType() == FbxNodeAttribute::eLight)
 					{
 						FString NodeName = FFbxHelper::GetNodeAttributeName(NodeAttribute, UInterchangeLightNode::StaticAssetTypeName());
 						FString NodeUid = FFbxHelper::GetNodeAttributeUniqueID(NodeAttribute, UInterchangeLightNode::StaticAssetTypeName());
@@ -50,7 +72,7 @@ namespace UE
 
 						if (!LightNode)
 						{
-							CreateLightNode(NodeContainer, NodeUid, NodeName);
+							CreateLightNode(NodeContainer, NodeUid, NodeName, static_cast<FbxLight&>(*NodeAttribute));
 						}
 					}
 				}
diff --git a/Engine/Plugins/Runtime/Interchange/Source/Parsers/Fbx/Private/FbxLight.h b/Engine/Plugins/Runtime/Interchange/Source/Parsers/Fbx/Private/FbxLight.h
--- a/Engine/Plugins/Runtime/Interchange/Source/Parsers/Fbx/Private/FbxLight.h
+++ b/Engine/Plugins/Runtime/Interchange/Source/Parsers/Fbx/Private/FbxLight.h
@@ -30,7 +30,7 @@ namespace UE
 				void AddAllLights(FbxScene* SDKScene, UInterchangeBaseNodeContainer& NodeContainer);
 
 			protected:
-				UInterchangeLightNode* CreateLightNode(UInterchangeBaseNodeContainer& NodeContainer, const FString& NodeUID, const FString& NodeName);
+				UInterchangeLightNode* CreateLightNode(UInterchangeBaseNodeContainer& NodeContainer, const FString& NodeUID, const FString& NodeName, const FbxLight& LightAttribute);
 				void AddLightsRecursively(FbxNode* Node, UInterchangeBaseNodeContainer& NodeContainer);
 
 			private:
diff --git a/Engine/Plugins/Runtime/Interchange/Source/Pipelines/Private/InterchangeGenericScenesPipeline.cpp b/Engine/Plugins/Runtime/Interchange/Source/Pipelines/Private/InterchangeGenericScenesPipeline.cpp
--- a/Engine/Plugins/Runtime/Interchange/Source/Pipelines/Private/InterchangeGenericScenesPipeline.cpp
+++ b/Engine/Plugins/Runtime/Interchange/Source/Pipelines/Private/InterchangeGenericScenesPipeline.cpp
@@ -4,6 +4,7 @@
 
 #include "InterchangeActorFactoryNode.h"
 #include "InterchangeCameraNode.h"
+#include "InterchangeCineCameraFactoryNode.h"
 #include "InterchangeLightNode.h"
 #include "InterchangeMeshNode.h"
 #include "InterchangePipelineLog.h"
@@ -11,7 +12,10 @@
 
 #include "Animation/SkeletalMeshActor.h"
 #include "CineCameraActor.h"
+#include "Engine/DirectionalLight.h"
 #include "Engine/PointLight.h"
+#include "Engine/RectLight.h"
+#include "Engine/SpotLight.h"
 #include "Engine/StaticMeshActor.h"
 
 bool UInterchangeGenericLevelPipeline::ExecutePreImportPipeline(UInterchangeBaseNodeContainer* InBaseNodeContainer, const TArray<UInterchangeSourceData*>& InSourceDatas)
@@ -59,7 +63,23 @@ void UInterchangeGenericLevelPipeline::CreateActorFactoryNode(const UInterchange
 		return;
 	}
 
-	UInterchangeActorFactoryNode* ActorFactoryNode = NewObject<UInterchangeActorFactoryNode>(FactoryNodeContainer, NAME_None);
+	UInterchangeActorFactoryNode* ActorFactoryNode = nullptr;
+	const UInterchangeBaseNode* TranslatedAssetNode = nullptr;
+
+	FString AssetInstanceUid;
+	if (SceneNode->GetCustomAssetInstanceUid(AssetInstanceUid))
+	{
+		TranslatedAssetNode = FactoryNodeContainer->GetNode(AssetInstanceUid);
+	}
+
+	if (TranslatedAssetNode && TranslatedAssetNode->IsA<UInterchangeCameraNode>())
+	{
+		ActorFactoryNode = NewObject<UInterchangeCineCameraFactoryNode>(FactoryNodeContainer, NAME_None);
+	}
+	else
+	{
+		ActorFactoryNode = NewObject<UInterchangeActorFactoryNode>(FactoryNodeContainer, NAME_None);
+	}
 
 	if (!ensure(ActorFactoryNode))
 	{
@@ -81,29 +101,71 @@ void UInterchangeGenericLevelPipeline::CreateActorFactoryNode(const UInterchange
 		ActorFactoryNode->SetCustomGlobalTransform(GlobalTransform);
 	}
 	
-	FString AssetInstanceUid;
-	if (SceneNode->GetCustomAssetInstanceUid(AssetInstanceUid))
-	{
-		UInterchangeBaseNode* TranslatedAssetNode = FactoryNodeContainer->GetNode(AssetInstanceUid);
+	ActorFactoryNode->SetCustomMobility(EComponentMobility::Static);
 
-		if (UInterchangeMeshNode* MeshNode = Cast<UInterchangeMeshNode>(TranslatedAssetNode))
+	if (TranslatedAssetNode)
+	{
+		if (const UInterchangeMeshNode* MeshNode = Cast<UInterchangeMeshNode>(TranslatedAssetNode))
 		{
 			if (MeshNode->IsSkinnedMesh())
 			{
 				ActorFactoryNode->SetCustomActorClassName(ASkeletalMeshActor::StaticClass()->GetPathName());
+				ActorFactoryNode->SetCustomMobility(EComponentMobility::Movable);
 			}
 			else
 			{
 				ActorFactoryNode->SetCustomActorClassName(AStaticMeshActor::StaticClass()->GetPathName());
 			}
 		}
-		else if (UInterchangeLightNode* LightNode = Cast<UInterchangeLightNode>(TranslatedAssetNode))
+		else if (const UInterchangeLightNode* LightNode = Cast<UInterchangeLightNode>(TranslatedAssetNode))
 		{
-			ActorFactoryNode->SetCustomActorClassName(APointLight::StaticClass()->GetPathName());
+			//Test for spot before point since a spot light is a point light
+			if (LightNode->IsA<UInterchangeSpotLightNode>())
+			{
+				ActorFactoryNode->SetCustomActorClassName(ASpotLight::StaticClass()->GetPathName());
+			}
+			else if (LightNode->IsA<UInterchangePointLightNode>())
+			{
+				ActorFactoryNode->SetCustomActorClassName(APointLight::StaticClass()->GetPathName());
+			} 
+			else if (LightNode->IsA<UInterchangeRectLightNode>())
+			{
+				ActorFactoryNode->SetCustomActorClassName(ARectLight::StaticClass()->GetPathName());
+			}
+			else if (LightNode->IsA<UInterchangeDirectionalLightNode>())
+			{
+				ActorFactoryNode->SetCustomActorClassName(ADirectionalLight::StaticClass()->GetPathName());
+			}
+			else
+			{
+				ActorFactoryNode->SetCustomActorClassName(APointLight::StaticClass()->GetPathName());
+			}
 		}
-		else if (UInterchangeCameraNode* CameraNode = Cast<UInterchangeCameraNode>(TranslatedAssetNode))
+		else if (const UInterchangeCameraNode* CameraNode = Cast<UInterchangeCameraNode>(TranslatedAssetNode))
 		{
 			ActorFactoryNode->SetCustomActorClassName(ACineCameraActor::StaticClass()->GetPathName());
+			ActorFactoryNode->SetCustomMobility(EComponentMobility::Movable);
+
+			if (UInterchangeCineCameraFactoryNode* CineCameraFactoryNode = Cast<UInterchangeCineCameraFactoryNode>(ActorFactoryNode))
+			{
+				float FocalLength;
+				if (CameraNode->GetCustomFocalLength(FocalLength))
+				{
+					CineCameraFactoryNode->SetCustomFocalLength(FocalLength);
+				}
+				
+				float SensorHeight;
+				if (CameraNode->GetCustomSensorHeight(SensorHeight))
+				{
+					CineCameraFactoryNode->SetCustomSensorHeight(SensorHeight);
+				}
+
+				float SensorWidth;
+				if (CameraNode->GetCustomSensorWidth(SensorWidth))
+				{
+					CineCameraFactoryNode->SetCustomSensorWidth(SensorWidth);
+				}
+			}
 		}
 	}
 
diff --git a/Engine/Source/Runtime/Interchange/Core/Private/Nodes/InterchangeBaseNode.cpp b/Engine/Source/Runtime/Interchange/Core/Private/Nodes/InterchangeBaseNode.cpp
--- a/Engine/Source/Runtime/Interchange/Core/Private/Nodes/InterchangeBaseNode.cpp
+++ b/Engine/Source/Runtime/Interchange/Core/Private/Nodes/InterchangeBaseNode.cpp
@@ -241,7 +241,7 @@ FString UInterchangeBaseNode::InvalidNodeUid()
 	return FString();
 }
 
-void UInterchangeBaseNode::ApplyAllCustomAttributeToAsset(UObject* Object) const
+void UInterchangeBaseNode::ApplyAllCustomAttributeToObject(UObject* Object) const
 {
 	UClass* ObjectClass = Object->GetClass();
 	for (const TPair<UClass*, TArray<UE::Interchange::FApplyAttributeToAsset>>& ClassDelegatePair : ApplyCustomAttributeDelegates)
@@ -259,7 +259,7 @@ void UInterchangeBaseNode::ApplyAllCustomAttributeToAsset(UObject* Object) const
 	}
 }
 
-void UInterchangeBaseNode::FillAllCustomAttributeFromAsset(UObject* Object) const
+void UInterchangeBaseNode::FillAllCustomAttributeFromObject(UObject* Object) const
 {
 	UClass* ObjectClass = Object->GetClass();
 	for (const TPair<UClass*, TArray<UE::Interchange::FFillAttributeToAsset>>& ClassDelegatePair : FillCustomAttributeDelegates)
@@ -289,3 +289,42 @@ void UInterchangeBaseNode::Serialize(FArchive& Ar)
 
 	}
 }
+
+FProperty* InterchangePrivateNodeBase::FindPropertyByPathChecked(TVariant<UObject*, uint8*>& Container, UStruct* Outer, FStringView PropertyPath)
+{
+	int32 SeparatorIndex;
+	FStringView PropertyName;
+	FStringView RestOfPropertyPath(PropertyPath);
+	FProperty* Property = nullptr;
+
+	do
+	{
+		if (!RestOfPropertyPath.FindChar(TEXT('.'), SeparatorIndex))
+		{
+			SeparatorIndex = RestOfPropertyPath.Len();
+		}
+
+		PropertyName = FStringView(RestOfPropertyPath.GetData(), SeparatorIndex);
+		RestOfPropertyPath = FStringView(RestOfPropertyPath.GetData() + SeparatorIndex + 1);
+
+		Property = FindFieldChecked<FProperty>(Outer, FName(PropertyName.Len(), PropertyName.GetData()));
+
+		if (FStructProperty* StructProperty = CastField<FStructProperty>(Property))
+		{
+			Outer = StructProperty->Struct;
+
+			if (Container.IsType<UObject*>())
+			{
+				Container.Set<uint8*>(Property->ContainerPtrToValuePtr<uint8>(Container.Get<UObject*>()));
+			}
+			else
+			{
+				Container.Set<uint8*>(Property->ContainerPtrToValuePtr<uint8>(Container.Get<uint8*>()));
+			}
+		}
+		
+	} while (RestOfPropertyPath.Len() > 0);
+
+	check(Property);
+	return Property;
+}
diff --git a/Engine/Source/Runtime/Interchange/Core/Public/Nodes/InterchangeBaseNode.h b/Engine/Source/Runtime/Interchange/Core/Public/Nodes/InterchangeBaseNode.h
--- a/Engine/Source/Runtime/Interchange/Core/Public/Nodes/InterchangeBaseNode.h
+++ b/Engine/Source/Runtime/Interchange/Core/Public/Nodes/InterchangeBaseNode.h
@@ -66,6 +66,16 @@ namespace InterchangePrivateNodeBase
 		}
 		return true;
 	}
+
+	/**
+	 * Finds a property by name in Outer and supports looking into FStructProperties (embedded structs) with a '.' separating the property names.
+	 * 
+	 * @param Container - The container for the property values. If the final property is inside a UScriptStruct, the container will be set to the UScriptStruct instance address.
+	 * @param Outer - The UStruct containing the top level property.
+	 * @param PropertyPath - A dot separated chain of properties. Doesn't support going through external objects.
+	 * @return The Property matching the last name in the PropertyPath.
+	 */
+	INTERCHANGECORE_API FProperty* FindPropertyByPathChecked(TVariant<UObject*, uint8*>& Container, UStruct* Outer, FStringView PropertyPath);
 }
 
 /**
@@ -82,52 +92,82 @@ namespace InterchangePrivateNodeBase
  * @param AssetType - This is the asset you want to apply the storage value"
  * @param EnumType - Optional, specify it only if the AssetType member is an enum so we can type cast it in the apply function (we use uint8 to store the enum value)"
  */
-#define IMPLEMENT_NODE_ATTRIBUTE_KEY(AttributeName)																		\
+#define IMPLEMENT_NODE_ATTRIBUTE_KEY(AttributeName)																\
 const UE::Interchange::FAttributeKey Macro_Custom##AttributeName##Key = UE::Interchange::FAttributeKey(TEXT(#AttributeName));
 
 #if WITH_ENGINE
-#define IMPLEMENT_NODE_ATTRIBUTE_APPLY_UOBJECT(AttributeName, AttributeType, AssetType, EnumType)	\
-bool ApplyCustom##AttributeName##ToAsset(UObject* Asset) const										\
-{																									\
-	if (!Asset)																						\
-	{																								\
-		return false;																				\
-	}																								\
-	AssetType* TypedObject = Cast<AssetType>(Asset);												\
-	if (!TypedObject)																				\
-	{																								\
-		return false;																				\
-	}																								\
-	AttributeType ValueData;																		\
-	if (GetCustom##AttributeName(ValueData))														\
-	{																								\
-		TypedObject->AttributeName = EnumType(ValueData);											\
-		return true;																				\
-	}																								\
-	return false;																					\
-}																									\
-																									\
-bool FillCustom##AttributeName##FromAsset(UObject* Asset)											\
-{																									\
-	if (!Asset)																						\
-	{																								\
-		return false;																				\
-	}																								\
-	AssetType* TypedObject = Cast<AssetType>(Asset);												\
-	if (!TypedObject)																				\
-	{																								\
-		return false;																				\
-	}																								\
-	if (SetCustom##AttributeName((AttributeType)TypedObject->AttributeName, false))					\
-	{																								\
-		return true;																				\
-	}																								\
-	return false;																					\
+#define IMPLEMENT_NODE_ATTRIBUTE_APPLY_UOBJECT_BYNAME(AttributeName, AttributeType, ObjectType, PropertyName)	\
+bool ApplyCustom##AttributeName##ToAsset(UObject* Asset) const													\
+{																												\
+	if (!Asset)																									\
+	{																											\
+		return false;																							\
+	}																											\
+	ObjectType* TypedObject = Cast<ObjectType>(Asset);															\
+	if (!TypedObject)																							\
+	{																											\
+		return false;																							\
+	}																											\
+	AttributeType ValueData;																					\
+	if (GetCustom##AttributeName(ValueData))																	\
+	{																											\
+		TVariant<UObject*, uint8*> Container;																	\
+		Container.Set<UObject*>(Asset);																			\
+		if(FProperty* Property = InterchangePrivateNodeBase::FindPropertyByPathChecked(Container, Asset->GetClass(), PropertyName))\
+		{																										\
+			if (Container.IsType<UObject*>())																	\
+			{																									\
+				*(Property->ContainerPtrToValuePtr<AttributeType>(Container.Get<UObject*>())) = ValueData;		\
+			}																									\
+			else																								\
+			{																									\
+				*(Property->ContainerPtrToValuePtr<AttributeType>(Container.Get<uint8*>())) = ValueData;		\
+			}																									\
+		}																										\
+		return true;																							\
+	}																											\
+	return false;																								\
+}																												\
+																												\
+bool FillCustom##AttributeName##FromAsset(UObject* Asset)														\
+{																												\
+	if (!Asset)																									\
+	{																											\
+		return false;																							\
+	}																											\
+	ObjectType* TypedObject = Cast<ObjectType>(Asset);															\
+	if (!TypedObject)																							\
+	{																											\
+		return false;																							\
+	}																											\
+	TVariant<UObject*, uint8*> Container;																		\
+	Container.Set<UObject*>(Asset);																				\
+	if(FProperty* Property = InterchangePrivateNodeBase::FindPropertyByPathChecked(Container, Asset->GetClass(), PropertyName))\
+	{																											\
+		AttributeType ValueData;																				\
+		if (Container.IsType<UObject*>())																		\
+		{																										\
+			ValueData = *(Property->ContainerPtrToValuePtr<AttributeType>(Container.Get<UObject*>()));			\
+		}																										\
+		else																									\
+		{																										\
+			ValueData = *(Property->ContainerPtrToValuePtr<AttributeType>(Container.Get<uint8*>()));			\
+		}																										\
+		if (SetCustom##AttributeName(ValueData, false))															\
+		{																										\
+			return true;																						\
+		}																										\
+	}																											\
+																												\
+	return false;																								\
 }
+#define IMPLEMENT_NODE_ATTRIBUTE_APPLY_UOBJECT(AttributeName, AttributeType, ObjectType)						\
+	IMPLEMENT_NODE_ATTRIBUTE_APPLY_UOBJECT_BYNAME(AttributeName, AttributeType, ObjectType, TEXT(#AttributeName))
 
 #else //#if WITH_ENGINE
 
-#define IMPLEMENT_NODE_ATTRIBUTE_APPLY_UOBJECT(AttributeName, AttributeType, AssetType, EnumType)
+#define IMPLEMENT_NODE_ATTRIBUTE_APPLY_UOBJECT(AttributeName, AttributeType, AssetType)
+#define IMPLEMENT_NODE_ATTRIBUTE_APPLY_UOBJECT_BYNAME(AttributeName, AttributeType, AssetType, PropertyName)
 
 #endif //#if WITH_ENGINE
 
@@ -1201,9 +1241,9 @@ public:
 	 * See the macros IMPLEMENT_NODE_ATTRIBUTE_SETTER at the top of the file to know how delegates are setup for property.
 	 *
 	 */
-	void ApplyAllCustomAttributeToAsset(UObject* Object) const;
+	void ApplyAllCustomAttributeToObject(UObject* Object) const;
 
-	void FillAllCustomAttributeFromAsset(UObject* Object) const;
+	void FillAllCustomAttributeFromObject(UObject* Object) const;
 
 	virtual void Serialize(FArchive& Ar) override;
 
diff --git a/Engine/Source/Runtime/Interchange/Engine/Private/Tasks/InterchangeTaskCompletion.cpp b/Engine/Source/Runtime/Interchange/Engine/Private/Tasks/InterchangeTaskCompletion.cpp
--- a/Engine/Source/Runtime/Interchange/Engine/Private/Tasks/InterchangeTaskCompletion.cpp
+++ b/Engine/Source/Runtime/Interchange/Engine/Private/Tasks/InterchangeTaskCompletion.cpp
@@ -81,36 +81,41 @@ void UE::Interchange::FTaskPreCompletion::DoTask(ENamedThreads::Type CurrentThre
 				Message->DestinationAssetName = ImportedObject->GetPathName();
 				Message->AssetType = ImportedObject->GetClass();
 
-				if (bIsAsset)
+				//Clear any async flag from the created asset
+				const EInternalObjectFlags AsyncFlags = EInternalObjectFlags::Async | EInternalObjectFlags::AsyncLoading;
+				ImportedObject->ClearInternalFlags(AsyncFlags);
+				//Make sure the package is dirty
+				ImportedObject->MarkPackageDirty();
+
+				if (!bIsAsset)
 				{
-					//Clear any async flag from the created asset
-					const EInternalObjectFlags AsyncFlags = EInternalObjectFlags::Async | EInternalObjectFlags::AsyncLoading;
-					ImportedObject->ClearInternalFlags(AsyncFlags);
-					//Make sure the package is dirty
-					ImportedObject->MarkPackageDirty();
+					if (AActor* Actor = Cast<AActor>(ImportedObject))
+					{
+						Actor->RegisterAllComponents();
+					}
+					else if (UActorComponent* Component = Cast<UActorComponent>(ImportedObject))
+					{
+						Component->RegisterComponent();
+					}
+				}
+
 #if WITH_EDITOR
-					//Make sure the asset is built correctly
-					ImportedObject->PostEditChange();
+				ImportedObject->PostEditChange();
 #endif
-					//Post import broadcast
+
+				//Post import broadcast
+				if (bIsAsset)
+				{
+					AsyncHelper->AssetImportResult->AddImportedObject(ImportedObject);
+
 					if (!AsyncHelper->TaskData.ReimportObject)
 					{
 						//Notify the asset registry, only when we have created the asset
 						FAssetRegistryModule::AssetCreated(ImportedObject);
 					}
-					AsyncHelper->AssetImportResult->AddImportedObject(ImportedObject);
 				}
 				else
 				{
-					if (AActor* Actor = Cast<AActor>(ImportedObject))
-					{
-						Actor->RegisterAllComponents();
-					}
-					else if (UActorComponent* Component = Cast<UActorComponent>(ImportedObject))
-					{
-						Component->RegisterComponent();
-					}
-
 					AsyncHelper->SceneImportResult->AddImportedObject(ImportedObject);
 				}
 
