Interchange: Move files out of experimental

#rnx

diff --git a/Engine/Plugins/Interchange/Editor/Source/ResultsBrowser/InterchangeResultsBrowser.Build.cs b/Engine/Plugins/Interchange/Editor/Source/ResultsBrowser/InterchangeResultsBrowser.Build.cs
new file mode 100644
--- /dev/null
+++ b/Engine/Plugins/Interchange/Editor/Source/ResultsBrowser/InterchangeResultsBrowser.Build.cs
@@ -0,0 +1,32 @@
+// Copyright Epic Games, Inc. All Rights Reserved.
+
+namespace UnrealBuildTool.Rules
+{
+	public class InterchangeResultsBrowser : ModuleRules
+	{
+		public InterchangeResultsBrowser(ReadOnlyTargetRules Target) : base(Target)
+		{
+			PublicDependencyModuleNames.AddRange(
+				new string[]
+				{
+					"Core",
+					"CoreUObject",
+					"Engine",
+					"InterchangeCore",
+					"InterchangeEngine",
+				}
+			);
+
+			PrivateDependencyModuleNames.AddRange(
+				new string[]
+				{
+					"InputCore",
+					"Projects",
+					"Slate",
+					"SlateCore",
+					"WorkspaceMenuStructure"
+				}
+			);
+		}
+	}
+}
diff --git a/Engine/Plugins/Interchange/Editor/Source/ResultsBrowser/Private/InterchangeResultsBrowserModule.cpp b/Engine/Plugins/Interchange/Editor/Source/ResultsBrowser/Private/InterchangeResultsBrowserModule.cpp
new file mode 100644
--- /dev/null
+++ b/Engine/Plugins/Interchange/Editor/Source/ResultsBrowser/Private/InterchangeResultsBrowserModule.cpp
@@ -0,0 +1,136 @@
+// Copyright Epic Games, Inc. All Rights Reserved.
+
+#include "InterchangeResultsBrowserModule.h"
+
+#include "CoreMinimal.h"
+#include "Engine/Engine.h"
+#include "Framework/Docking/TabManager.h"
+#include "InterchangeResultsBrowserStyle.h"
+#include "InterchangeManager.h"
+#include "SInterchangeResultsBrowserWindow.h"
+#include "Styling/SlateIconFinder.h"
+#include "Widgets/Docking/SDockTab.h"
+#include "WorkspaceMenuStructure.h"
+#include "WorkspaceMenuStructureModule.h"
+
+
+#define LOCTEXT_NAMESPACE "InterchangeResultsBrowser"
+
+class FInterchangeResultsBrowserModule : public IInterchangeResultsBrowserModule
+{
+	/** IModuleInterface implementation */
+	virtual void StartupModule() override;
+	virtual void ShutdownModule() override;
+
+private:
+
+	TSharedRef<SDockTab> SpawnInterchangeResults(const FSpawnTabArgs& Args);
+	void OpenErrorBrowser(TStrongObjectPtr<UInterchangeResultsContainer> InResultsContainer);
+
+	/** Pointer to the style set to use for the UI. */
+	TSharedPtr<ISlateStyle> InterchangeResultsBrowserStyle = nullptr;
+
+	/** Owning pointer to the InterchangeResultsContainer being displayed */
+	TStrongObjectPtr<UInterchangeResultsContainer> ResultsContainer;
+
+	/** Pointer to the Slate widget */
+	TSharedPtr<SInterchangeResultsBrowserWindow> InterchangeResultsBrowserWindow;
+
+	/** Whether the results list is filtered */
+	// @TODO: load this from config
+	bool bIsFiltered = false;
+};
+
+
+IMPLEMENT_MODULE(FInterchangeResultsBrowserModule, InterchangeResultsBrowser)
+
+
+void FInterchangeResultsBrowserModule::StartupModule()
+{
+	auto RegisterItems = [this]()
+	{
+		UInterchangeManager& InterchangeManager = UInterchangeManager::GetInterchangeManager();
+		InterchangeManager.OnBatchImportComplete.AddRaw(this, &FInterchangeResultsBrowserModule::OpenErrorBrowser);
+
+		auto UnregisterItems = [this]()
+		{
+			UInterchangeManager& InterchangeManager = UInterchangeManager::GetInterchangeManager();
+			InterchangeManager.OnBatchImportComplete.RemoveAll(this);
+		};
+
+		InterchangeManager.OnPreDestroyInterchangeManager.AddLambda(UnregisterItems);
+	};
+
+	if (GEngine)
+	{
+		RegisterItems();
+	}
+	else
+	{
+		FCoreDelegates::OnPostEngineInit.AddLambda(RegisterItems);
+	}
+	
+	if (!InterchangeResultsBrowserStyle.IsValid())
+	{
+		InterchangeResultsBrowserStyle = MakeShared<FInterchangeResultsBrowserStyle>();
+	}
+
+	FGlobalTabmanager::Get()->RegisterNomadTabSpawner("InterchangeResults", FOnSpawnTab::CreateRaw(this, &FInterchangeResultsBrowserModule::SpawnInterchangeResults))
+		.SetDisplayName(LOCTEXT("InterchangeResultsBrowser", "Interchange Results Browser"))
+		.SetTooltipText(LOCTEXT("InterchangeResultsBrowserTooltipText", "Open the Interchange Results Browser tab."))
+		.SetGroup(WorkspaceMenu::GetMenuStructure().GetDeveloperToolsLogCategory())
+		.SetIcon(FSlateIconFinder::FindIcon("InterchangeResultsBrowser.TabIcon"));
+}
+
+
+void FInterchangeResultsBrowserModule::ShutdownModule()
+{
+	InterchangeResultsBrowserStyle = nullptr;
+}
+
+
+TSharedRef<SDockTab> FInterchangeResultsBrowserModule::SpawnInterchangeResults(const FSpawnTabArgs& Args)
+{
+	TSharedRef<SDockTab> DockTab = SNew(SDockTab)
+		.TabRole(ETabRole::NomadTab);
+
+	DockTab->SetContent(
+		SAssignNew(InterchangeResultsBrowserWindow, SInterchangeResultsBrowserWindow)
+		.OwnerTab(DockTab)
+		.InterchangeResultsContainer(ResultsContainer.Get())
+		.IsFiltered(bIsFiltered)
+		.OnFilterChangedState_Lambda([this](bool bNewState) { this->bIsFiltered = bNewState; })
+	);
+
+	return DockTab;
+}
+
+
+void FInterchangeResultsBrowserModule::OpenErrorBrowser(TStrongObjectPtr<UInterchangeResultsContainer> InResultsContainer)
+{
+	// Only showing when we have errors or warnings for now
+	const bool bShouldShow = [&InResultsContainer]()
+	{
+		for (UInterchangeResult* Result : InResultsContainer->GetResults())
+		{
+			if (Result->GetResultType() != EInterchangeResultType::Success)
+			{
+				return true;
+			}
+		}
+
+		return false;
+	}();
+
+	if (GIsAutomationTesting || !bShouldShow)
+	{
+		return;
+	}
+
+	ResultsContainer = InResultsContainer;
+	ResultsContainer->Finalize();
+	FGlobalTabmanager::Get()->TryInvokeTab(FName("InterchangeResults"));
+	InterchangeResultsBrowserWindow->Set(ResultsContainer.Get());
+}
+
+#undef LOCTEXT_NAMESPACE
diff --git a/Engine/Plugins/Interchange/Editor/Source/ResultsBrowser/Private/InterchangeResultsBrowserStyle.h b/Engine/Plugins/Interchange/Editor/Source/ResultsBrowser/Private/InterchangeResultsBrowserStyle.h
new file mode 100644
--- /dev/null
+++ b/Engine/Plugins/Interchange/Editor/Source/ResultsBrowser/Private/InterchangeResultsBrowserStyle.h
@@ -0,0 +1,44 @@
+// Copyright Epic Games, Inc. All Rights Reserved.
+
+#pragma once
+
+#include "Interfaces/IPluginManager.h"
+#include "Brushes/SlateImageBrush.h"
+#include "Styling/SlateStyle.h"
+#include "Styling/SlateStyleRegistry.h"
+
+/**
+ * Implements the visual style of the text asset editor UI.
+ */
+class FInterchangeResultsBrowserStyle final : public FSlateStyleSet
+{
+public:
+
+	FInterchangeResultsBrowserStyle()
+		: FSlateStyleSet("InterchangeResultsBrowserStyle")
+	{
+		const FVector2D Icon16x16(16.0f, 16.0f);
+
+		const FString BaseDir = IPluginManager::Get().FindPlugin("InterchangeEditor")->GetBaseDir();
+		SetContentRoot(BaseDir / TEXT("Content"));
+
+		FSlateImageBrush* SuccessBrush = new FSlateImageBrush(RootToContentDir(TEXT("Resources/Result_Success"), TEXT(".png")), Icon16x16);
+		Set("InterchangeResultsBrowser.ResultType.Success", SuccessBrush);
+		
+		FSlateImageBrush* WarningBrush = new FSlateImageBrush(RootToContentDir(TEXT("Resources/Result_Warning"), TEXT(".png")), Icon16x16);
+		Set("InterchangeResultsBrowser.ResultType.Warning", WarningBrush);
+
+		FSlateImageBrush* ErrorBrush = new FSlateImageBrush(RootToContentDir(TEXT("Resources/Result_Fail"), TEXT(".png")), Icon16x16);
+		Set("InterchangeResultsBrowser.ResultType.Error", ErrorBrush);
+
+		FSlateImageBrush* TabIcon = new FSlateImageBrush(RootToContentDir(TEXT("Resources/TabIcon"), TEXT(".png")), Icon16x16);
+		Set("InterchangeResultsBrowser.TabIcon", TabIcon);
+
+		FSlateStyleRegistry::RegisterSlateStyle(*this);
+	}
+
+	~FInterchangeResultsBrowserStyle()
+	{
+		FSlateStyleRegistry::UnRegisterSlateStyle(*this);
+	}
+};
diff --git a/Engine/Plugins/Interchange/Editor/Source/ResultsBrowser/Private/SInterchangeResultsBrowserWindow.cpp b/Engine/Plugins/Interchange/Editor/Source/ResultsBrowser/Private/SInterchangeResultsBrowserWindow.cpp
new file mode 100644
--- /dev/null
+++ b/Engine/Plugins/Interchange/Editor/Source/ResultsBrowser/Private/SInterchangeResultsBrowserWindow.cpp
@@ -0,0 +1,433 @@
+// Copyright Epic Games, Inc. All Rights Reserved.
+
+#include "SInterchangeResultsBrowserWindow.h"
+
+#include "Styling/AppStyle.h"
+#include "Framework/MultiBox/MultiBoxBuilder.h"
+#include "IDocumentation.h"
+#include "Modules/ModuleManager.h"
+#include "Styling/SlateIconFinder.h"
+#include "Widgets/Images/SImage.h"
+#include "Widgets/Input/SButton.h"
+#include "Widgets/Input/SCheckBox.h"
+#include "Widgets/Layout/SSplitter.h"
+#include "Widgets/Layout/SUniformGridPanel.h"
+#include "Widgets/Layout/SSpacer.h"
+
+#define LOCTEXT_NAMESPACE "InterchangeResultsBrowser"
+
+
+namespace SResultsBrowserWidgetDefs
+{
+	const FName ColumnID_ResultType("ResultType");
+	const FName ColumnID_Source("Source");
+	const FName ColumnID_Dest("Dest");
+	const FName ColumnID_AssetType("AssetType");
+	const FName ColumnID_Message("Message");
+}
+
+
+SInterchangeResultsBrowserWindow::SInterchangeResultsBrowserWindow()
+	: OwnerTab(nullptr),
+	  bIsFiltered(false)
+{
+}
+
+
+SInterchangeResultsBrowserWindow::~SInterchangeResultsBrowserWindow()
+{
+}
+
+
+void SInterchangeResultsBrowserWindow::Construct(const FArguments& InArgs)
+{
+	SortByColumn = SResultsBrowserWidgetDefs::ColumnID_ResultType;
+	SortMode = EColumnSortMode::Ascending;
+	OnFilterChangedState = InArgs._OnFilterChangedState;
+	bIsFiltered = InArgs._IsFiltered;
+	ResultsContainer = InArgs._InterchangeResultsContainer;
+
+	OwnerTab = InArgs._OwnerTab;
+	check(OwnerTab.IsValid());
+
+	TSharedRef<SHeaderRow> HeaderRowWidget = SNew(SHeaderRow);
+
+	HeaderRowWidget->AddColumn(
+		SHeaderRow::Column(SResultsBrowserWidgetDefs::ColumnID_ResultType)
+		.SortMode(this, &SInterchangeResultsBrowserWindow::GetColumnSortMode, SResultsBrowserWidgetDefs::ColumnID_ResultType)
+		.OnSort(this, &SInterchangeResultsBrowserWindow::OnColumnSortModeChanged)
+		[
+			SNew(SSpacer)
+		]
+		.FixedWidth(24.0f)
+		);
+
+	HeaderRowWidget->AddColumn(
+		SHeaderRow::Column(SResultsBrowserWidgetDefs::ColumnID_Source)
+		.DefaultLabel(LOCTEXT("ColumnLabel_SourceFile", "Source File"))
+		.SortMode(this, &SInterchangeResultsBrowserWindow::GetColumnSortMode, SResultsBrowserWidgetDefs::ColumnID_Source)
+		.OnSort(this, &SInterchangeResultsBrowserWindow::OnColumnSortModeChanged)
+		.FillWidth(5.0f)
+	);
+
+	HeaderRowWidget->AddColumn(
+		SHeaderRow::Column(SResultsBrowserWidgetDefs::ColumnID_Dest)
+		.DefaultLabel(LOCTEXT("ColumnLabel_DestinationAsset", "Destination Asset"))
+		.SortMode(this, &SInterchangeResultsBrowserWindow::GetColumnSortMode, SResultsBrowserWidgetDefs::ColumnID_Dest)
+		.OnSort(this, &SInterchangeResultsBrowserWindow::OnColumnSortModeChanged)
+		.FillWidth(5.0f)
+	);
+
+	HeaderRowWidget->AddColumn(
+		SHeaderRow::Column(SResultsBrowserWidgetDefs::ColumnID_AssetType)
+		.DefaultLabel(LOCTEXT("ColumnLabel_AssetType", "Asset Type"))
+		.SortMode(this, &SInterchangeResultsBrowserWindow::GetColumnSortMode, SResultsBrowserWidgetDefs::ColumnID_AssetType)
+		.OnSort(this, &SInterchangeResultsBrowserWindow::OnColumnSortModeChanged)
+		.FillWidth(3.0f)
+	);
+
+	HeaderRowWidget->AddColumn(
+		SHeaderRow::Column(SResultsBrowserWidgetDefs::ColumnID_Message)
+		.DefaultLabel(LOCTEXT("ColumnLabel_Details", "Details"))
+		.FillWidth(9.0f)
+	);
+
+	this->ChildSlot
+	[
+		SNew(SBorder)
+		.Padding(FMargin(10.0f, 3.0f))
+		.BorderImage(FAppStyle::GetBrush("ToolPanel.DarkGroupBorder"))
+		[
+			SNew(SVerticalBox)
+			+ SVerticalBox::Slot()
+			.AutoHeight()
+			.HAlign(HAlign_Left)
+			.Padding(10)
+			[
+				SNew(SHorizontalBox)
+				+ SHorizontalBox::Slot()
+				.AutoWidth()
+				[
+					SNew(SCheckBox)
+					.IsChecked(bIsFiltered ? ECheckBoxState::Checked : ECheckBoxState::Unchecked)
+					.OnCheckStateChanged(this, &SInterchangeResultsBrowserWindow::OnFilterStateChanged)
+				]
+				+ SHorizontalBox::Slot()
+				.AutoWidth()
+				[
+					SNew(STextBlock)
+					.Text(LOCTEXT("OnlyShowErrors", "Only display errors/warnings"))
+				]
+			]
+			+ SVerticalBox::Slot()
+			.FillHeight(1.0f)
+			.Padding(2)
+			[
+				SNew(SSplitter)
+				.Orientation(Orient_Vertical)
+				+ SSplitter::Slot()
+				.Value(0.4f)
+				[
+					SAssignNew(ListView, SListView<UInterchangeResult*>)
+					.ItemHeight(30)
+					.ListItemsSource(&FilteredResults)
+					.HeaderRow(HeaderRowWidget)
+					.SelectionMode(ESelectionMode::Single)
+					.OnGenerateRow(this, &SInterchangeResultsBrowserWindow::OnGenerateRowForList)
+					.OnSelectionChanged(this, &SInterchangeResultsBrowserWindow::OnSelectionChanged)
+				]
+				+ SSplitter::Slot()
+				.Value(0.2f)
+				[
+					SNew(SBorder)
+					.Padding(3)
+					.BorderImage(FAppStyle::GetBrush("ToolPanel.GroupBorder"))
+					[
+						SNew(SBox)
+						.Padding(10)
+						[
+							SAssignNew(Description, STextBlock)
+							.AutoWrapText(true)
+						]
+					]
+				]
+			]
+			+ SVerticalBox::Slot()
+			.AutoHeight()
+			.HAlign(HAlign_Right)
+			.Padding(2)
+			[
+				SNew(SUniformGridPanel)
+				.SlotPadding(2)
+				+ SUniformGridPanel::Slot(0, 0)
+				[
+					IDocumentation::Get()->CreateAnchor(FString("Engine/Content/Interchange/ErrorBrowser"))
+				]
+				+ SUniformGridPanel::Slot(1, 0)
+				[
+					SNew(SButton)
+					.HAlign(HAlign_Center)
+					.Text(LOCTEXT("ErrorBrowser_Clear", "Clear"))
+					.IsEnabled(false)
+					.OnClicked(this, &SInterchangeResultsBrowserWindow::OnCloseDialog)
+				]
+				+ SUniformGridPanel::Slot(2, 0)
+				[
+					SNew(SButton)
+					.HAlign(HAlign_Center)
+					.Text(LOCTEXT("ErrorBrowser_Close", "Close"))
+					.OnClicked(this, &SInterchangeResultsBrowserWindow::OnCloseDialog)
+				]
+			]
+		]
+	];
+
+	Set(InArgs._InterchangeResultsContainer);
+}
+
+
+void SInterchangeResultsBrowserWindow::Set(UInterchangeResultsContainer* Data)
+{
+	ResultsContainer = Data;
+	RepopulateItems();
+}
+
+
+void SInterchangeResultsBrowserWindow::RepopulateItems()
+{
+	FilteredResults.Empty();
+
+	if (ResultsContainer != nullptr)
+	{
+		for (UInterchangeResult* Item : ResultsContainer->GetResults())
+		{
+			if (Item->GetResultType() != EInterchangeResultType::Success || !bIsFiltered)
+			{
+				FilteredResults.Add(Item);
+			}
+		}
+
+		RequestSort();
+	}
+}
+
+
+void SInterchangeResultsBrowserWindow::CloseErrorBrowser()
+{
+	if (TSharedPtr<SDockTab> OwnerTabPin = OwnerTab.Pin())
+	{
+		OwnerTabPin->RequestCloseTab();
+	}
+	OwnerTab = nullptr;
+}
+
+
+TSharedRef<ITableRow> SInterchangeResultsBrowserWindow::OnGenerateRowForList(UInterchangeResult* InItem, const TSharedRef<STableViewBase>& OwnerTable)
+{
+	TSharedRef<ITableRow> Row =
+		SNew(SInterchangeResultsBrowserListRow, OwnerTable)
+		.InterchangeResultsBrowserWidget(SharedThis(this))
+		.Item(InItem);
+
+	return Row;
+}
+
+
+EColumnSortMode::Type SInterchangeResultsBrowserWindow::GetColumnSortMode(const FName ColumnId) const
+{
+	if (SortByColumn != ColumnId)
+	{
+		return EColumnSortMode::None;
+	}
+
+	return SortMode;
+}
+
+
+void SInterchangeResultsBrowserWindow::OnColumnSortModeChanged(const EColumnSortPriority::Type SortPriority, const FName& ColumnId, const EColumnSortMode::Type InSortMode)
+{
+	SortByColumn = ColumnId;
+	SortMode = InSortMode;
+
+	RequestSort();
+}
+
+
+void SInterchangeResultsBrowserWindow::OnFilterStateChanged(ECheckBoxState NewState)
+{
+	bIsFiltered = (NewState == ECheckBoxState::Checked);
+	OnFilterChangedState.ExecuteIfBound(bIsFiltered);
+	RepopulateItems();
+}
+
+
+void SInterchangeResultsBrowserWindow::RequestSort()
+{
+	if (SortByColumn == SResultsBrowserWidgetDefs::ColumnID_ResultType)
+	{
+		if (SortMode == EColumnSortMode::Ascending)
+		{
+			FilteredResults.Sort([](const UInterchangeResult& A, const UInterchangeResult& B) {
+				return static_cast<int>(A.GetResultType()) < static_cast<int>(B.GetResultType()); });
+		}
+		else if (SortMode == EColumnSortMode::Descending)
+		{
+			FilteredResults.Sort([](const UInterchangeResult& A, const UInterchangeResult& B) {
+				return static_cast<int>(A.GetResultType()) >= static_cast<int>(B.GetResultType()); });
+		}
+	}
+	else if (SortByColumn == SResultsBrowserWidgetDefs::ColumnID_Source)
+	{
+		if (SortMode == EColumnSortMode::Ascending)
+		{
+			FilteredResults.Sort([](const UInterchangeResult& A, const UInterchangeResult& B) {
+				return A.SourceAssetName < B.SourceAssetName; });
+		}
+		else if (SortMode == EColumnSortMode::Descending)
+		{
+			FilteredResults.Sort([](const UInterchangeResult& A, const UInterchangeResult& B) {
+				return A.SourceAssetName >= B.SourceAssetName; });
+		}
+	}
+	else if (SortByColumn == SResultsBrowserWidgetDefs::ColumnID_Dest)
+	{
+		if (SortMode == EColumnSortMode::Ascending)
+		{
+			FilteredResults.Sort([](const UInterchangeResult& A, const UInterchangeResult& B) {
+				return A.DestinationAssetName < B.DestinationAssetName; });
+		}
+		else if (SortMode == EColumnSortMode::Descending)
+		{
+			FilteredResults.Sort([](const UInterchangeResult& A, const UInterchangeResult& B) {
+				return A.DestinationAssetName >= B.DestinationAssetName; });
+		}
+	}
+	else if (SortByColumn == SResultsBrowserWidgetDefs::ColumnID_AssetType)
+	{
+		if (SortMode == EColumnSortMode::Ascending)
+		{
+			FilteredResults.Sort([](const UInterchangeResult& A, const UInterchangeResult& B) {
+				FString NameA = A.AssetType ? A.AssetType->GetName() : FString();
+				FString NameB = B.AssetType ? B.AssetType->GetName() : FString();
+				return NameA < NameB;
+				});
+		}
+		else if (SortMode == EColumnSortMode::Descending)
+		{
+			FilteredResults.Sort([](const UInterchangeResult& A, const UInterchangeResult& B) {
+				FString NameA = A.AssetType ? A.AssetType->GetName() : FString();
+				FString NameB = B.AssetType ? B.AssetType->GetName() : FString();
+				return NameA >= NameB;
+				});
+		}
+	}
+
+	ListView->RequestListRefresh();
+}
+
+
+
+void SInterchangeResultsBrowserWindow::OnSelectionChanged(UInterchangeResult* InItem, ESelectInfo::Type SelectInfo)
+{
+	if (InItem != nullptr)
+	{
+		Description->SetText(InItem->GetText());
+	}
+	else
+	{
+		Description->SetText(FText());
+	}
+}
+
+
+
+
+void SInterchangeResultsBrowserListRow::Construct(const FArguments& InArgs, const TSharedRef<STableViewBase>& InOwnerTableView)
+{
+	InterchangeResultsBrowserWidgetPtr = InArgs._InterchangeResultsBrowserWidget;
+	Item = InArgs._Item;
+
+	SMultiColumnTableRow<UInterchangeResult*>::Construct(FSuperRowType::FArguments(), InOwnerTableView);
+}
+
+
+static FName GetImageFromResultType(EInterchangeResultType ResultType)
+{
+	switch (ResultType)
+	{
+	case EInterchangeResultType::Success: return FName("InterchangeResultsBrowser.ResultType.Success");
+	case EInterchangeResultType::Warning: return FName("InterchangeResultsBrowser.ResultType.Warning");
+	case EInterchangeResultType::Error: return FName("InterchangeResultsBrowser.ResultType.Error");
+	}
+
+	return FName();
+}
+
+
+TSharedRef<SWidget> SInterchangeResultsBrowserListRow::GenerateWidgetForColumn(const FName& ColumnName)
+{
+	if (TSharedPtr<SInterchangeResultsBrowserWindow> InterchangeResultsBrowserWidget = InterchangeResultsBrowserWidgetPtr.Pin())
+	{
+		const FMargin RowPadding(8, 2, 2, 0);
+
+		if (ColumnName == SResultsBrowserWidgetDefs::ColumnID_ResultType)
+		{
+			return SNew(SHorizontalBox)
+				+ SHorizontalBox::Slot()
+				.Padding(RowPadding)
+				.AutoWidth()
+				[
+					SNew(SImage)
+					.Image(FSlateIconFinder::FindIcon(GetImageFromResultType(Item->GetResultType())).GetOptionalIcon())
+				];
+		}
+		else if (ColumnName == SResultsBrowserWidgetDefs::ColumnID_Source)
+		{
+			return SNew(SHorizontalBox)
+				+ SHorizontalBox::Slot()
+				.Padding(RowPadding)
+				[
+					SNew(STextBlock)
+					.Text(FText::FromString(FPaths::GetCleanFilename(Item->SourceAssetName)))
+					.ToolTipText(FText::FromString(Item->SourceAssetName))
+				];
+		}
+		else if (ColumnName == SResultsBrowserWidgetDefs::ColumnID_Dest)
+		{
+			return SNew(SHorizontalBox)
+				+ SHorizontalBox::Slot()
+				.Padding(RowPadding)
+				[
+					SNew(STextBlock)
+					.Text(FText::FromString(FPackageName::ObjectPathToObjectName(Item->DestinationAssetName)))
+					.ToolTipText(FText::FromString(Item->DestinationAssetName))
+				];
+		}
+		else if (ColumnName == SResultsBrowserWidgetDefs::ColumnID_AssetType)
+		{
+			return SNew(SHorizontalBox)
+				+ SHorizontalBox::Slot()
+				.Padding(RowPadding)
+				[
+					SNew(STextBlock)
+					.Text(FText::FromString(Item->AssetType ? Item->AssetType->GetName() : FString()))
+				];
+		}
+		else if (ColumnName == SResultsBrowserWidgetDefs::ColumnID_Message)
+		{
+			return SNew(SHorizontalBox)
+				+ SHorizontalBox::Slot()
+				.Padding(RowPadding)
+				[
+					SNew(STextBlock)
+					.Text(Item->GetText())
+					.ToolTipText(Item->GetText())
+				];
+		}
+	}
+
+	return SNullWidget::NullWidget;
+}
+
+
+#undef LOCTEXT_NAMESPACE
diff --git a/Engine/Plugins/Interchange/Editor/Source/ResultsBrowser/Private/SInterchangeResultsBrowserWindow.h b/Engine/Plugins/Interchange/Editor/Source/ResultsBrowser/Private/SInterchangeResultsBrowserWindow.h
new file mode 100644
--- /dev/null
+++ b/Engine/Plugins/Interchange/Editor/Source/ResultsBrowser/Private/SInterchangeResultsBrowserWindow.h
@@ -0,0 +1,103 @@
+// Copyright Epic Games, Inc. All Rights Reserved.
+
+#pragma once
+
+#include "CoreMinimal.h"
+#include "InputCoreTypes.h"
+#include "Input/Reply.h"
+#include "InterchangeResultsContainer.h"
+#include "Styling/SlateColor.h"
+#include "Widgets/DeclarativeSyntaxSupport.h"
+#include "Widgets/SCompoundWidget.h"
+#include "Widgets/Docking/SDockTab.h"
+#include "Widgets/Views/SListView.h"
+#include "Widgets/Views/STableViewBase.h"
+#include "Widgets/Views/STableRow.h"
+
+
+class SInterchangeResultsBrowserWindow : public SCompoundWidget
+{
+public:
+	DECLARE_DELEGATE_OneParam(FOnFilterChangedState, bool)
+
+	SLATE_BEGIN_ARGS(SInterchangeResultsBrowserWindow)
+		: _OwnerTab()
+		, _InterchangeResultsContainer()
+		, _IsFiltered(false)
+		{}
+
+		SLATE_ARGUMENT(TWeakPtr<SDockTab>, OwnerTab)
+		SLATE_ARGUMENT(UInterchangeResultsContainer*, InterchangeResultsContainer)
+		SLATE_ARGUMENT(bool, IsFiltered)
+		SLATE_EVENT(FOnFilterChangedState, OnFilterChangedState)
+	SLATE_END_ARGS()
+
+public:
+
+	SInterchangeResultsBrowserWindow();
+	~SInterchangeResultsBrowserWindow();
+
+	void Construct(const FArguments& InArgs);
+	virtual bool SupportsKeyboardFocus() const override { return true; }
+
+	void Set(UInterchangeResultsContainer* Data);
+	void CloseErrorBrowser();
+
+	FReply OnCloseDialog()
+	{
+		CloseErrorBrowser();
+		return FReply::Handled();
+	}
+
+	virtual FReply OnKeyDown(const FGeometry& MyGeometry, const FKeyEvent& InKeyEvent) override
+	{
+		if (InKeyEvent.GetKey() == EKeys::Escape)
+		{
+			return OnCloseDialog();
+		}
+
+		return FReply::Unhandled();
+	}
+
+private:
+
+	/** Called by SListView to get a widget corresponding to the supplied item */
+	TSharedRef<ITableRow> OnGenerateRowForList(UInterchangeResult* Item, const TSharedRef<STableViewBase>& OwnerTable);
+
+	EColumnSortMode::Type GetColumnSortMode(const FName ColumnId) const;
+
+	void OnColumnSortModeChanged(const EColumnSortPriority::Type SortPriority, const FName& ColumnId, const EColumnSortMode::Type InSortMode);
+	void OnSelectionChanged(UInterchangeResult* InItem, ESelectInfo::Type SelectInfo);
+	void OnFilterStateChanged(ECheckBoxState NewState);
+	void RepopulateItems();
+	void RequestSort();
+
+
+	TWeakPtr<SDockTab> OwnerTab;
+	UInterchangeResultsContainer* ResultsContainer;
+	TArray<UInterchangeResult*> FilteredResults;
+	TSharedPtr<SListView<UInterchangeResult*>> ListView;
+	TSharedPtr<STextBlock> Description;
+	FName SortByColumn;
+	EColumnSortMode::Type SortMode;
+	FOnFilterChangedState OnFilterChangedState;
+	bool bIsFiltered;
+};
+
+
+class SInterchangeResultsBrowserListRow : public SMultiColumnTableRow<UInterchangeResult*>
+{
+public:
+	SLATE_BEGIN_ARGS(SInterchangeResultsBrowserListRow) {}
+		SLATE_ARGUMENT(TSharedPtr<SInterchangeResultsBrowserWindow>, InterchangeResultsBrowserWidget)
+		SLATE_ARGUMENT(UInterchangeResult*, Item)
+	SLATE_END_ARGS()
+
+	void Construct(const FArguments& InArgs, const TSharedRef<STableViewBase>& InOwnerTableView);
+	virtual TSharedRef<SWidget> GenerateWidgetForColumn(const FName& ColumnName) override;
+
+private:
+
+	TWeakPtr<SInterchangeResultsBrowserWindow> InterchangeResultsBrowserWidgetPtr;
+	UInterchangeResult* Item;
+};
diff --git a/Engine/Plugins/Interchange/Runtime/Source/Import/Private/Texture/InterchangeIESTranslator.cpp b/Engine/Plugins/Interchange/Runtime/Source/Import/Private/Texture/InterchangeIESTranslator.cpp
new file mode 100644
--- /dev/null
+++ b/Engine/Plugins/Interchange/Runtime/Source/Import/Private/Texture/InterchangeIESTranslator.cpp
@@ -0,0 +1,99 @@
+// Copyright Epic Games, Inc. All Rights Reserved. 
+
+#include "Texture/InterchangeIESTranslator.h"
+
+#include "Engine/Texture.h"
+#include "Engine/TextureLightProfile.h"
+#include "IESConverter.h"
+#include "InterchangeImportLog.h"
+#include "Misc/ConfigCacheIni.h"
+#include "Misc/FileHelper.h"
+#include "Misc/Paths.h"
+#include "Texture/TextureTranslatorUtilities.h"
+
+static bool GInterchangeEnableIESImport = true;
+static FAutoConsoleVariableRef CCvarInterchangeEnableIESImport(
+	TEXT("Interchange.FeatureFlags.Import.IES"),
+	GInterchangeEnableIESImport,
+	TEXT("Whether IES support is enabled."),
+	ECVF_Default);
+
+TArray<FString> UInterchangeIESTranslator::GetSupportedFormats() const
+{
+	if (GInterchangeEnableIESImport || GIsAutomationTesting)
+	{
+		TArray<FString> Formats{ TEXT("ies;IES light profile") };
+		return Formats;
+	}
+	else
+	{
+		return TArray<FString>{};
+	}
+}
+
+bool UInterchangeIESTranslator::Translate(UInterchangeBaseNodeContainer& BaseNodeContainer) const
+{
+	return UE::Interchange::FTextureTranslatorUtilities::GenericTextureLightProfileTranslate(GetSourceData(), BaseNodeContainer);
+}
+
+TOptional<UE::Interchange::FImportLightProfile> UInterchangeIESTranslator::GetLightProfilePayloadData(const UInterchangeSourceData* PayloadSourceData, const FString& PayLoadKey) const
+{
+	check(GetSourceData() == PayloadSourceData);
+
+	if (!GetSourceData())
+	{
+		UE_LOG(LogInterchangeImport, Error, TEXT("Failed to import IES, bad source data."));
+		return {};
+	}
+
+	TArray64<uint8> SourceDataBuffer;
+	FString Filename = GetSourceData()->GetFilename();
+
+	//Make sure the key fit the filename, The key should always be valid
+	if (!Filename.Equals(PayLoadKey))
+	{
+		UE_LOG(LogInterchangeImport, Error, TEXT("Failed to import IES, wrong payload key. [%s]"), *Filename);
+		return {};
+	}
+
+	if (!FPaths::FileExists(Filename))
+	{
+		UE_LOG(LogInterchangeImport, Error, TEXT("Failed to import IES, cannot open file. [%s]"), *Filename);
+		return {};
+	}
+
+	if (!FFileHelper::LoadFileToArray(SourceDataBuffer, *Filename))
+	{
+		UE_LOG(LogInterchangeImport, Error, TEXT("Failed to import IES, cannot load file content into an array. [%s]"), *Filename);
+		return {};
+	}
+
+	const uint8* Buffer = SourceDataBuffer.GetData();
+
+	FIESConverter IESConverter(Buffer, SourceDataBuffer.Num());
+
+	if(IESConverter.IsValid())
+	{
+		UE::Interchange::FImportLightProfile Payload;
+
+		Payload.Init2DWithParams(
+			IESConverter.GetWidth(),
+			IESConverter.GetHeight(),
+			TSF_RGBA16F,
+			false
+			);
+
+		Payload.CompressionSettings = TC_HDR;
+		Payload.Brightness = IESConverter.GetBrightness();
+		Payload.TextureMultiplier = IESConverter.GetMultiplier();
+
+
+		const TArray<uint8>&  RawData = IESConverter.GetRawData();;
+
+		FPlatformMemory::Memcpy(Payload.RawData.GetData(), RawData.GetData(), Payload.RawData.GetSize());
+
+		return Payload;
+	}
+
+	return {};
+}
diff --git a/Engine/Plugins/Interchange/Runtime/Source/Import/Private/Texture/InterchangePSDTranslator.cpp b/Engine/Plugins/Interchange/Runtime/Source/Import/Private/Texture/InterchangePSDTranslator.cpp
new file mode 100644
--- /dev/null
+++ b/Engine/Plugins/Interchange/Runtime/Source/Import/Private/Texture/InterchangePSDTranslator.cpp
@@ -0,0 +1,487 @@
+// Copyright Epic Games, Inc. All Rights Reserved.
+
+#include "Texture/InterchangePSDTranslator.h"
+
+#include "Engine/Texture.h"
+#include "Engine/Texture2D.h"
+#include "IImageWrapper.h"
+#include "IImageWrapperModule.h"
+#include "InterchangeImportLog.h"
+#include "InterchangeTextureNode.h"
+#include "Misc/FileHelper.h"
+#include "Misc/Paths.h"
+#include "Misc/ScopedSlowTask.h"
+#include "Modules/ModuleManager.h"
+#include "Nodes/InterchangeBaseNodeContainer.h"
+#include "Texture/TextureTranslatorUtilities.h"
+
+static bool GInterchangeEnablePSDImport = true;
+static FAutoConsoleVariableRef CCvarInterchangeEnablePSDImport(
+	TEXT("Interchange.FeatureFlags.Import.PSD"),
+	GInterchangeEnablePSDImport,
+	TEXT("Whether PSD support is enabled."),
+	ECVF_Default);
+
+//////////////////////////////////////////////////////////////////////////
+// PSD helper local function
+namespace UE
+{
+	namespace Interchange
+	{
+		namespace Private
+		{
+#pragma pack(push,1)
+			struct FPSDFileHeader
+			{
+				int32	Signature;		// 8BPS
+				int16	Version;		// Version
+				int16	nChannels;		// Number of Channels (3=RGB) (4=RGBA)
+				int32	Height;			// Number of Image Rows
+				int32	Width;			// Number of Image Columns
+				int16	Depth;			// Number of Bits per Channel
+				int16	Mode;			// Image Mode (0=Bitmap)(1=Grayscale)(2=Indexed)(3=RGB)(4=CYMK)(7=Multichannel)
+				uint8	Pad[6];			// Padding
+
+				/**
+				 * @return Whether file has a valid signature
+				 */
+				bool IsValid() const
+				{
+					// Fail on bad signature
+					return (Signature == 0x38425053);
+				}
+
+				/**
+				 * @return Whether file has a supported version
+				 */
+				bool IsSupported() const
+				{
+					// Fail on bad version
+					if (Version != 1)
+						return false;
+
+					// Fail on anything other than 1, 3 or 4 channels
+					if ((nChannels != 1) && (nChannels != 3) && (nChannels != 4))
+						return false;
+
+					// Fail on anything other than 8 Bits/channel or 16 Bits/channel  
+					if ((Depth != 8) && (Depth != 16))
+						return false;
+
+					// Fail on anything other than Grayscale and RGB
+					// We can add support for indexed later if needed.
+					if (Mode != 1 && Mode != 3)
+						return false;
+
+					return true;
+				}
+			};
+#pragma pack(pop)
+
+
+			static bool ReadData(uint8* pOut, const uint8*& pBuffer, FPSDFileHeader& Info)
+			{
+				const uint8* pPlane = nullptr;
+				const uint8* pRowTable = nullptr;
+				int32 iPlane;
+				int16 CompressionType;
+				int32 iPixel;
+				int32 iRow;
+				int32 CompressedBytes;
+				int32 iByte;
+				int32 Count;
+				uint8 Value;
+
+				// Double check to make sure this is a valid request
+				if (!Info.IsValid() || !Info.IsSupported())
+				{
+					return false;
+				}
+
+				const uint8* pCur = pBuffer + sizeof(FPSDFileHeader);
+				int32 NPixels = Info.Width * Info.Height;
+
+				int32 ClutSize = ((int32)pCur[0] << 24) +
+					((int32)pCur[1] << 16) +
+					((int32)pCur[2] << 8) +
+					((int32)pCur[3] << 0);
+				pCur += 4 + ClutSize;
+
+				// Skip Image Resource Section
+				int32 ImageResourceSize = ((int32)pCur[0] << 24) +
+					((int32)pCur[1] << 16) +
+					((int32)pCur[2] << 8) +
+					((int32)pCur[3] << 0);
+				pCur += 4 + ImageResourceSize;
+
+				// Skip Layer and Mask Section
+				int32 LayerAndMaskSize = ((int32)pCur[0] << 24) +
+					((int32)pCur[1] << 16) +
+					((int32)pCur[2] << 8) +
+					((int32)pCur[3] << 0);
+				pCur += 4 + LayerAndMaskSize;
+
+				// Determine number of bytes per pixel
+				int32 BytesPerPixel = 3;
+				const int32 BytesPerChannel = (Info.Depth / 8);
+				switch (Info.Mode)
+				{
+					case 1: // 'GrayScale'
+						BytesPerPixel = BytesPerChannel;
+						break;
+					case 2:
+						BytesPerPixel = 1;
+						return false;  // until we support indexed...
+						break;
+					case 3: // 'RGBColor'
+						if (Info.nChannels == 3)
+							BytesPerPixel = 3 * BytesPerChannel;
+						else
+							BytesPerPixel = 4 * BytesPerChannel;
+						break;
+					default:
+						return false;
+						break;
+				}
+
+				// Get Compression Type
+				CompressionType = ((int32)pCur[0] << 8) + ((int32)pCur[1] << 0);
+				pCur += 2;
+
+				// Fail on 16 Bits/channel with RLE. This can occur when the file is not saved with 'Maximize Compatibility'. Compression doesn't appear to be standard.
+				if (CompressionType == 1 && Info.Depth == 16)
+				{
+					return false;
+				}
+
+				// If no alpha channel, set alpha to opaque (255 or 65536).
+				if (Info.nChannels != 4)
+				{
+					if (Info.Depth == 8)
+					{
+						const uint32 Channels = 4;
+						const uint32 BufferSize = Info.Width * Info.Height * Channels * sizeof(uint8);
+						FMemory::Memset(pOut, 0xff, BufferSize);
+					}
+					else if (Info.Depth == 16)
+					{
+						const uint32 Channels = 4;
+						const uint32 BufferSize = Info.Width * Info.Height * Channels * sizeof(uint16);
+						FMemory::Memset(pOut, 0xff, BufferSize);
+					}
+				}
+
+				// Uncompressed?
+				if (CompressionType == 0)
+				{
+					if (Info.Depth == 8)
+					{
+						FColor* Dest = (FColor*)pOut;
+						for (int32 Pixel = 0; Pixel < NPixels; Pixel++)
+						{
+							if (Info.nChannels == 1)
+							{
+								Dest[Pixel].R = pCur[NPixels + Pixel];
+								Dest[Pixel].G = pCur[NPixels + Pixel];
+								Dest[Pixel].B = pCur[NPixels + Pixel];
+							}
+							else
+							{
+								Dest[Pixel].R = pCur[NPixels * 0 + Pixel];
+								Dest[Pixel].G = pCur[NPixels * 1 + Pixel];
+								Dest[Pixel].B = pCur[NPixels * 2 + Pixel];
+								if (Info.nChannels == 4)
+								{
+									Dest[Pixel].A = pCur[NPixels * 3 + Pixel];
+								}
+							}
+						}
+					}
+					else if (Info.Depth == 16)
+					{
+						uint32 SrcOffset = 0;
+
+						if (Info.nChannels == 1)
+						{
+							uint16* Dest = (uint16*)pOut;
+							uint32 ChannelOffset = 0;
+
+							for (int32 Pixel = 0; Pixel < NPixels; Pixel++)
+							{
+								Dest[ChannelOffset + 0] = ((pCur[SrcOffset] << 8) + (pCur[SrcOffset + 1] << 0));
+								Dest[ChannelOffset + 1] = ((pCur[SrcOffset] << 8) + (pCur[SrcOffset + 1] << 0));
+								Dest[ChannelOffset + 2] = ((pCur[SrcOffset] << 8) + (pCur[SrcOffset + 1] << 0));
+
+								//Increment offsets
+								ChannelOffset += 4;
+								SrcOffset += BytesPerChannel;
+							}
+						}
+						else
+						{
+							// Loop through the planes	
+							for (iPlane = 0; iPlane < Info.nChannels; iPlane++)
+							{
+								uint16* Dest = (uint16*)pOut;
+								uint32 ChannelOffset = iPlane;
+
+								for (int32 Pixel = 0; Pixel < NPixels; Pixel++)
+								{
+									Dest[ChannelOffset] = ((pCur[SrcOffset] << 8) + (pCur[SrcOffset + 1] << 0));
+
+									//Increment offsets
+									ChannelOffset += 4;
+									SrcOffset += BytesPerChannel;
+								}
+							}
+						}
+					}
+				}
+				// RLE?
+				else if (CompressionType == 1)
+				{
+					// Setup RowTable
+					pRowTable = pCur;
+					pCur += Info.nChannels * Info.Height * 2;
+
+					FColor* Dest = (FColor*)pOut;
+
+					// Loop through the planes
+					for (iPlane = 0; iPlane < Info.nChannels; iPlane++)
+					{
+						int32 iWritePlane = iPlane;
+						if (iWritePlane > BytesPerPixel - 1) iWritePlane = BytesPerPixel - 1;
+
+						// Loop through the rows
+						for (iRow = 0; iRow < Info.Height; iRow++)
+						{
+							// Load a row
+							CompressedBytes = (pRowTable[(iPlane * Info.Height + iRow) * 2] << 8) +
+								(pRowTable[(iPlane * Info.Height + iRow) * 2 + 1] << 0);
+
+							// Setup Plane
+							pPlane = pCur;
+							pCur += CompressedBytes;
+
+							// Decompress Row
+							iPixel = 0;
+							iByte = 0;
+							while ((iPixel < Info.Width) && (iByte < CompressedBytes))
+							{
+								int8 code = (int8)pPlane[iByte++];
+
+								// Is it a repeat?
+								if (code < 0)
+								{
+									Count = -(int32)code + 1;
+									Value = pPlane[iByte++];
+									while (Count-- > 0)
+									{
+										int32 idx = (iPixel)+(iRow * Info.Width);
+										if (Info.nChannels == 1)
+										{
+											Dest[idx].R = Value;
+											Dest[idx].G = Value;
+											Dest[idx].B = Value;
+										}
+										else
+										{
+											switch (iWritePlane)
+											{
+												case 0: Dest[idx].R = Value; break;
+												case 1: Dest[idx].G = Value; break;
+												case 2: Dest[idx].B = Value; break;
+												case 3: Dest[idx].A = Value; break;
+											}
+										}
+										iPixel++;
+									}
+								}
+								// Must be a literal then
+								else
+								{
+									Count = (int32)code + 1;
+									while (Count-- > 0)
+									{
+										Value = pPlane[iByte++];
+										int32 idx = (iPixel)+(iRow * Info.Width);
+
+										if (Info.nChannels == 1)
+										{
+											Dest[idx].R = Value;
+											Dest[idx].G = Value;
+											Dest[idx].B = Value;
+										}
+										else
+										{
+											switch (iWritePlane)
+											{
+												case 0: Dest[idx].R = Value; break;
+												case 1: Dest[idx].G = Value; break;
+												case 2: Dest[idx].B = Value; break;
+												case 3: Dest[idx].A = Value; break;
+											}
+										}
+										iPixel++;
+									}
+								}
+							}
+
+							// Confirm that we decoded the right number of bytes
+							check(iByte == CompressedBytes);
+							check(iPixel == Info.Width);
+						}
+					}
+				}
+				else
+					return false;
+
+				// Success!
+				return true;
+			}
+
+			static void GetPSDHeader(const uint8* Buffer, FPSDFileHeader& Info)
+			{
+				Info.Signature = ((int32)Buffer[0] << 24) +
+					((int32)Buffer[1] << 16) +
+					((int32)Buffer[2] << 8) +
+					((int32)Buffer[3] << 0);
+
+				Info.Version = ((int32)Buffer[4] << 8) +
+					((int32)Buffer[5] << 0);
+
+				Info.nChannels = ((int32)Buffer[12] << 8) +
+					((int32)Buffer[13] << 0);
+
+				Info.Height = ((int32)Buffer[14] << 24) +
+					((int32)Buffer[15] << 16) +
+					((int32)Buffer[16] << 8) +
+					((int32)Buffer[17] << 0);
+
+				Info.Width = ((int32)Buffer[18] << 24) +
+					((int32)Buffer[19] << 16) +
+					((int32)Buffer[20] << 8) +
+					((int32)Buffer[21] << 0);
+
+				Info.Depth = ((int32)Buffer[22] << 8) +
+					((int32)Buffer[23] << 0);
+
+				Info.Mode = ((int32)Buffer[24] << 8) +
+					((int32)Buffer[25] << 0);
+			}
+
+		}//ns Private
+	}//ns Interchange
+}//ns UE
+
+TArray<FString> UInterchangePSDTranslator::GetSupportedFormats() const
+{
+	if (GInterchangeEnablePSDImport || GIsAutomationTesting)
+	{
+		TArray<FString> Formats{ TEXT("psd;Photoshop Document") };
+		return Formats;
+	}
+	else
+	{
+		return TArray<FString>{};
+	}
+}
+
+bool UInterchangePSDTranslator::Translate(UInterchangeBaseNodeContainer& BaseNodeContainer) const
+{
+	return UE::Interchange::FTextureTranslatorUtilities::Generic2DTextureTranslate(GetSourceData(), BaseNodeContainer);
+}
+
+TOptional<UE::Interchange::FImportImage> UInterchangePSDTranslator::GetTexturePayloadData(const UInterchangeSourceData* PayloadSourceData, const FString& PayLoadKey) const
+{
+	check(PayloadSourceData == GetSourceData());
+
+	if (!GetSourceData())
+	{
+		UE_LOG(LogInterchangeImport, Error, TEXT("Failed to import PSD, bad source data."));
+		return TOptional<UE::Interchange::FImportImage>();
+	}
+
+	TArray64<uint8> SourceDataBuffer;
+	FString Filename = GetSourceData()->GetFilename();
+
+	//Make sure the key fit the filename, The key should always be valid
+	if (!Filename.Equals(PayLoadKey))
+	{
+		UE_LOG(LogInterchangeImport, Error, TEXT("Failed to import PSD, wrong payload key. [%s]"), *Filename);
+		return TOptional<UE::Interchange::FImportImage>();
+	}
+
+	if (!FPaths::FileExists(Filename))
+	{
+		UE_LOG(LogInterchangeImport, Error, TEXT("Failed to import PSD, cannot open file. [%s]"), *Filename);
+		return TOptional<UE::Interchange::FImportImage>();
+	}
+
+	if (!FFileHelper::LoadFileToArray(SourceDataBuffer, *Filename))
+	{
+		UE_LOG(LogInterchangeImport, Error, TEXT("Failed to import PSD, cannot load file content into an array. [%s]"), *Filename);
+		return TOptional<UE::Interchange::FImportImage>();
+	}
+
+	const uint8* Buffer = SourceDataBuffer.GetData();
+	const uint8* BufferEnd = Buffer + SourceDataBuffer.Num();
+
+	const int32 Length = BufferEnd - Buffer;
+
+	//
+	// PSD File
+	//
+	UE::Interchange::Private::FPSDFileHeader PSDHeader;
+	if (Length < sizeof(UE::Interchange::Private::FPSDFileHeader))
+	{
+		return TOptional<UE::Interchange::FImportImage>();
+	}
+
+	UE::Interchange::Private::GetPSDHeader(Buffer, PSDHeader);
+
+	if (!PSDHeader.IsValid())
+	{
+		return TOptional<UE::Interchange::FImportImage>();
+	}
+
+	if (!PSDHeader.IsSupported())
+	{
+		UE_LOG(LogInterchangeImport, Error, TEXT("Format of this PSD is not supported. Only Grayscale and RGBColor PSD images are currently supported, in 8-bit or 16-bit."));
+		return TOptional<UE::Interchange::FImportImage>();
+	}
+
+	// Select the texture's source format
+	ETextureSourceFormat TextureFormat = TSF_Invalid;
+	if (PSDHeader.Depth == 8)
+	{
+		TextureFormat = TSF_BGRA8;
+	}
+	else if (PSDHeader.Depth == 16)
+	{
+		TextureFormat = TSF_RGBA16;
+	}
+
+	if (TextureFormat == TSF_Invalid)
+	{
+		UE_LOG(LogInterchangeImport, Error, TEXT("PSD file contains data in an unsupported format."));
+		return TOptional<UE::Interchange::FImportImage>();
+	}
+
+	UE::Interchange::FImportImage PayloadData;
+
+	// The PSD is supported. Load it up.        
+	PayloadData.Init2DWithOneMip(
+		PSDHeader.Width,
+		PSDHeader.Height,
+		TextureFormat
+	);
+
+	if (!UE::Interchange::Private::ReadData(static_cast<uint8*>(PayloadData.RawData.GetData()), Buffer, PSDHeader))
+	{
+		UE_LOG(LogInterchangeImport, Error, TEXT("Failed to read this PSD"));
+		return TOptional<UE::Interchange::FImportImage>();
+	}
+
+	return PayloadData;
+}
diff --git a/Engine/Plugins/Interchange/Runtime/Source/Import/Public/Mesh/InterchangeStaticMeshPayload.h b/Engine/Plugins/Interchange/Runtime/Source/Import/Public/Mesh/InterchangeStaticMeshPayload.h
new file mode 100644
--- /dev/null
+++ b/Engine/Plugins/Interchange/Runtime/Source/Import/Public/Mesh/InterchangeStaticMeshPayload.h
@@ -0,0 +1,18 @@
+// Copyright Epic Games, Inc. All Rights Reserved.
+
+#pragma once
+
+#include "CoreMinimal.h"
+#include "MeshDescription.h"
+
+namespace UE
+{
+	namespace Interchange
+	{
+		struct FStaticMeshPayloadData
+		{
+			FMeshDescription MeshDescription;
+		};
+
+	} // namespace Interchange
+} // namespace UE
diff --git a/Engine/Plugins/Interchange/Runtime/Source/Import/Public/Texture/InterchangePCXTranslator.h b/Engine/Plugins/Interchange/Runtime/Source/Import/Public/Texture/InterchangePCXTranslator.h
new file mode 100644
--- /dev/null
+++ b/Engine/Plugins/Interchange/Runtime/Source/Import/Public/Texture/InterchangePCXTranslator.h
@@ -0,0 +1,47 @@
+// Copyright Epic Games, Inc. All Rights Reserved.
+
+#pragma once
+
+#include "CoreMinimal.h"
+#include "UObject/Object.h"
+#include "UObject/ObjectMacros.h"
+#include "InterchangeTranslatorBase.h"
+#include "Nodes/InterchangeBaseNodeContainer.h"
+#include "Texture/InterchangeTexturePayloadData.h"
+#include "Texture/InterchangeTexturePayloadInterface.h"
+
+#include "InterchangePCXTranslator.generated.h"
+
+UCLASS(BlueprintType, Experimental)
+class INTERCHANGEIMPORT_API UInterchangePCXTranslator : public UInterchangeTranslatorBase, public IInterchangeTexturePayloadInterface
+{
+	GENERATED_BODY()
+public:
+
+	virtual TArray<FString> GetSupportedFormats() const override;
+	virtual EInterchangeTranslatorAssetType GetSupportedAssetTypes() const override { return EInterchangeTranslatorAssetType::Textures; }
+
+	/**
+	 * Translates the associated source data into a node hold by the specified nodes container.
+	 *
+	 * @param BaseNodeContainer - The unreal objects descriptions container where to put the translated source data.
+	 * @return true if the translator can translate the source data, false otherwise.
+	 */
+	virtual bool Translate(UInterchangeBaseNodeContainer& BaseNodeContainer) const override;
+
+	/* IInterchangeTexturePayloadInterface Begin */
+
+	/**
+	 * Once the translation is done, the import process need a way to retrieve payload data.
+	 * This payload will be use by the factories to create the asset.
+	 *
+	 * @param PayloadSourceData - The source data containing the data to translate
+	 * @param PayloadKey - The key to retrieve the a particular payload contain into the specified source data.
+	 * @return a PayloadData containing the import image data. The TOptional will not be set if there is an error.
+	 */
+	virtual TOptional<UE::Interchange::FImportImage> GetTexturePayloadData(const UInterchangeSourceData* PayloadSourceData, const FString& PayLoadKey) const override;
+
+	/* IInterchangeTexturePayloadInterface End */
+};
+
+
diff --git a/Engine/Plugins/Interchange/Runtime/Source/Import/Public/Texture/InterchangePSDTranslator.h b/Engine/Plugins/Interchange/Runtime/Source/Import/Public/Texture/InterchangePSDTranslator.h
new file mode 100644
--- /dev/null
+++ b/Engine/Plugins/Interchange/Runtime/Source/Import/Public/Texture/InterchangePSDTranslator.h
@@ -0,0 +1,47 @@
+// Copyright Epic Games, Inc. All Rights Reserved.
+
+#pragma once
+
+#include "CoreMinimal.h"
+#include "UObject/Object.h"
+#include "UObject/ObjectMacros.h"
+#include "InterchangeTranslatorBase.h"
+#include "Nodes/InterchangeBaseNodeContainer.h"
+#include "Texture/InterchangeTexturePayloadData.h"
+#include "Texture/InterchangeTexturePayloadInterface.h"
+
+#include "InterchangePSDTranslator.generated.h"
+
+UCLASS(BlueprintType, Experimental)
+class INTERCHANGEIMPORT_API UInterchangePSDTranslator : public UInterchangeTranslatorBase, public IInterchangeTexturePayloadInterface
+{
+	GENERATED_BODY()
+public:
+
+	virtual TArray<FString> GetSupportedFormats() const override;
+	virtual EInterchangeTranslatorAssetType GetSupportedAssetTypes() const override { return EInterchangeTranslatorAssetType::Textures; }
+
+	/**
+	 * Translate the associated source data into a node hold by the specified nodes container.
+	 *
+	 * @param BaseNodeContainer - The unreal objects descriptions container where to put the translated source data.
+	 * @return true if the translator can translate the source data, false otherwise.
+	 */
+	virtual bool Translate(UInterchangeBaseNodeContainer& BaseNodeContainer) const override;
+
+	/* IInterchangeTexturePayloadInterface Begin */
+
+	/**
+	 * Once the translation is done, the import process need a way to retrieve payload data.
+	 * This payload will be use by the factories to create the asset.
+	 *
+	 * @param PayloadSourceData - The source data containing the data to translate
+	 * @param PayloadKey - The key to retrieve the a particular payload contain into the specified source data.
+	 * @return a PayloadData containing the import image data. The TOptional will not be set if there is an error.
+	 */
+	virtual TOptional<UE::Interchange::FImportImage> GetTexturePayloadData(const UInterchangeSourceData* PayloadSourceData, const FString& PayLoadKey) const override;
+
+	/* IInterchangeTexturePayloadInterface End */
+};
+
+
diff --git a/Engine/Plugins/Interchange/Runtime/Source/Messages/Public/Fbx/InterchangeFbxMessages.h b/Engine/Plugins/Interchange/Runtime/Source/Messages/Public/Fbx/InterchangeFbxMessages.h
new file mode 100644
--- /dev/null
+++ b/Engine/Plugins/Interchange/Runtime/Source/Messages/Public/Fbx/InterchangeFbxMessages.h
@@ -0,0 +1,117 @@
+// Copyright Epic Games, Inc. All Rights Reserved.
+
+#pragma once
+
+#include "CoreMinimal.h"
+#include "InterchangeResult.h"
+
+#include "InterchangeFbxMessages.generated.h"
+
+
+/**
+ * Base class for FBX parser warnings
+ */
+UCLASS()
+class INTERCHANGEMESSAGES_API UInterchangeResultMeshWarning : public UInterchangeResultWarning
+{
+	GENERATED_BODY()
+
+public:
+
+	UPROPERTY()
+	FString MeshName;
+};
+
+/**
+ * Base class for FBX parser warnings
+ */
+UCLASS()
+class INTERCHANGEMESSAGES_API UInterchangeResultTextureWarning : public UInterchangeResultWarning
+{
+	GENERATED_BODY()
+
+public:
+
+	UPROPERTY()
+	FString TextureName;
+};
+
+/**
+ * Base class for FBX parser errors
+ */
+UCLASS()
+class INTERCHANGEMESSAGES_API UInterchangeResultMeshError : public UInterchangeResultError
+{
+	GENERATED_BODY()
+
+public:
+
+	UPROPERTY()
+	FString MeshName;
+};
+
+
+/**
+ * A generic class for FBX parser warnings, with no additional metadata, and where the text is specified by the user
+ */
+UCLASS()
+class INTERCHANGEMESSAGES_API UInterchangeResultMeshWarning_Generic : public UInterchangeResultMeshWarning
+{
+	GENERATED_BODY()
+
+public:
+	virtual FText GetText() const override;
+
+	UPROPERTY()
+	FText Text;
+};
+
+
+/**
+ * A generic class for FBX parser errors, with no additional metadata, and where the text is specified by the user
+ */
+UCLASS()
+class INTERCHANGEMESSAGES_API UInterchangeResultMeshError_Generic : public UInterchangeResultMeshError
+{
+	GENERATED_BODY()
+
+public:
+	virtual FText GetText() const override;
+
+	UPROPERTY()
+	FText Text;
+};
+
+
+/**
+ * 
+ */
+UCLASS()
+class INTERCHANGEMESSAGES_API UInterchangeResultMeshWarning_TooManyUVs : public UInterchangeResultMeshWarning
+{
+	GENERATED_BODY()
+
+public:
+	virtual FText GetText() const override;
+
+	UPROPERTY()
+	int32 ExcessUVs;
+};
+
+/**
+ * A generic class for FBX parser warnings, with no additional metadata, and where the text is specified by the user
+ */
+UCLASS()
+class INTERCHANGEMESSAGES_API UInterchangeResultTextureWarning_TextureFileDoNotExist : public UInterchangeResultTextureWarning
+{
+	GENERATED_BODY()
+
+public:
+	virtual FText GetText() const override;
+
+	UPROPERTY()
+	FText Text;
+
+	UPROPERTY()
+	FString MaterialName;
+};
\ No newline at end of file
diff --git a/Engine/Plugins/Interchange/Runtime/Source/Parsers/Fbx/Private/FbxMesh.h b/Engine/Plugins/Interchange/Runtime/Source/Parsers/Fbx/Private/FbxMesh.h
new file mode 100644
--- /dev/null
+++ b/Engine/Plugins/Interchange/Runtime/Source/Parsers/Fbx/Private/FbxMesh.h
@@ -0,0 +1,128 @@
+// Copyright Epic Games, Inc. All Rights Reserved.
+
+#pragma once
+
+#include "CoreMinimal.h"
+#include "FbxAPI.h"
+#include "FbxHelper.h"
+#include "FbxInclude.h"
+
+/** Forward declarations */
+struct FMeshDescription;
+class UInterchangeBaseNodeContainer;
+class UInterchangeMeshNode;
+
+namespace UE
+{
+	namespace Interchange
+	{
+		namespace Private
+		{
+			class FFbxMesh;
+
+			class FMeshDescriptionImporter
+			{
+			public:
+				FMeshDescriptionImporter(FFbxParser& InParser, FMeshDescription* InMeshDescription, FbxScene* InSDKScene, FbxGeometryConverter* InSDKGeometryConverter);
+				
+				/*
+				 * Fill the mesh description using the Mesh parameter.
+				 */
+				bool FillStaticMeshDescriptionFromFbxMesh(FbxMesh* Mesh);
+				
+				/*
+				 * Fill the mesh description using the Mesh parameter and also fill the OutJointNodeUniqueIDs so the MeshDescription bone Index can be map to the correct interchange joint scene node.
+				 */
+				bool FillSkinnedMeshDescriptionFromFbxMesh(FbxMesh* Mesh, TArray<FString>& OutJointUniqueNames);
+
+				/*
+				 * Fill the mesh description using the Shape parameter.
+				 */
+				bool FillMeshDescriptionFromFbxShape(FbxShape* Shape);
+
+				/**
+				 * Add messages to the message log
+				 */
+				template <typename T>
+				T* AddMessage(FbxGeometryBase* FbxNode) const
+				{
+					T* Item = Parser.AddMessage<T>();
+					Item->MeshName = FFbxHelper::GetMeshName(FbxNode);
+					Item->InterchangeKey = FFbxHelper::GetMeshUniqueID(FbxNode);
+					return Item;
+				}
+
+			private:
+				
+				enum class EMeshType : uint8
+				{
+					None = 0, //No mesh type to import
+					Static = 1, //static mesh
+					Skinned = 2, //skinned mesh with joints
+				};
+
+				bool FillMeshDescriptionFromFbxMesh(FbxMesh* Mesh, TArray<FString>& OutJointUniqueNames, EMeshType MeshType);
+				bool IsOddNegativeScale(FbxAMatrix& TotalMatrix);
+				
+				//TODO move the real function from RenderCore to FVector, so we do not have to add render core to compute such a simple thing
+				float FbxGetBasisDeterminantSign(const FVector& XAxis, const FVector& YAxis, const FVector& ZAxis)
+				{
+					FMatrix Basis(
+						FPlane(XAxis, 0),
+						FPlane(YAxis, 0),
+						FPlane(ZAxis, 0),
+						FPlane(0, 0, 0, 1)
+					);
+					return (Basis.Determinant() < 0) ? -1.0f : +1.0f;
+				}
+
+				FFbxParser& Parser;
+				FMeshDescription* MeshDescription;
+				FbxScene* SDKScene;
+				FbxGeometryConverter* SDKGeometryConverter;
+				bool bInitialized = false;
+			};
+
+			class FMeshPayloadContext : public FPayloadContextBase
+			{
+			public:
+				virtual ~FMeshPayloadContext() {}
+				virtual FString GetPayloadType() const override { return TEXT("Mesh-PayloadContext"); }
+				virtual bool FetchPayloadToFile(FFbxParser& Parser, const FString& PayloadFilepath) override;
+				bool bIsSkinnedMesh = false;
+				FbxMesh* Mesh = nullptr;
+				FbxScene* SDKScene = nullptr;
+				FbxGeometryConverter* SDKGeometryConverter = nullptr;
+			};
+
+			class FMorphTargetPayloadContext : public FPayloadContextBase
+			{
+			public:
+				virtual ~FMorphTargetPayloadContext() {}
+				virtual FString GetPayloadType() const override { return TEXT("MorphTarget-PayloadContext"); }
+				virtual bool FetchPayloadToFile(FFbxParser& Parser, const FString& PayloadFilepath) override;
+				FbxShape* Shape = nullptr;
+				FbxScene* SDKScene = nullptr;
+				FbxGeometryConverter* SDKGeometryConverter = nullptr;
+			};
+
+			class FFbxMesh
+			{
+			public:
+				explicit FFbxMesh(FFbxParser& InParser)
+					: Parser(InParser)
+				{}
+
+				void AddAllMeshes(FbxScene* SDKScene, FbxGeometryConverter* SDKGeometryConverter, UInterchangeBaseNodeContainer& NodeContainer, TMap<FString, TSharedPtr<FPayloadContextBase>>& PayloadContexts);
+				static bool GetGlobalJointBindPoseTransform(FbxScene* SDKScene, FbxNode* Joint, FbxAMatrix& GlobalBindPoseJointMatrix);
+			protected:
+				/** Add joint to the interchange mesh node joint dependencies. Return false if there is no valid joint (not a valid skinned mesh) */
+				bool ExtractSkinnedMeshNodeJoints(FbxScene* SDKScene, UInterchangeBaseNodeContainer& NodeContainer, FbxMesh* Mesh, UInterchangeMeshNode* MeshNode);
+				UInterchangeMeshNode* CreateMeshNode(UInterchangeBaseNodeContainer& NodeContainer, const FString& NodeName, const FString& NodeUniqueID);
+
+			private:
+				FFbxParser& Parser;
+			};
+		}//ns Private
+	}//ns Interchange
+}//ns UE
diff --git a/Engine/Source/Runtime/Interchange/Engine/Private/Tasks/InterchangeTaskCreateAsset.cpp b/Engine/Source/Runtime/Interchange/Engine/Private/Tasks/InterchangeTaskCreateAsset.cpp
new file mode 100644
--- /dev/null
+++ b/Engine/Source/Runtime/Interchange/Engine/Private/Tasks/InterchangeTaskCreateAsset.cpp
@@ -0,0 +1,408 @@
+// Copyright Epic Games, Inc. All Rights Reserved.
+#include "InterchangeTaskCreateAsset.h"
+
+#include "Async/TaskGraphInterfaces.h"
+#include "CoreMinimal.h"
+#include "InterchangeAssetImportData.h"
+#include "InterchangeEngineLogPrivate.h"
+#include "InterchangeFactoryBase.h"
+#include "InterchangeManager.h"
+#include "InterchangeResult.h"
+#include "InterchangeSourceData.h"
+#include "InterchangeTranslatorBase.h"
+#include "Misc/Paths.h"
+#include "PackageUtils/PackageUtils.h"
+#include "Stats/Stats.h"
+#include "Templates/SharedPointer.h"
+#include "UObject/Object.h"
+#include "UObject/UObjectGlobals.h"
+#include "UObject/ObjectMacros.h"
+#include "UObject/WeakObjectPtrTemplates.h"
+
+namespace UE
+{
+	namespace Interchange
+	{
+		namespace Private
+		{
+			void InternalGetPackageName(const UE::Interchange::FImportAsyncHelper& AsyncHelper, const int32 SourceIndex, const FString& PackageBasePath, const UInterchangeFactoryBaseNode* FactoryNode, FString& OutPackageName, FString& OutAssetName)
+			{
+				TRACE_CPUPROFILER_EVENT_SCOPE("UE::Interchange::Private::InternalGetPackageName")
+				const UInterchangeSourceData* SourceData = AsyncHelper.SourceDatas[SourceIndex];
+				check(SourceData);
+				FString NodeDisplayName = FactoryNode->GetAssetName();
+
+				// Set the asset name and the package name
+				OutAssetName = NodeDisplayName;
+				SanitizeObjectName(OutAssetName);
+
+				FString SanitizedPackageBasePath = PackageBasePath;
+				SanitizeObjectPath(SanitizedPackageBasePath);
+
+				FString SubPath;
+				if (FactoryNode->GetCustomSubPath(SubPath))
+				{
+					SanitizeObjectPath(SubPath);
+				}
+
+				OutPackageName = FPaths::Combine(*SanitizedPackageBasePath, *SubPath, *OutAssetName);
+			}
+
+			UObject* GetExistingObjectFromAssetImportData(UObject* ReimportObject, UInterchangeFactoryBaseNode* FactoryNode, FString& OutPackageName, FString& OutAssetName)
+			{
+				TRACE_CPUPROFILER_EVENT_SCOPE("UE::Interchange::Private::GetExistingObjectFromAssetImportData")
+				UInterchangeAssetImportData* OriginalAssetImportData = nullptr;
+				TArray<UObject*> SubObjects;
+				GetObjectsWithOuter(ReimportObject, SubObjects);
+				for (UObject* SubObject : SubObjects)
+				{
+					OriginalAssetImportData = Cast<UInterchangeAssetImportData>(SubObject);
+					if (OriginalAssetImportData)
+					{
+						break;
+					}
+				}
+
+				if (OriginalAssetImportData)
+				{
+					if (UInterchangeBaseNodeContainer* OriginalNodeContainer = OriginalAssetImportData->NodeContainer)
+					{
+						if (UInterchangeFactoryBaseNode* OriginalFactoryNode = OriginalNodeContainer->GetFactoryNode(OriginalAssetImportData->NodeUniqueID))
+						{
+							FSoftObjectPath ReferenceObject;
+							OriginalFactoryNode->GetCustomReferenceObject(ReferenceObject);
+							if (ReferenceObject.TryLoad() == ReimportObject)
+							{
+								OutPackageName = ReimportObject->GetPackage()->GetPathName();
+								OutAssetName = ReimportObject->GetName();
+
+								FactoryNode->SetDisplayLabel(OutAssetName);
+								FactoryNode->SetAssetName(OutAssetName);
+								// Hack for the texture reimport with a new file. (to revisited for MVP as this not a future proof solution. Should there be some sort of adapter that tell us
+								// how to do the mapping? Or should the pipeline do the mapping on a reimport? After all it is the pipeline that chose the name of the asset.
+								return ReimportObject;
+							}
+						}
+					}
+				}
+				return nullptr;
+			}
+
+		}//ns Private
+	}//ns Interchange
+}//ns UE
+
+void UE::Interchange::FTaskCreatePackage::DoTask(ENamedThreads::Type CurrentThread, const FGraphEventRef& MyCompletionGraphEvent)
+{
+	TRACE_CPUPROFILER_EVENT_SCOPE("UE::Interchange::FTaskCreatePackage::DoTask")
+#if INTERCHANGE_TRACE_ASYNCHRONOUS_TASK_ENABLED
+	INTERCHANGE_TRACE_ASYNCHRONOUS_TASK(CreatePackage)
+#endif
+	TSharedPtr<UE::Interchange::FImportAsyncHelper, ESPMode::ThreadSafe> AsyncHelper = WeakAsyncHelper.Pin();
+	check(AsyncHelper.IsValid());
+
+	//Verify if the task was cancel
+	if (AsyncHelper->bCancel)
+	{
+		return;
+	}
+
+	//The create package thread must always execute on the game thread
+	check(IsInGameThread());
+
+	// Create factory
+	UInterchangeFactoryBase* Factory = NewObject<UInterchangeFactoryBase>(GetTransientPackage(), FactoryClass);
+	Factory->SetResultsContainer(AsyncHelper->AssetImportResult->GetResults());
+
+	{
+		FScopeLock Lock(&AsyncHelper->CreatedFactoriesLock);
+		AsyncHelper->CreatedFactories.Add(FactoryNode->GetUniqueID(), Factory);
+	}
+
+	UPackage* Pkg = nullptr;
+	FString PackageName;
+	FString AssetName;
+	UObject* ReimportObject = AsyncHelper->TaskData.ReimportObject;
+	//If we do a reimport no need to create a package
+	if (ReimportObject)
+	{
+		Private::InternalGetPackageName(*AsyncHelper, SourceIndex, PackageBasePath, FactoryNode, PackageName, AssetName);
+		UPackage* ResultFindPackage = FindPackage(nullptr, *PackageName);
+		UObject* ExistingObject = nullptr;
+		if (ResultFindPackage)
+		{
+			ExistingObject = FindObject<UObject>(ResultFindPackage, *AssetName);
+		}
+		else
+		{
+			ExistingObject = FindFirstObject<UObject>(*AssetName, EFindFirstObjectOptions::NativeFirst | EFindFirstObjectOptions::EnsureIfAmbiguous);
+		}
+
+		if (!ExistingObject)
+		{
+			ExistingObject = Private::GetExistingObjectFromAssetImportData(ReimportObject, FactoryNode, PackageName, AssetName);
+		}
+
+		if (ExistingObject != ReimportObject)
+		{
+			//Set the existing object so other factory can link UObject correctly (i.e. mesh link to existing material)
+			FactoryNode->SetCustomReferenceObject(FSoftObjectPath(ExistingObject));
+			//Skip this asset, the re-import is not for this asset
+			return;
+		}
+		
+
+		if (ExistingObject)
+		{
+			Pkg = ExistingObject->GetPackage();
+			PackageName = Pkg->GetPathName();
+			//Import Asset describe by the node
+			UInterchangeFactoryBase::FCreateAssetParams CreateAssetParams;
+			CreateAssetParams.AssetName = AssetName;
+			CreateAssetParams.AssetNode = FactoryNode;
+			CreateAssetParams.Parent = Pkg;
+			CreateAssetParams.SourceData = AsyncHelper->SourceDatas[SourceIndex];
+			CreateAssetParams.Translator = AsyncHelper->Translators[SourceIndex];
+			if (AsyncHelper->BaseNodeContainers.IsValidIndex(SourceIndex))
+			{
+				CreateAssetParams.NodeContainer = AsyncHelper->BaseNodeContainers[SourceIndex].Get();
+			}
+			CreateAssetParams.ReimportObject = ReimportObject;
+			FactoryNode->SetCustomReferenceObject(FSoftObjectPath(ReimportObject));
+			//We call CreateEmptyAsset to ensure any resource use by an existing UObject is release on the game thread
+			Factory->CreateEmptyAsset(CreateAssetParams);
+		}
+		else
+		{
+			UInterchangeResultError_Generic* Message = Factory->AddMessage<UInterchangeResultError_Generic>();
+			Message->SourceAssetName = AsyncHelper->SourceDatas[SourceIndex]->GetFilename();
+			Message->DestinationAssetName = AssetName;
+			Message->AssetType = FactoryNode->GetObjectClass();
+			Message->Text = NSLOCTEXT("Interchange", "CannotFindPackageDuringReimport", "Cannot find an existing package.");
+
+			//Skip this asset
+			return;
+		}
+	}
+	else
+	{
+		Private::InternalGetPackageName(*AsyncHelper, SourceIndex, PackageBasePath, FactoryNode, PackageName, AssetName);
+		// We can not create assets that share the name of a map file in the same location
+		if (UE::Interchange::FPackageUtils::IsMapPackageAsset(PackageName))
+		{
+			UInterchangeResultError_Generic* Message = Factory->AddMessage<UInterchangeResultError_Generic>();
+			Message->SourceAssetName = AsyncHelper->SourceDatas[SourceIndex]->GetFilename();
+			Message->DestinationAssetName = AssetName;
+			Message->AssetType = FactoryNode->GetObjectClass();
+			Message->Text = NSLOCTEXT("Interchange", "MapExistsWithSameName", "You cannot create an asset with this name, as there is already a map file with the same name in this folder.");
+
+			//Skip this asset
+			return;
+		}
+
+		Pkg = CreatePackage(*PackageName);
+		if (Pkg == nullptr)
+		{
+			UInterchangeResultError_Generic* Message = Factory->AddMessage<UInterchangeResultError_Generic>();
+			Message->SourceAssetName = AsyncHelper->SourceDatas[SourceIndex]->GetFilename();
+			Message->DestinationAssetName = AssetName;
+			Message->AssetType = FactoryNode->GetObjectClass();
+			Message->Text = FText::Format(NSLOCTEXT("Interchange", "CouldntCreatePackage", "It was not possible to create a package named '{0}'; the asset will not be imported."), FText::FromString(PackageName));
+
+			//Skip this asset
+			return;
+		}
+
+		//Import Asset describe by the node
+		UInterchangeFactoryBase::FCreateAssetParams CreateAssetParams;
+		CreateAssetParams.AssetName = AssetName;
+		CreateAssetParams.AssetNode = FactoryNode;
+		CreateAssetParams.Parent = Pkg;
+		CreateAssetParams.SourceData = AsyncHelper->SourceDatas[SourceIndex];
+		CreateAssetParams.Translator = AsyncHelper->Translators[SourceIndex];
+		if (AsyncHelper->BaseNodeContainers.IsValidIndex(SourceIndex))
+		{
+			CreateAssetParams.NodeContainer = AsyncHelper->BaseNodeContainers[SourceIndex].Get();
+		}
+		CreateAssetParams.ReimportObject = ReimportObject;
+		//Make sure the asset UObject is created with the correct type on the main thread
+		UObject* NodeAsset = Factory->CreateEmptyAsset(CreateAssetParams);
+		if (NodeAsset)
+		{
+			if (!NodeAsset->HasAnyInternalFlags(EInternalObjectFlags::Async))
+			{
+				//Since the async flag is not set we must be in the game thread
+				ensure(IsInGameThread());
+				NodeAsset->SetInternalFlags(EInternalObjectFlags::Async);
+			}
+			FScopeLock Lock(&AsyncHelper->ImportedAssetsPerSourceIndexLock);
+			TArray<UE::Interchange::FImportAsyncHelper::FImportedObjectInfo>& ImportedInfos = AsyncHelper->ImportedAssetsPerSourceIndex.FindOrAdd(SourceIndex);
+			UE::Interchange::FImportAsyncHelper::FImportedObjectInfo& AssetInfo = ImportedInfos.AddDefaulted_GetRef();
+			AssetInfo.ImportedObject = NodeAsset;
+			AssetInfo.Factory = Factory;
+			AssetInfo.FactoryNode = FactoryNode;
+			AssetInfo.bIsReimport = bool(ReimportObject != nullptr);
+			FactoryNode->SetCustomReferenceObject(FSoftObjectPath(NodeAsset));
+		}
+	}
+
+	// Make sure the destination package is loaded
+	Pkg->FullyLoad();
+	
+	{
+		FScopeLock Lock(&AsyncHelper->CreatedPackagesLock);
+		AsyncHelper->CreatedPackages.Add(PackageName, Pkg);
+	}
+}
+
+void UE::Interchange::FTaskCreateAsset::DoTask(ENamedThreads::Type CurrentThread, const FGraphEventRef& MyCompletionGraphEvent)
+{
+#if INTERCHANGE_TRACE_ASYNCHRONOUS_TASK_ENABLED
+	INTERCHANGE_TRACE_ASYNCHRONOUS_TASK(CreateAsset)
+#endif
+	TSharedPtr<UE::Interchange::FImportAsyncHelper, ESPMode::ThreadSafe> AsyncHelper = WeakAsyncHelper.Pin();
+	check(AsyncHelper.IsValid());
+
+	//Verify if the task was cancel
+	if (AsyncHelper->bCancel)
+	{
+		return;
+	}
+
+	UInterchangeFactoryBase* Factory = nullptr;
+	{
+		FScopeLock Lock(&AsyncHelper->CreatedFactoriesLock);
+		Factory = AsyncHelper->CreatedFactories.FindChecked(FactoryNode->GetUniqueID());
+	}
+
+	UPackage* Pkg = nullptr;
+	FString PackageName;
+	FString AssetName;
+	Private::InternalGetPackageName(*AsyncHelper, SourceIndex, PackageBasePath, FactoryNode, PackageName, AssetName);
+	bool bSkipAsset = false;
+	UObject* ExistingObject = nullptr;
+	UObject* ReimportObject = AsyncHelper->TaskData.ReimportObject;
+	if (ReimportObject)
+	{
+		//When we re-import one particular asset, if the source file contain other assets, we want to set the node= reference UObject for those asset to the existing asset
+		//The way to discover this case is to compare the re-import asset with the node asset.
+		UPackage* ResultFindPackage = FindPackage(nullptr, *PackageName);
+		if (ResultFindPackage)
+		{
+			ExistingObject = FindObject<UObject>(ResultFindPackage, *AssetName);
+		}
+		else
+		{
+			ExistingObject = FindFirstObject<UObject>(*AssetName, EFindFirstObjectOptions::NativeFirst | EFindFirstObjectOptions::EnsureIfAmbiguous);
+		}
+		
+		if (!ExistingObject)
+		{
+			ExistingObject = UE::Interchange::Private::GetExistingObjectFromAssetImportData(ReimportObject, FactoryNode, PackageName, AssetName);
+		}
+
+		bSkipAsset = !ExistingObject || ExistingObject != ReimportObject;
+		if (!bSkipAsset)
+		{
+			Pkg = ReimportObject->GetPackage();
+			PackageName = Pkg->GetPathName();
+			AssetName = ExistingObject->GetName();
+		}
+		else if(ExistingObject)
+		{
+			Pkg = ExistingObject->GetPackage();
+			PackageName = Pkg->GetPathName();
+		}
+	}
+	else
+	{
+		FScopeLock Lock(&AsyncHelper->CreatedPackagesLock);
+		UPackage** PkgPtr = AsyncHelper->CreatedPackages.Find(PackageName);
+
+		if (!PkgPtr || !(*PkgPtr))
+		{
+			UInterchangeResultError_Generic* Message = Factory->AddMessage<UInterchangeResultError_Generic>();
+			Message->SourceAssetName = AsyncHelper->SourceDatas[SourceIndex]->GetFilename();
+			Message->DestinationAssetName = AssetName;
+			Message->AssetType = FactoryNode->GetObjectClass();
+			Message->Text = NSLOCTEXT("Interchange", "BadPackage", "It was not possible to create the asset as its package was not created correctly.");
+
+			return;
+		}
+
+		if (!AsyncHelper->SourceDatas.IsValidIndex(SourceIndex) || !AsyncHelper->Translators.IsValidIndex(SourceIndex))
+		{
+			UInterchangeResultError_Generic* Message = Factory->AddMessage<UInterchangeResultError_Generic>();
+			Message->DestinationAssetName = AssetName;
+			Message->AssetType = FactoryNode->GetObjectClass();
+			Message->Text = NSLOCTEXT("Interchange", "SourceDataOrTranslatorInvalid", "It was not possible to create the asset as its translator was not created correctly.");
+
+			return;
+		}
+
+		Pkg = *PkgPtr;
+	}
+
+	UObject* NodeAsset = nullptr;
+	if (bSkipAsset)
+	{
+		NodeAsset = ExistingObject;
+	}
+	else
+	{
+		UInterchangeTranslatorBase* Translator = AsyncHelper->Translators[SourceIndex];
+		//Import Asset describe by the node
+		UInterchangeFactoryBase::FCreateAssetParams CreateAssetParams;
+		CreateAssetParams.AssetName = AssetName;
+		CreateAssetParams.AssetNode = FactoryNode;
+		CreateAssetParams.Parent = Pkg;
+		CreateAssetParams.SourceData = AsyncHelper->SourceDatas[SourceIndex];
+		CreateAssetParams.Translator = Translator;
+		if (AsyncHelper->BaseNodeContainers.IsValidIndex(SourceIndex))
+		{
+			CreateAssetParams.NodeContainer = AsyncHelper->BaseNodeContainers[SourceIndex].Get();
+		}
+		CreateAssetParams.ReimportObject = ReimportObject;
+
+		NodeAsset = Factory->CreateAsset(CreateAssetParams);
+	}
+	if (NodeAsset)
+	{
+		if (!bSkipAsset)
+		{
+			FScopeLock Lock(&AsyncHelper->ImportedAssetsPerSourceIndexLock);
+			TArray<UE::Interchange::FImportAsyncHelper::FImportedObjectInfo>& ImportedInfos = AsyncHelper->ImportedAssetsPerSourceIndex.FindOrAdd(SourceIndex);
+			UE::Interchange::FImportAsyncHelper::FImportedObjectInfo* AssetInfoPtr = ImportedInfos.FindByPredicate([NodeAsset](const UE::Interchange::FImportAsyncHelper::FImportedObjectInfo& CurInfo)
+			{
+				return CurInfo.ImportedObject == NodeAsset;
+			});
+
+			if (!AssetInfoPtr)
+			{
+				UE::Interchange::FImportAsyncHelper::FImportedObjectInfo& AssetInfo = ImportedInfos.AddDefaulted_GetRef();
+				AssetInfo.ImportedObject = NodeAsset;
+				AssetInfo.Factory = Factory;
+				AssetInfo.FactoryNode = FactoryNode;
+				AssetInfo.bIsReimport = bool(ReimportObject != nullptr);
+			}
+
+			// Fill in destination asset and type in any results which have been added previously by a translator or pipeline, now that we have a corresponding factory.
+			UInterchangeResultsContainer* Results = AsyncHelper->AssetImportResult->GetResults();
+			for (UInterchangeResult* Result : Results->GetResults())
+			{
+				if (!Result->InterchangeKey.IsEmpty() && (Result->DestinationAssetName.IsEmpty() || Result->AssetType == nullptr))
+				{
+					TArray<FString> TargetAssets;
+					FactoryNode->GetTargetNodeUids(TargetAssets);
+					if (TargetAssets.Contains(Result->InterchangeKey))
+					{
+						Result->DestinationAssetName = NodeAsset->GetPathName();
+						Result->AssetType = NodeAsset->GetClass();
+					}
+				}
+			}
+		}
+
+		FactoryNode->SetCustomReferenceObject(FSoftObjectPath(NodeAsset));
+	}
+}
\ No newline at end of file
diff --git a/Engine/Source/Runtime/Interchange/Engine/Private/Tasks/InterchangeTaskCreateAsset.h b/Engine/Source/Runtime/Interchange/Engine/Private/Tasks/InterchangeTaskCreateAsset.h
new file mode 100644
--- /dev/null
+++ b/Engine/Source/Runtime/Interchange/Engine/Private/Tasks/InterchangeTaskCreateAsset.h
@@ -0,0 +1,106 @@
+// Copyright Epic Games, Inc. All Rights Reserved.
+
+#pragma once
+
+#include "Async/TaskGraphInterfaces.h"
+#include "CoreMinimal.h"
+#include "InterchangeFactoryBase.h"
+#include "InterchangeManager.h"
+#include "Stats/Stats.h"
+#include "UObject/WeakObjectPtrTemplates.h"
+#include "Nodes/InterchangeFactoryBaseNode.h"
+
+namespace UE
+{
+	namespace Interchange
+	{
+
+		/**
+		 * This task create package, Cook::PackageTracker::NotifyUObjectCreated is not thread safe, so we need to create the packages on the main thread
+		 */
+		class FTaskCreatePackage
+		{
+		private:
+			FString PackageBasePath;
+			int32 SourceIndex;
+			TWeakPtr<FImportAsyncHelper, ESPMode::ThreadSafe> WeakAsyncHelper;
+			UInterchangeFactoryBaseNode* FactoryNode;
+			const UClass* FactoryClass;
+
+		public:
+			FTaskCreatePackage(const FString& InPackageBasePath, const int32 InSourceIndex, TWeakPtr<FImportAsyncHelper, ESPMode::ThreadSafe> InAsyncHelper, UInterchangeFactoryBaseNode* InFactoryNode, const UClass* InFactoryClass)
+				: PackageBasePath(InPackageBasePath)
+				, SourceIndex(InSourceIndex)
+				, WeakAsyncHelper(InAsyncHelper)
+				, FactoryNode(InFactoryNode)
+				, FactoryClass(InFactoryClass)
+			{
+				check(FactoryNode);
+				check(FactoryClass);
+			}
+
+			ENamedThreads::Type GetDesiredThread() const
+			{
+				// This is no longer true now that we have to construct a factory here
+				//if (WeakAsyncHelper.IsValid() && WeakAsyncHelper.Pin()->TaskData.ReimportObject)
+				//{
+				//	//When doing a reimport the package already exist, so we can get it outside of the main thread
+				//	return ENamedThreads::AnyBackgroundThreadNormalTask;
+				//}
+				return ENamedThreads::GameThread;
+			}
+
+			static ESubsequentsMode::Type GetSubsequentsMode()
+			{
+				return ESubsequentsMode::TrackSubsequents;
+			}
+
+			TStatId GetStatId() const
+			{
+				RETURN_QUICK_DECLARE_CYCLE_STAT(FTaskCreatePackage, STATGROUP_TaskGraphTasks);
+			}
+
+			void DoTask(ENamedThreads::Type CurrentThread, const FGraphEventRef& MyCompletionGraphEvent);
+		};
+
+		class FTaskCreateAsset
+		{
+		private:
+			FString PackageBasePath;
+			int32 SourceIndex;
+			TWeakPtr<FImportAsyncHelper, ESPMode::ThreadSafe> WeakAsyncHelper;
+			UInterchangeFactoryBaseNode* FactoryNode;
+			bool bCanRunOnAnyThread;
+
+		public:
+			FTaskCreateAsset(const FString& InPackageBasePath, const int32 InSourceIndex, TWeakPtr<FImportAsyncHelper, ESPMode::ThreadSafe> InAsyncHelper, UInterchangeFactoryBaseNode* InFactoryNode, bool bInCanRunOnAnyThread)
+				: PackageBasePath(InPackageBasePath)
+				, SourceIndex(InSourceIndex)
+				, WeakAsyncHelper(InAsyncHelper)
+				, FactoryNode(InFactoryNode)
+				, bCanRunOnAnyThread(bInCanRunOnAnyThread)
+			{
+				check(FactoryNode);
+			}
+
+			ENamedThreads::Type GetDesiredThread() const
+			{
+				return bCanRunOnAnyThread ? ENamedThreads::AnyBackgroundThreadNormalTask : ENamedThreads::GameThread;
+			}
+
+			static ESubsequentsMode::Type GetSubsequentsMode()
+			{
+				return ESubsequentsMode::TrackSubsequents;
+			}
+
+			TStatId GetStatId() const
+			{
+				RETURN_QUICK_DECLARE_CYCLE_STAT(FTaskCreateAsset, STATGROUP_TaskGraphTasks);
+			}
+
+			void DoTask(ENamedThreads::Type CurrentThread, const FGraphEventRef& MyCompletionGraphEvent);
+		};
+
+
+	} //ns Interchange
+}//ns UE
