Fix re-import with new file if there is other texture assets next to the assets we re-import that already use the same sourcefile.
#rnx

diff --git a/Engine/Source/Runtime/Interchange/Engine/Private/Tasks/InterchangeTaskCreateAsset.cpp b/Engine/Source/Runtime/Interchange/Engine/Private/Tasks/InterchangeTaskCreateAsset.cpp
--- a/Engine/Source/Runtime/Interchange/Engine/Private/Tasks/InterchangeTaskCreateAsset.cpp
+++ b/Engine/Source/Runtime/Interchange/Engine/Private/Tasks/InterchangeTaskCreateAsset.cpp
@@ -126,23 +126,7 @@ void UE::Interchange::FTaskCreatePackage::DoTask(ENamedThreads::Type CurrentThre
 	//If we do a reimport no need to create a package
 	if (ReimportObject)
 	{
-		Private::InternalGetPackageName(*AsyncHelper, SourceIndex, PackageBasePath, FactoryNode, PackageName, AssetName);
-		UPackage* ResultFindPackage = FindPackage(nullptr, *PackageName);
-		UObject* ExistingObject = nullptr;
-		if (ResultFindPackage)
-		{
-			ExistingObject = FindObject<UObject>(ResultFindPackage, *AssetName);
-		}
-		else
-		{
-			ExistingObject = FindFirstObject<UObject>(*AssetName, EFindFirstObjectOptions::NativeFirst | EFindFirstObjectOptions::EnsureIfAmbiguous);
-		}
-
-		if (!ExistingObject)
-		{
-			ExistingObject = Private::GetExistingObjectFromAssetImportData(ReimportObject, FactoryNode, PackageName, AssetName);
-		}
-
+		UObject* ExistingObject = Private::GetExistingObjectFromAssetImportData(ReimportObject, FactoryNode, PackageName, AssetName);
 		if (ExistingObject != ReimportObject)
 		{
 			//Set the existing object so other factory can link UObject correctly (i.e. mesh link to existing material)
@@ -150,7 +134,6 @@ void UE::Interchange::FTaskCreatePackage::DoTask(ENamedThreads::Type CurrentThre
 			//Skip this asset, the re-import is not for this asset
 			return;
 		}
-		
 
 		if (ExistingObject)
 		{
@@ -284,23 +267,7 @@ void UE::Interchange::FTaskCreateAsset::DoTask(ENamedThreads::Type CurrentThread
 	UObject* ReimportObject = AsyncHelper->TaskData.ReimportObject;
 	if (ReimportObject)
 	{
-		//When we re-import one particular asset, if the source file contain other assets, we want to set the node= reference UObject for those asset to the existing asset
-		//The way to discover this case is to compare the re-import asset with the node asset.
-		UPackage* ResultFindPackage = FindPackage(nullptr, *PackageName);
-		if (ResultFindPackage)
-		{
-			ExistingObject = FindObject<UObject>(ResultFindPackage, *AssetName);
-		}
-		else
-		{
-			ExistingObject = FindFirstObject<UObject>(*AssetName, EFindFirstObjectOptions::NativeFirst | EFindFirstObjectOptions::EnsureIfAmbiguous);
-		}
-		
-		if (!ExistingObject)
-		{
-			ExistingObject = UE::Interchange::Private::GetExistingObjectFromAssetImportData(ReimportObject, FactoryNode, PackageName, AssetName);
-		}
-
+		ExistingObject = Private::GetExistingObjectFromAssetImportData(ReimportObject, FactoryNode, PackageName, AssetName);
 		bSkipAsset = !ExistingObject || ExistingObject != ReimportObject;
 		if (!bSkipAsset)
 		{
