Interchange: When canceling a multi file import, canceling at any point will cancel all queued imports.

diff --git a/Engine/Source/Runtime/Interchange/Engine/Private/InterchangeManager.cpp b/Engine/Source/Runtime/Interchange/Engine/Private/InterchangeManager.cpp
--- a/Engine/Source/Runtime/Interchange/Engine/Private/InterchangeManager.cpp
+++ b/Engine/Source/Runtime/Interchange/Engine/Private/InterchangeManager.cpp
@@ -985,7 +985,7 @@ UInterchangeManager::ImportInternal(const FString& ContentPath, const UInterchan
 			}
 		}
 	}
-	bool bImportCancel = false;
+	bool bImportAborted = false; // True when we're unable to go through with the import process
 
 	//Create a task for every source data
 	UE::Interchange::FImportAsyncHelperData TaskData;
@@ -1038,7 +1038,8 @@ UInterchangeManager::ImportInternal(const FString& ContentPath, const UInterchan
 #if WITH_EDITORONLY_DATA
 		const bool bShowPipelineStacksConfigurationDialog = !bIsUnattended
 															&& FInterchangeProjectSettingsUtils::ShouldShowPipelineStacksConfigurationDialog(bImportScene, *SourceData)
-															&& !bImportAllWithDefault;
+															&& !bImportAllWithDefault
+															&& !bImportCanceled;
 #else
 		const bool bShowPipelineStacksConfigurationDialog = false;
 #endif
@@ -1092,7 +1093,7 @@ UInterchangeManager::ImportInternal(const FString& ContentPath, const UInterchan
 				EInterchangePipelineConfigurationDialogResult DialogResult = RegisteredPipelineConfiguration->ScriptedShowReimportPipelineConfigurationDialog(PipelineStack, DuplicateSourceData);
 				if (DialogResult == EInterchangePipelineConfigurationDialogResult::Cancel)
 				{
-					bImportCancel = true;
+					bImportCanceled = true;
 				}
 				if (DialogResult == EInterchangePipelineConfigurationDialogResult::ImportAll)
 				{
@@ -1129,7 +1130,7 @@ UInterchangeManager::ImportInternal(const FString& ContentPath, const UInterchan
 				EInterchangePipelineConfigurationDialogResult DialogResult = bImportScene ? RegisteredPipelineConfiguration->ScriptedShowScenePipelineConfigurationDialog(DuplicateSourceData) : RegisteredPipelineConfiguration->ScriptedShowPipelineConfigurationDialog(DuplicateSourceData);
 				if (DialogResult == EInterchangePipelineConfigurationDialogResult::Cancel)
 				{
-					bImportCancel = true;
+					bImportCanceled = true;
 				}
 				if (DialogResult == EInterchangePipelineConfigurationDialogResult::ImportAll)
 				{
@@ -1137,7 +1138,7 @@ UInterchangeManager::ImportInternal(const FString& ContentPath, const UInterchan
 				}
 			}
 
-			if (!bImportCancel)
+			if (!bImportCanceled)
 			{
 				if (!DefaultPipelineStacks.Contains(PipelineStackName))
 				{
@@ -1203,7 +1204,7 @@ UInterchangeManager::ImportInternal(const FString& ContentPath, const UInterchan
 				else
 				{
 					//Log an error, we cannot import asset without a valid pipeline, there is no pipeline stack defined
-					bImportCancel = true;
+					bImportAborted = true;
 				}
 			}
 		}
@@ -1229,11 +1230,11 @@ UInterchangeManager::ImportInternal(const FString& ContentPath, const UInterchan
 	}
 
 	//Cancel the import do not queue task
-	if (bImportCancel)
+	if (bImportCanceled || bImportAborted)
 	{
 		if (FEngineAnalytics::IsAvailable())
 		{
-			Attribs.Add(FAnalyticsEventAttribute(TEXT("Canceled"), bImportCancel));
+			Attribs.Add(FAnalyticsEventAttribute(TEXT("Canceled"), bImportCanceled));
 		}
 
 		AsyncHelper->InitCancel();
@@ -1372,15 +1373,27 @@ void UInterchangeManager::ReleaseAsyncHelper(TWeakPtr<UE::Interchange::FImportAs
 	FString ImportTaskNumberStr = TEXT(" (") + FString::FromInt(ImportTaskNumber) + TEXT(")");
 	if (ImportTaskNumber == 0)
 	{
-		bImportAllWithDefault = false;
 		SetActiveMode(false);
 
 		if (Notification.IsValid())
 		{
-			FText TitleText = NSLOCTEXT("Interchange", "Asynchronous_import_end", "Import Done");
+			FText TitleText;
+			if (bImportCanceled)
+			{
+				TitleText = NSLOCTEXT("Interchange", "Asynchronous_import_canceled", "Import Canceled");
+				bSucceeded = true; // Mark the "cancelation" as a success so that the notification goes away
+			}
+			else
+			{
+				TitleText = NSLOCTEXT("Interchange", "Asynchronous_import_end", "Import Done");
+			}
+
 			Notification->SetComplete(TitleText, FText::GetEmpty(), bSucceeded);
 			Notification = nullptr; //This should delete the notification
 		}
+
+		bImportAllWithDefault = false;
+		bImportCanceled = false;
 	}
 	else if(Notification.IsValid())
 	{
diff --git a/Engine/Source/Runtime/Interchange/Engine/Public/InterchangeManager.h b/Engine/Source/Runtime/Interchange/Engine/Public/InterchangeManager.h
--- a/Engine/Source/Runtime/Interchange/Engine/Public/InterchangeManager.h
+++ b/Engine/Source/Runtime/Interchange/Engine/Public/InterchangeManager.h
@@ -566,9 +566,13 @@ private:
 	//If interchange is currently importing we have a timer to watch the cancel and we block GC 
 	FThreadSafeBool bIsActive = false;
 
-	//If the user want to use the same import pipeline stack for all the queue task
+	//If the user wants to use the same import pipeline stack for all the queue task
 	//This boolean is reset to false when the ImportTasks array is empty.
 	bool bImportAllWithDefault = false;
 
+	//Indicates that the import process was canceled by the user.
+	//This boolean is reset to false when the ImportTasks array is empty.
+	bool bImportCanceled = false;
+
 	friend class UE::Interchange::FScopedTranslator;
 };
