Interchange: Refactor how reimport find the matching data when there is multiple asset in a source file.
#rnx

diff --git a/Engine/Source/Runtime/Interchange/Engine/Private/Tasks/InterchangeTaskCreateAsset.cpp b/Engine/Source/Runtime/Interchange/Engine/Private/Tasks/InterchangeTaskCreateAsset.cpp
--- a/Engine/Source/Runtime/Interchange/Engine/Private/Tasks/InterchangeTaskCreateAsset.cpp
+++ b/Engine/Source/Runtime/Interchange/Engine/Private/Tasks/InterchangeTaskCreateAsset.cpp
@@ -47,45 +47,74 @@ namespace UE
 
 				OutPackageName = FPaths::Combine(*SanitizedPackageBasePath, *SubPath, *OutAssetName);
 			}
-
-			UObject* GetExistingObjectFromAssetImportData(UObject* ReimportObject, UInterchangeFactoryBaseNode* FactoryNode, FString& OutPackageName, FString& OutAssetName)
+			bool ShouldReimportFactoryNode(UInterchangeFactoryBaseNode* FactoryNode, const UInterchangeBaseNodeContainer* NodeContainer, UObject* ReimportObject)
 			{
-				TRACE_CPUPROFILER_EVENT_SCOPE("UE::Interchange::Private::GetExistingObjectFromAssetImportData")
-				UInterchangeAssetImportData* OriginalAssetImportData = nullptr;
-				TArray<UObject*> SubObjects;
-				GetObjectsWithOuter(ReimportObject, SubObjects);
-				for (UObject* SubObject : SubObjects)
+				TRACE_CPUPROFILER_EVENT_SCOPE("UE::Interchange::Private::ShouldReimportFactoryNode")
+
+				if (!NodeContainer)
 				{
-					OriginalAssetImportData = Cast<UInterchangeAssetImportData>(SubObject);
-					if (OriginalAssetImportData)
+					return false;
+				}
+				//Find all potential factory node
+				TArray<UInterchangeFactoryBaseNode*> PotentialFactoryNodes;
+				UClass* FactoryClass = FactoryNode->GetObjectClass();
+				NodeContainer->IterateNodesOfType<UInterchangeFactoryBaseNode>([FactoryClass, &PotentialFactoryNodes](const FString& NodeUniqueID, UInterchangeFactoryBaseNode* CurrentFactoryNode)
 					{
-						break;
-					}
+						if (UClass* CurrentFactoryClass = CurrentFactoryNode->GetObjectClass())
+						{
+							if (CurrentFactoryClass->IsChildOf(FactoryClass))
+							{
+								PotentialFactoryNodes.Add(CurrentFactoryNode);
+							}
+						}
+					});
+
+				if (PotentialFactoryNodes.Num() == 1)
+				{
+					//There is only one factory node that will generate this class of UObject, no need to match the unique id or the name.
+					ensure(PotentialFactoryNodes[0] == FactoryNode);
+					return true;
 				}
 
-				if (OriginalAssetImportData)
+				//If the source was used to import multiple asset, see if the FactoryNode match the ReimportObject.
+				TArray<UObject*> SubObjects;
+				GetObjectsWithOuter(ReimportObject, SubObjects);
+				for (UObject* SubObject : SubObjects)
 				{
-					if (UInterchangeBaseNodeContainer* OriginalNodeContainer = OriginalAssetImportData->NodeContainer)
+					if(UInterchangeAssetImportData* OriginalAssetImportData = Cast<UInterchangeAssetImportData>(SubObject))
 					{
-						if (UInterchangeFactoryBaseNode* OriginalFactoryNode = OriginalNodeContainer->GetFactoryNode(OriginalAssetImportData->NodeUniqueID))
+						if (UInterchangeBaseNodeContainer* OriginalNodeContainer = OriginalAssetImportData->NodeContainer)
 						{
-							FSoftObjectPath ReferenceObject;
-							OriginalFactoryNode->GetCustomReferenceObject(ReferenceObject);
-							if (ReferenceObject.TryLoad() == ReimportObject)
+							//Find the original factory node used by the last ReimportObject import
+							if (UInterchangeFactoryBaseNode* OriginalFactoryNode = OriginalNodeContainer->GetFactoryNode(OriginalAssetImportData->NodeUniqueID))
 							{
-								OutPackageName = ReimportObject->GetPackage()->GetPathName();
-								OutAssetName = ReimportObject->GetName();
-
-								FactoryNode->SetDisplayLabel(OutAssetName);
-								FactoryNode->SetAssetName(OutAssetName);
-								// Hack for the texture reimport with a new file. (to revisited for MVP as this not a future proof solution. Should there be some sort of adapter that tell us
-								// how to do the mapping? Or should the pipeline do the mapping on a reimport? After all it is the pipeline that chose the name of the asset.
-								return ReimportObject;
+								//Compare the original factory node UObject class to the factory node UObject class
+								if (OriginalFactoryNode->GetObjectClass()->IsChildOf(FactoryClass))
+								{
+									//Compare the original factory node, unique id to the factory node unique id
+									if (OriginalFactoryNode->GetUniqueID().Equals(FactoryNode->GetUniqueID()))
+									{
+										return true;
+									}
+									//Compare the original factory node, name to the factory node name
+									if (OriginalFactoryNode->GetDisplayLabel().Equals(FactoryNode->GetDisplayLabel()))
+									{
+										FString PackageSubPath;
+										FactoryNode->GetCustomSubPath(PackageSubPath);
+										FString OriginalPackageSubPath;
+										OriginalFactoryNode->GetCustomSubPath(PackageSubPath);
+										//Make sure both sub path are equal
+										if (PackageSubPath.Equals(OriginalPackageSubPath))
+										{
+											return true;
+										}
+									}
+								}
 							}
 						}
 					}
 				}
-				return nullptr;
+				return false;
 			}
 
 		}//ns Private
@@ -126,19 +155,19 @@ void UE::Interchange::FTaskCreatePackage::DoTask(ENamedThreads::Type CurrentThre
 	//If we do a reimport no need to create a package
 	if (ReimportObject)
 	{
-		UObject* ExistingObject = Private::GetExistingObjectFromAssetImportData(ReimportObject, FactoryNode, PackageName, AssetName);
-		if (ExistingObject != ReimportObject)
+		UInterchangeBaseNodeContainer* NodeContainer = nullptr;
+		if (AsyncHelper->BaseNodeContainers.IsValidIndex(SourceIndex))
 		{
-			//Set the existing object so other factory can link UObject correctly (i.e. mesh link to existing material)
-			FactoryNode->SetCustomReferenceObject(FSoftObjectPath(ExistingObject));
-			//Skip this asset, the re-import is not for this asset
-			return;
+			NodeContainer = AsyncHelper->BaseNodeContainers[SourceIndex].Get();
 		}
-
-		if (ExistingObject)
+		
+		if (Private::ShouldReimportFactoryNode(FactoryNode, NodeContainer, ReimportObject))
 		{
-			Pkg = ExistingObject->GetPackage();
+			FactoryNode->SetDisplayLabel(ReimportObject->GetName());
+			FactoryNode->SetAssetName(ReimportObject->GetName());
+			Pkg = ReimportObject->GetPackage();
 			PackageName = Pkg->GetPathName();
+			AssetName = ReimportObject->GetName();
 			//Import Asset describe by the node
 			UInterchangeFactoryBase::FCreateAssetParams CreateAssetParams;
 			CreateAssetParams.AssetName = AssetName;
@@ -146,24 +175,15 @@ void UE::Interchange::FTaskCreatePackage::DoTask(ENamedThreads::Type CurrentThre
 			CreateAssetParams.Parent = Pkg;
 			CreateAssetParams.SourceData = AsyncHelper->SourceDatas[SourceIndex];
 			CreateAssetParams.Translator = AsyncHelper->Translators[SourceIndex];
-			if (AsyncHelper->BaseNodeContainers.IsValidIndex(SourceIndex))
-			{
-				CreateAssetParams.NodeContainer = AsyncHelper->BaseNodeContainers[SourceIndex].Get();
-			}
+			CreateAssetParams.NodeContainer = NodeContainer;
 			CreateAssetParams.ReimportObject = ReimportObject;
 			FactoryNode->SetCustomReferenceObject(FSoftObjectPath(ReimportObject));
-			//We call CreateEmptyAsset to ensure any resource use by an existing UObject is release on the game thread
+			//We call CreateEmptyAsset to ensure any resource use by an existing UObject is released on the game thread
 			Factory->CreateEmptyAsset(CreateAssetParams);
 		}
 		else
 		{
-			UInterchangeResultError_Generic* Message = Factory->AddMessage<UInterchangeResultError_Generic>();
-			Message->SourceAssetName = AsyncHelper->SourceDatas[SourceIndex]->GetFilename();
-			Message->DestinationAssetName = AssetName;
-			Message->AssetType = FactoryNode->GetObjectClass();
-			Message->Text = NSLOCTEXT("Interchange", "CannotFindPackageDuringReimport", "Cannot find an existing package.");
-
-			//Skip this asset
+			//Skip this asset, reimport object original factory node does not match this factory node
 			return;
 		}
 	}
@@ -263,22 +283,24 @@ void UE::Interchange::FTaskCreateAsset::DoTask(ENamedThreads::Type CurrentThread
 	FString AssetName;
 	Private::InternalGetPackageName(*AsyncHelper, SourceIndex, PackageBasePath, FactoryNode, PackageName, AssetName);
 	bool bSkipAsset = false;
-	UObject* ExistingObject = nullptr;
 	UObject* ReimportObject = AsyncHelper->TaskData.ReimportObject;
 	if (ReimportObject)
 	{
-		ExistingObject = Private::GetExistingObjectFromAssetImportData(ReimportObject, FactoryNode, PackageName, AssetName);
-		bSkipAsset = !ExistingObject || ExistingObject != ReimportObject;
-		if (!bSkipAsset)
+		UInterchangeBaseNodeContainer* NodeContainer = nullptr;
+		if (AsyncHelper->BaseNodeContainers.IsValidIndex(SourceIndex))
+		{
+			NodeContainer = AsyncHelper->BaseNodeContainers[SourceIndex].Get();
+		}
+
+		if (Private::ShouldReimportFactoryNode(FactoryNode, NodeContainer, ReimportObject))
 		{
 			Pkg = ReimportObject->GetPackage();
 			PackageName = Pkg->GetPathName();
-			AssetName = ExistingObject->GetName();
+			AssetName = ReimportObject->GetName();
 		}
-		else if(ExistingObject)
+		else
 		{
-			Pkg = ExistingObject->GetPackage();
-			PackageName = Pkg->GetPathName();
+			bSkipAsset = true;
 		}
 	}
 	else
@@ -313,7 +335,7 @@ void UE::Interchange::FTaskCreateAsset::DoTask(ENamedThreads::Type CurrentThread
 	UObject* NodeAsset = nullptr;
 	if (bSkipAsset)
 	{
-		NodeAsset = ExistingObject;
+		NodeAsset = nullptr;
 	}
 	else
 	{
