Make Interchange always on by default
#rnx

diff --git a/Engine/Source/Developer/AssetTools/Private/AssetTools.cpp b/Engine/Source/Developer/AssetTools/Private/AssetTools.cpp
--- a/Engine/Source/Developer/AssetTools/Private/AssetTools.cpp
+++ b/Engine/Source/Developer/AssetTools/Private/AssetTools.cpp
@@ -1774,15 +1774,7 @@ TArray<UObject*> UAssetToolsImpl::ImportAssetsInternal(const TArray<FString>& Fi
 	// Use Interchange if it's enabled and we weren't asked to use a specific UFactory
 	const bool bUseInterchangeFramework = [SpecifiedFactory, &InterchangeManager]()
 	{
-		bool bResult = false;
-		
-#if WITH_EDITOR
-		const UEditorExperimentalSettings* EditorExperimentalSettings = GetDefault<UEditorExperimentalSettings>();
-		bResult = EditorExperimentalSettings->bEnableInterchangeFramework;
-#endif
-
-		bResult = bResult && (SpecifiedFactory == nullptr);
-		return bResult;
+		return UInterchangeManager::IsInterchangeImportEnabled() && (SpecifiedFactory == nullptr);
 	}();
 
 	FScopedSlowTask SlowTask(ValidFiles.Num(), LOCTEXT("ImportSlowTask", "Importing"), !bUseInterchangeFramework);
@@ -4220,8 +4212,7 @@ TArray<UObject*> UAssetToolsImpl::ImportAssetsWithDialogImplementation(const FSt
 	// Generate the file types and extensions represented by the selected factories
 	ObjectTools::GenerateFactoryFileExtensions(Factories, FileTypes, AllExtensions, FilterIndexToFactory);
 
-	const UEditorExperimentalSettings* EditorExperimentalSettings = GetDefault<UEditorExperimentalSettings>();
-	if (EditorExperimentalSettings->bEnableInterchangeFramework)
+	if (UInterchangeManager::IsInterchangeImportEnabled())
 	{
 		TArray<FString> InterchangeFileExtensions = UInterchangeManager::GetInterchangeManager().GetSupportedFormats(EInterchangeTranslatorType::Assets);
 		ObjectTools::AppendFormatsFileExtensions(InterchangeFileExtensions, FileTypes, AllExtensions);
diff --git a/Engine/Source/Editor/UnrealEd/Classes/Settings/EditorExperimentalSettings.h b/Engine/Source/Editor/UnrealEd/Classes/Settings/EditorExperimentalSettings.h
--- a/Engine/Source/Editor/UnrealEd/Classes/Settings/EditorExperimentalSettings.h
+++ b/Engine/Source/Editor/UnrealEd/Classes/Settings/EditorExperimentalSettings.h
@@ -41,10 +41,6 @@ public:
 	UPROPERTY(EditAnywhere, config, Category = Performance, meta = (DisplayName = "Enable async sound compilation and loading"))
 	bool bEnableAsyncSoundWaveCompilation;
 
-	/** Enable interchange framework, the interchange framework is a new import system which can import asynchronously and in parallel. See the interchange project settings to configure the import pipeline*/
-	UPROPERTY(EditAnywhere, config, Category = Interchange, meta = (DisplayName = "Enable interchange framework import"))
-	bool bEnableInterchangeFramework;
-
 	/** Allows the editor to run on HDR monitors on Windows 10 */
 	UPROPERTY(EditAnywhere, config, Category = HDR, meta = (ConfigRestartRequired = true, DisplayName = "Enable Editor Support for HDR Monitors"))
 	bool bHDREditor;
diff --git a/Engine/Source/Editor/UnrealEd/Private/Editor.cpp b/Engine/Source/Editor/UnrealEd/Private/Editor.cpp
--- a/Engine/Source/Editor/UnrealEd/Private/Editor.cpp
+++ b/Engine/Source/Editor/UnrealEd/Private/Editor.cpp
@@ -72,7 +72,6 @@
 
 #if WITH_EDITOR
 #include "Subsystems/AssetEditorSubsystem.h"
-#include "Settings/EditorExperimentalSettings.h"
 #endif
 
 #define LOCTEXT_NAMESPACE "UnrealEd.Editor"
@@ -265,11 +264,8 @@ bool FReimportManager::Reimport( UObject* Obj, bool bAskForNewFileIfMissing, boo
 	// Warn that were about to reimport, so prep for it
 	PreReimport.Broadcast( Obj );
 
-	bool bUseInterchangeFramework = false;
+	bool bUseInterchangeFramework = UInterchangeManager::IsInterchangeImportEnabled();;
 	UInterchangeManager& InterchangeManager = UInterchangeManager::GetInterchangeManager();
-	const UEditorExperimentalSettings* EditorExperimentalSettings = GetDefault<UEditorExperimentalSettings>();
-
-	bUseInterchangeFramework = EditorExperimentalSettings->bEnableInterchangeFramework;
 
 	bool bSuccess = false;
 	if ( Obj )
diff --git a/Engine/Source/Editor/UnrealEd/Private/Factories/InterchangeReimportHandler.cpp b/Engine/Source/Editor/UnrealEd/Private/Factories/InterchangeReimportHandler.cpp
--- a/Engine/Source/Editor/UnrealEd/Private/Factories/InterchangeReimportHandler.cpp
+++ b/Engine/Source/Editor/UnrealEd/Private/Factories/InterchangeReimportHandler.cpp
@@ -12,7 +12,6 @@
 #include "InterchangeAssetImportData.h"
 #include "InterchangeFactoryBase.h"
 #include "InterchangeManager.h"
-#include "Settings/EditorExperimentalSettings.h"
 #include "UObject/UObjectGlobals.h"
 #include "UObject/ObjectMacros.h"
 
@@ -26,16 +25,8 @@ UInterchangeReimportHandler::UInterchangeReimportHandler(const FObjectInitialize
 //~ Begin FReimportHandler Interface
 bool UInterchangeReimportHandler::CanReimport(UObject* Obj, TArray<FString>& OutFilenames)
 {
-	const UEditorExperimentalSettings* EditorExperimentalSettings = GetDefault<UEditorExperimentalSettings>();
-	const bool bUseInterchange = EditorExperimentalSettings->bEnableInterchangeFramework;
-
-	if (bUseInterchange)
-	{
-		const UInterchangeManager& InterchangeManager = UInterchangeManager::GetInterchangeManager();
-		return InterchangeManager.CanReimport(Obj, OutFilenames);
-	}
-
-	return false;
+	const UInterchangeManager& InterchangeManager = UInterchangeManager::GetInterchangeManager();
+	return InterchangeManager.CanReimport(Obj, OutFilenames);
 }
 
 void UInterchangeReimportHandler::SetReimportPaths(UObject* Obj, const FString& NewReimportPath, const int32 SourceFileIndex)
diff --git a/Engine/Source/Editor/UnrealEd/Private/FbxMeshUtils.cpp b/Engine/Source/Editor/UnrealEd/Private/FbxMeshUtils.cpp
--- a/Engine/Source/Editor/UnrealEd/Private/FbxMeshUtils.cpp
+++ b/Engine/Source/Editor/UnrealEd/Private/FbxMeshUtils.cpp
@@ -31,7 +31,6 @@
 #include "InterchangeManager.h"
 #include "InterchangeMeshUtilities.h"
 #include "InterchangeProjectSettings.h"
-#include "Settings/EditorExperimentalSettings.h"
 
 #include "Misc/FbxErrors.h"
 #include "Framework/Notifications/NotificationManager.h"
@@ -138,9 +137,8 @@ namespace FbxMeshUtils
 		}
 
 		//We will use interchange only if interchange is enabled and the mesh we want to add a LOD was imported with interchange
-		const UEditorExperimentalSettings* EditorExperimentalSettings = GetDefault<UEditorExperimentalSettings>();
 		const UInterchangeAssetImportData* SelectedInterchangeAssetImportData = Cast<UInterchangeAssetImportData>(BaseStaticMesh->GetAssetImportData());
-		if (EditorExperimentalSettings->bEnableInterchangeFramework && SelectedInterchangeAssetImportData)
+		if (UInterchangeManager::IsInterchangeImportEnabled() && SelectedInterchangeAssetImportData)
 		{
 			UInterchangeSourceData* SourceData = UInterchangeManager::GetInterchangeManager().CreateSourceData(Filename);
 			//Call interchange mesh utilities to import custom LOD
@@ -425,9 +423,8 @@ namespace FbxMeshUtils
 		}
 
 		//We will use interchange only if interchange is enable and the skeletalmesh we want to add a LOD was import with interchange
-		const UEditorExperimentalSettings* EditorExperimentalSettings = GetDefault<UEditorExperimentalSettings>();
 		UInterchangeAssetImportData* SelectedInterchangeAssetImportData = Cast<UInterchangeAssetImportData>(SelectedSkelMesh->GetAssetImportData());
-		if (EditorExperimentalSettings->bEnableInterchangeFramework && SelectedInterchangeAssetImportData)
+		if (UInterchangeManager::IsInterchangeImportEnabled() && SelectedInterchangeAssetImportData)
 		{
 			UInterchangeSourceData* SourceData = UInterchangeManager::GetInterchangeManager().CreateSourceData(Filename);
 			//Call interchange mesh utilities to import custom LOD
@@ -860,8 +857,7 @@ namespace FbxMeshUtils
 		const bool bPromptOnFail = bSourceFileExists;
 
 		//We will use interchange only if interchange is enable and the skeletalmesh we want to add a LOD was import with interchange
-		const UEditorExperimentalSettings* EditorExperimentalSettings = GetDefault<UEditorExperimentalSettings>();
-		if(EditorExperimentalSettings->bEnableInterchangeFramework && SelectedInterchangeAssetImportData)
+		if(UInterchangeManager::IsInterchangeImportEnabled() && SelectedInterchangeAssetImportData)
 		{
 			auto CustomLodImportContinuation = [bNotifyCB, SkeletalMesh, StaticMesh, LODLevel](TFuture<bool> Result)
 			{
diff --git a/Engine/Source/Editor/UnrealEd/Private/FileHelpers.cpp b/Engine/Source/Editor/UnrealEd/Private/FileHelpers.cpp
--- a/Engine/Source/Editor/UnrealEd/Private/FileHelpers.cpp
+++ b/Engine/Source/Editor/UnrealEd/Private/FileHelpers.cpp
@@ -28,7 +28,6 @@
 #include "UncontrolledChangelistsModule.h"
 #include "SourceControlOperations.h"
 #include "Editor/UnrealEdEngine.h"
-#include "Settings/EditorExperimentalSettings.h"
 #include "Settings/EditorLoadingSavingSettings.h"
 #include "Factories/Factory.h"
 #include "Factories/FbxSceneImportFactory.h"
@@ -494,8 +493,7 @@ FString FEditorFileUtils::GetFilterString(EFileInteraction Interaction)
 
 				ObjectTools::GenerateFactoryFileExtensions(Factories, FileTypes, AllExtensions, FilterIndexToFactory);
 
-				const UEditorExperimentalSettings* EditorExperimentalSettings = GetDefault<UEditorExperimentalSettings>();
-				if (EditorExperimentalSettings->bEnableInterchangeFramework)
+				if (UInterchangeManager::IsInterchangeImportEnabled())
 				{
 					TArray<FString> InterchangeFileExtensions = UInterchangeManager::GetInterchangeManager().GetSupportedFormats(EInterchangeTranslatorType::Scenes);
 					ObjectTools::AppendFormatsFileExtensions(InterchangeFileExtensions, FileTypes, AllExtensions);
@@ -1483,10 +1481,8 @@ void FEditorFileUtils::Import(const FString& InFilename)
 {
 	const FScopedBusyCursor BusyCursor;
 
-	const UEditorExperimentalSettings* EditorExperimentalSettings = GetDefault<UEditorExperimentalSettings>();
-
 	UE::Interchange::FScopedSourceData ScopedSourceData(InFilename);
-	const bool bImportThroughInterchange = EditorExperimentalSettings->bEnableInterchangeFramework && UInterchangeManager::GetInterchangeManager().CanTranslateSourceData(ScopedSourceData.GetSourceData());
+	const bool bImportThroughInterchange = UInterchangeManager::GetInterchangeManager().CanTranslateSourceData(ScopedSourceData.GetSourceData());
 
 	USceneImportFactory* SceneFactory = nullptr;
 
diff --git a/Engine/Source/Editor/UnrealEd/Private/SkinWeightsUtilities.cpp b/Engine/Source/Editor/UnrealEd/Private/SkinWeightsUtilities.cpp
--- a/Engine/Source/Editor/UnrealEd/Private/SkinWeightsUtilities.cpp
+++ b/Engine/Source/Editor/UnrealEd/Private/SkinWeightsUtilities.cpp
@@ -40,7 +40,6 @@
 #include "Interfaces/ITargetPlatform.h"
 #include "Interfaces/ITargetPlatformManagerModule.h"
 #include "Misc/CoreMisc.h"
-#include "Settings/EditorExperimentalSettings.h"
 
 DEFINE_LOG_CATEGORY_STATIC(LogSkinWeightsUtilities, Log, All);
 
@@ -68,11 +67,8 @@ bool FSkinWeightsUtilities::ImportAlternateSkinWeight(USkeletalMesh* SkeletalMes
 
 	//If Interchange is enable use it if not use the old path
 
-	bool bUseInterchangeFramework = false;
+	bool bUseInterchangeFramework = UInterchangeManager::IsInterchangeImportEnabled();
 	UInterchangeManager& InterchangeManager = UInterchangeManager::GetInterchangeManager();
-#if WITH_EDITOR
-	bUseInterchangeFramework = GetDefault<UEditorExperimentalSettings>()->bEnableInterchangeFramework;
-#endif
 	const FString FileExtension = FPaths::GetExtension(AbsoluteFilePath);
 
 	FAssetToolsModule& AssetToolsModule = FModuleManager::LoadModuleChecked<FAssetToolsModule>("AssetTools");
@@ -403,10 +399,7 @@ FString FSkinWeightsUtilities::PickSkinWeightPath(int32 LODIndex, USkeletalMesh*
 {
 	FString PickedFileName("");
 	
-	bool bUseInterchangeFramework = false;
-#if WITH_EDITOR
-	bUseInterchangeFramework = GetDefault<UEditorExperimentalSettings>()->bEnableInterchangeFramework;
-#endif
+	bool bUseInterchangeFramework = UInterchangeManager::IsInterchangeImportEnabled();
 	const UInterchangeAssetImportData* SelectedInterchangeAssetImportData = Cast<UInterchangeAssetImportData>(SkeletalMesh->GetAssetImportData());
 
 	if (bUseInterchangeFramework && SelectedInterchangeAssetImportData)
diff --git a/Engine/Source/Runtime/Interchange/Engine/Private/InterchangeManager.cpp b/Engine/Source/Runtime/Interchange/Engine/Private/InterchangeManager.cpp
--- a/Engine/Source/Runtime/Interchange/Engine/Private/InterchangeManager.cpp
+++ b/Engine/Source/Runtime/Interchange/Engine/Private/InterchangeManager.cpp
@@ -40,6 +40,13 @@
 #include "Widgets/Notifications/SNotificationList.h"
 
 
+static bool GInterchangeImportEnable = true;
+static FAutoConsoleVariableRef CCvarInterchangeImportEnable(
+	TEXT("Interchange.FeatureFlags.Import.Enable"),
+	GInterchangeImportEnable,
+	TEXT("Whether Interchange import is enabled."),
+	ECVF_Default);
+
 namespace UE::Interchange::Private
 {
 	const FLogCategoryBase* GetLogInterchangePtr()
@@ -157,6 +164,17 @@ namespace UE::Interchange::Private
 	}
 }
 
+UE::Interchange::FScopedInterchangeImportEnableState::FScopedInterchangeImportEnableState(const bool bScopeValue)
+{
+	bOriginalInterchangeImportEnableState = CCvarInterchangeImportEnable->GetBool();
+	CCvarInterchangeImportEnable->Set(bScopeValue);
+}
+
+UE::Interchange::FScopedInterchangeImportEnableState::~FScopedInterchangeImportEnableState()
+{
+	CCvarInterchangeImportEnable->Set(bOriginalInterchangeImportEnableState);
+}
+
 UE::Interchange::FScopedSourceData::FScopedSourceData(const FString& Filename)
 {
 	//Found the translator
@@ -700,12 +718,16 @@ bool UInterchangeManager::RegisterWriter(const UClass* WriterClass)
 TArray<FString> UInterchangeManager::GetSupportedFormats(const EInterchangeTranslatorType ForTranslatorType) const
 {
 	TArray<FString> FileExtensions;
+	if (!IsInterchangeImportEnabled())
+	{
+		return FileExtensions;
+	}
 
 	for (const UClass* TranslatorClass : RegisteredTranslatorsClass)
 	{
 		const UInterchangeTranslatorBase* TranslatorBaseCDO = TranslatorClass->GetDefaultObject<UInterchangeTranslatorBase>();
 
-		if(EnumHasAnyFlags(TranslatorBaseCDO->GetTranslatorType(), ForTranslatorType))
+		if (EnumHasAnyFlags(TranslatorBaseCDO->GetTranslatorType(), ForTranslatorType))
 		{
 			FileExtensions.Append(TranslatorBaseCDO->GetSupportedFormats());
 		}
@@ -717,6 +739,10 @@ TArray<FString> UInterchangeManager::GetSupportedFormats(const EInterchangeTrans
 TArray<FString> UInterchangeManager::GetSupportedAssetTypeFormats(const EInterchangeTranslatorAssetType ForTranslatorAssetType) const
 {
 	TArray<FString> FileExtensions;
+	if (!IsInterchangeImportEnabled())
+	{
+		return FileExtensions;
+	}
 
 	for (const UClass* TranslatorClass : RegisteredTranslatorsClass)
 	{
@@ -733,11 +759,21 @@ TArray<FString> UInterchangeManager::GetSupportedAssetTypeFormats(const EInterch
 
 bool UInterchangeManager::CanTranslateSourceData(const UInterchangeSourceData* SourceData) const
 {
+	if (!IsInterchangeImportEnabled())
+	{
+		return false;
+	}
+
 	return GetTranslatorForSourceData(SourceData) != nullptr;
 }
 
 bool UInterchangeManager::CanReimport(const UObject* Object, TArray<FString>& OutFilenames) const
 {
+	if (!IsInterchangeImportEnabled())
+	{
+		return false;
+	}
+
 	const UClass* RegisteredFactoryClass = GetRegisteredFactoryClass(Object->GetClass());
 	if (!RegisteredFactoryClass)
 	{
@@ -767,6 +803,7 @@ bool UInterchangeManager::CanReimport(const UObject* Object, TArray<FString>& Ou
 void UInterchangeManager::StartQueuedTasks(bool bCancelAllTasks /*= false*/)
 {
 	TRACE_CPUPROFILER_EVENT_SCOPE("UInterchangeManager::StartQueuedTasks")
+	ensure(IsInterchangeImportEnabled());
 	if (!ensure(IsInGameThread()))
 	{
 		//Do not crash but we will not start any queued tasks if we are not in the game thread
@@ -894,6 +931,7 @@ TTuple<UE::Interchange::FAssetImportResultRef, UE::Interchange::FSceneImportResu
 UInterchangeManager::ImportInternal(const FString& ContentPath, const UInterchangeSourceData* SourceData, const FImportAssetParameters& ImportAssetParameters, const UE::Interchange::EImportType ImportType)
 {
 	TRACE_CPUPROFILER_EVENT_SCOPE("UInterchangeManager::ImportInternal")
+	ensure(IsInterchangeImportEnabled());
 	check(IsInGameThread());
 	static int32 GeneratedUniqueID = 0;
 	int32 UniqueId = ++GeneratedUniqueID;
@@ -1254,6 +1292,11 @@ UInterchangeManager::ImportInternal(const FString& ContentPath, const UInterchan
 	return TTuple<UE::Interchange::FAssetImportResultRef, UE::Interchange::FSceneImportResultRef>{ AsyncHelper->AssetImportResult, AsyncHelper->SceneImportResult };
 }
 
+bool UInterchangeManager::IsInterchangeImportEnabled()
+{
+	return CCvarInterchangeImportEnable->GetBool();
+}
+
 bool UInterchangeManager::ExportAsset(const UObject* Asset, bool bIsAutomated)
 {
 	return false;
diff --git a/Engine/Source/Runtime/Interchange/Engine/Public/InterchangeManager.h b/Engine/Source/Runtime/Interchange/Engine/Public/InterchangeManager.h
--- a/Engine/Source/Runtime/Interchange/Engine/Public/InterchangeManager.h
+++ b/Engine/Source/Runtime/Interchange/Engine/Public/InterchangeManager.h
@@ -42,6 +42,15 @@ namespace UE
 {
 	namespace Interchange
 	{
+		class INTERCHANGEENGINE_API FScopedInterchangeImportEnableState
+		{
+		public:
+			explicit FScopedInterchangeImportEnableState(const bool bScopeValue);
+			~FScopedInterchangeImportEnableState();
+		private:
+			bool bOriginalInterchangeImportEnableState;
+		};
+
 		class INTERCHANGEENGINE_API FScopedSourceData
 		{
 		public:
@@ -313,6 +322,9 @@ public:
 	/** Return the interchange manager singleton.*/
 	static UInterchangeManager& GetInterchangeManager();
 
+	/** Return the CVar which make interchange enable or not.*/
+	static bool IsInterchangeImportEnabled();
+
 	/** delegate type fired when new assets have been imported. Note: InCreatedObject can be NULL if import failed. Params: UFactory* InFactory, UObject* InCreatedObject */
 	DECLARE_MULTICAST_DELEGATE_OneParam(FInterchangeOnAssetPostImport, UObject*);
 	/** delegate type fired when new assets have been reimported. Note: InCreatedObject can be NULL if import failed. Params: UObject* InCreatedObject */
