The attribute storage serialize FName with FString and must change the data only if we are loading. The reference archive do not have any flag for loading or saving, so pattern like if(IsSaving)... else ... is dangerous if the data is change in the else part.

#rnx

diff --git a/Engine/Source/Runtime/Interchange/Core/Public/Types/AttributeStorage.h b/Engine/Source/Runtime/Interchange/Core/Public/Types/AttributeStorage.h
--- a/Engine/Source/Runtime/Interchange/Core/Public/Types/AttributeStorage.h
+++ b/Engine/Source/Runtime/Interchange/Core/Public/Types/AttributeStorage.h
@@ -176,15 +176,17 @@ namespace UE
 			friend FArchive& operator<<(FArchive& Ar, FAttributeKey& AttributeKey)
 			{
 				FString SerializeString;
-				if(Ar.IsSaving())
+				//When serializing FName with FString we always must do if(IsLoading()) else to avoid changing the data
+				//When a reference serialization is execute (reference do not have IsLoading or IsSaving set)
+				if (Ar.IsLoading())
 				{
-					SerializeString = AttributeKey.Key.ToString();
 					Ar << SerializeString;
+					AttributeKey.Key = FName(*SerializeString);
 				}
 				else
 				{
+					SerializeString = AttributeKey.Key.ToString();
 					Ar << SerializeString;
-					AttributeKey.Key = FName(*SerializeString);
 				}
 				return Ar;
 			}
