Fixed warning when deleting assets imported with interchange

diff --git a/Engine/Source/Runtime/Interchange/Engine/Private/InterchangeManager.cpp b/Engine/Source/Runtime/Interchange/Engine/Private/InterchangeManager.cpp
--- a/Engine/Source/Runtime/Interchange/Engine/Private/InterchangeManager.cpp
+++ b/Engine/Source/Runtime/Interchange/Engine/Private/InterchangeManager.cpp
@@ -613,6 +613,10 @@ UInterchangePipelineBase* UE::Interchange::GeneratePipelineInstance(const FSoftO
 	{
 		UE_LOG(LogInterchangeEngine, Error, TEXT("Cannot generate a pipeline instance because the pipeline asset %s type is unknown."), *PipelineInstance.GetWithoutSubPath().ToString());
 	}
+
+	// Make sure that the instance does not carry over standalone and public flags as they are not actual assets to be persisted
+	GeneratedPipeline->ClearFlags(EObjectFlags::RF_Standalone|EObjectFlags::RF_Public);
+
 	return GeneratedPipeline;
 }
 
@@ -1154,7 +1158,8 @@ UInterchangeManager::ImportInternal(const FString& ContentPath, const UInterchan
 				{
 					//Duplicate the pipeline saved in the asset import data
 					UInterchangePipelineBase* GeneratedPipeline = Cast<UInterchangePipelineBase>(StaticDuplicateObject(SourcePipeline, GetTransientPackage()));
-
+					// Make sure that the instance does not carry over standalone and public flags as they are not actual assets to be persisted
+					GeneratedPipeline->ClearFlags(EObjectFlags::RF_Standalone|EObjectFlags::RF_Public);
 					if (bImportScene)
 					{
 						GeneratedPipeline->AdjustSettingsForContext(EInterchangePipelineContext::SceneReimport, nullptr);
