Fix PVS static analysis

#rnx

diff --git a/Engine/Source/Runtime/Interchange/Engine/Private/Tasks/InterchangeTaskCompletion.cpp b/Engine/Source/Runtime/Interchange/Engine/Private/Tasks/InterchangeTaskCompletion.cpp
--- a/Engine/Source/Runtime/Interchange/Engine/Private/Tasks/InterchangeTaskCompletion.cpp
+++ b/Engine/Source/Runtime/Interchange/Engine/Private/Tasks/InterchangeTaskCompletion.cpp
@@ -55,7 +55,10 @@ void UE::Interchange::FTaskPreCompletion::DoTask(ENamedThreads::Type CurrentThre
 				for (const FImportAsyncHelper::FImportedObjectInfo& ObjectInfo : ObjectInfosPerSourceIndexPair.Value)
 				{
 					//Cancel factories so they can do proper cleanup
-					ObjectInfo.Factory->Cancel();
+					if (ObjectInfo.Factory)
+					{
+						ObjectInfo.Factory->Cancel();
+					}
 				}
 				break;
 			}
