Interchange: Fix analytic event name to not be dynamic

#rnx

diff --git a/Engine/Source/Runtime/Interchange/Engine/Private/InterchangeManager.cpp b/Engine/Source/Runtime/Interchange/Engine/Private/InterchangeManager.cpp
--- a/Engine/Source/Runtime/Interchange/Engine/Private/InterchangeManager.cpp
+++ b/Engine/Source/Runtime/Interchange/Engine/Private/InterchangeManager.cpp
@@ -64,7 +64,7 @@ namespace UE::Interchange::Private
 		return const_cast<FInterchangeImportSettings&>(FInterchangeProjectSettingsUtils::GetImportSettings(InterchangeProjectSettings, ImportType == EImportType::ImportType_Scene));
 	}
 
-	void FillPipelineAnalyticData(UInterchangePipelineBase* Pipeline, TArray<FAnalyticsEventAttribute>& Attribs, const int32 UniqueId, const FString& ParentPipeline)
+	void FillPipelineAnalyticData(UInterchangePipelineBase* Pipeline, const int32 UniqueId, const FString& ParentPipeline)
 	{
 		if (!FEngineAnalytics::IsAvailable())
 		{
@@ -74,14 +74,6 @@ namespace UE::Interchange::Private
 		int32 PortFlags = 0;
 		UClass* Class = Pipeline->GetClass();
 		FString PipelineChainName = ParentPipeline.IsEmpty() ? Pipeline->GetName() : ParentPipeline + TEXT(".") + Pipeline->GetName();
-		if (ParentPipeline.IsEmpty())
-		{
-			Attribs.Add(FAnalyticsEventAttribute(TEXT("Pipeline.Name"), PipelineChainName));
-		}
-		else
-		{
-			Attribs.Add(FAnalyticsEventAttribute(TEXT("SubPipeline.Name"), PipelineChainName));
-		}
 
 		TArray<FAnalyticsEventAttribute> PipelineAttribs;
 		PipelineAttribs.Add(FAnalyticsEventAttribute(TEXT("UniqueId"), UniqueId));
@@ -140,7 +132,7 @@ namespace UE::Interchange::Private
 				if (SubPipeline->IsInOuter(Pipeline))
 				{
 					//Go recursive with subObject, like if they are part of the same object
-					FillPipelineAnalyticData(SubPipeline, Attribs, UniqueId, PipelineChainName);
+					FillPipelineAnalyticData(SubPipeline, UniqueId, PipelineChainName);
 				}
 			}
 			else
@@ -160,7 +152,7 @@ namespace UE::Interchange::Private
 			}
 		}
 		
-		FString EventString = FString::Printf(TEXT("Interchange.Usage.Import.Pipeline.%s"), *PipelineChainName);
+		FString EventString = TEXT("Interchange.Usage.Import.Pipeline");
 		FEngineAnalytics::GetProvider().RecordEvent(EventString, PipelineAttribs);
 	}
 }
@@ -230,46 +222,26 @@ void UE::Interchange::FImportAsyncHelper::SendAnalyticImportEndData()
 		return;
 	}
 
-	auto CollectObject = [&Attribs](bool bIsReimport, UObject* Object)
-	{
-		const int32 ObjectUniqueId = static_cast<int32>(Object->GetUniqueID());
-		const FString ObjectPathName = Object->GetPathName();
-		const FString ObjectClassName = Object->GetClass()->GetName();
-
-		Attribs.Add(FAnalyticsEventAttribute(FString::Printf(TEXT("Object.%d.Class"), ObjectUniqueId), ObjectClassName));
-		Attribs.Add(FAnalyticsEventAttribute(FString::Printf(TEXT("Object.%d.IsReimport"), ObjectUniqueId), bIsReimport));
-	};
-
+	int32 ImportedObjectCount = 0;
 	for (const TPair<int32, TArray<FImportedObjectInfo>>& SourceIndexAndImportedAssets : ImportedAssetsPerSourceIndex)
 	{
-		const TArray<FImportedObjectInfo>& ImportedObjects = SourceIndexAndImportedAssets.Value;
-		for (const FImportedObjectInfo& ObjectInfo : ImportedObjects)
-		{
-			if (!ObjectInfo.ImportedObject)
-			{
-				continue;
-			}
-			CollectObject(ObjectInfo.bIsReimport, ObjectInfo.ImportedObject);
-		}
+		ImportedObjectCount += SourceIndexAndImportedAssets.Value.Num();
 	}
 
 	for (const TPair<int32, TArray<FImportedObjectInfo>>& SourceIndexAndImportedScneObjects : ImportedSceneObjectsPerSourceIndex)
 	{
-		const TArray<FImportedObjectInfo>& ImportedObjects = SourceIndexAndImportedScneObjects.Value;
-		for (const FImportedObjectInfo& ObjectInfo : ImportedObjects)
-		{
-			if (!ObjectInfo.ImportedObject)
-			{
-				continue;
-			}
-			CollectObject(ObjectInfo.bIsReimport, ObjectInfo.ImportedObject);
-		}
+		ImportedObjectCount += SourceIndexAndImportedScneObjects.Value.Num();
 	}
 
+	Attribs.Add(FAnalyticsEventAttribute(TEXT("ImportObjectCount"), ImportedObjectCount));
+
 	//Report any warning or error message
-	auto CollectResultContainer = [&Attribs](const UInterchangeResultsContainer* ResultContainer)
+	TArray<FString> WarningMessages;
+	TArray<FString> ErrorMessages;
+	auto CollectResultContainer = [&WarningMessages, &ErrorMessages](const UInterchangeResultsContainer* ResultContainer)
 	{
 		TArray<UInterchangeResult*> InterchangeResults = ResultContainer->GetResults();
+		
 		for (const UInterchangeResult* InterchangeResult : InterchangeResults)
 		{
 			switch (InterchangeResult->GetResultType() )
@@ -277,25 +249,34 @@ void UE::Interchange::FImportAsyncHelper::SendAnalyticImportEndData()
 			case EInterchangeResultType::Success:
 				break;
 			case EInterchangeResultType::Warning:
-				Attribs.Add(FAnalyticsEventAttribute(TEXT("Message.Warning"), InterchangeResult->GetText().ToString()));
+				WarningMessages.Add(TEXT("{") + InterchangeResult->GetText().ToString() + TEXT("}"));
 				break;
 			case EInterchangeResultType::Error:
-				Attribs.Add(FAnalyticsEventAttribute(TEXT("Message.Error"), InterchangeResult->GetText().ToString()));
+				ErrorMessages.Add(TEXT("{") + InterchangeResult->GetText().ToString() + TEXT("}"));
 				break;
 			}
 		}
+		
 	};
 
 	if (const UInterchangeResultsContainer* ResultContainer = AssetImportResult->GetResults())
 	{
 		CollectResultContainer(ResultContainer);
 	}
-
 	if (const UInterchangeResultsContainer* ResultContainer = SceneImportResult->GetResults())
 	{
 		CollectResultContainer(ResultContainer);
 	}
 
+	if (WarningMessages.Num() > 0)
+	{
+		Attribs.Add(FAnalyticsEventAttribute(TEXT("WarningMessages"), WarningMessages));
+	}
+	if (ErrorMessages.Num() > 0)
+	{
+		Attribs.Add(FAnalyticsEventAttribute(TEXT("ErrorMessages"), ErrorMessages));
+	}
+
 	FString EventString = TEXT("Interchange.Usage.ImportResult");
 	FEngineAnalytics::GetProvider().RecordEvent(EventString, Attribs);
 }
@@ -1208,7 +1189,7 @@ UInterchangeManager::ImportInternal(const FString& ContentPath, const UInterchan
 				{
 					AsyncHelper->OriginalPipelines.Add(Pipeline);
 				}
-				UE::Interchange::Private::FillPipelineAnalyticData(Pipeline, Attribs, UniqueId, FString());
+				UE::Interchange::Private::FillPipelineAnalyticData(Pipeline, UniqueId, FString());
 			}
 		}
 		else
@@ -1270,7 +1251,7 @@ UInterchangeManager::ImportInternal(const FString& ContentPath, const UInterchan
 									GeneratedPipeline->LoadSettings(PipelineStackName);
 								}
 					
-								UE::Interchange::Private::FillPipelineAnalyticData(GeneratedPipeline, Attribs, UniqueId, FString());
+								UE::Interchange::Private::FillPipelineAnalyticData(GeneratedPipeline, UniqueId, FString());
 
 								AsyncHelper->Pipelines.Add(GeneratedPipeline);
 
@@ -1315,7 +1296,7 @@ UInterchangeManager::ImportInternal(const FString& ContentPath, const UInterchan
 			}
 			AsyncHelper->Pipelines.Add(GeneratedPipeline);
 			AsyncHelper->OriginalPipelines.Add(GeneratedPipeline);
-			UE::Interchange::Private::FillPipelineAnalyticData(GeneratedPipeline, Attribs, UniqueId, FString());
+			UE::Interchange::Private::FillPipelineAnalyticData(GeneratedPipeline, UniqueId, FString());
 		}
 	}
 
