fix UInterchangeManager::GetInterchangeManager crashing in app exit

diff --git a/Engine/Source/Editor/UnrealEd/Private/Editor.cpp b/Engine/Source/Editor/UnrealEd/Private/Editor.cpp
--- a/Engine/Source/Editor/UnrealEd/Private/Editor.cpp
+++ b/Engine/Source/Editor/UnrealEd/Private/Editor.cpp
@@ -926,10 +926,14 @@ FReimportManager::FReimportManager()
 FReimportManager::~FReimportManager()
 {
 	Handlers.Empty();
+
+	/*
+	// you can't do this because ~FReimportManager is called from FReimportManager::Instance at cexit shutdown time
 	if (InterchangePostReimportedDelegateHandle.IsValid())
 	{
 		UInterchangeManager::GetInterchangeManager().OnAssetPostImport.Remove(InterchangePostReimportedDelegateHandle);
 	}
+	*/
 }
 
 int32 FReimportHandler::GetPriority() const
diff --git a/Engine/Source/Runtime/Interchange/Engine/Private/InterchangeManager.cpp b/Engine/Source/Runtime/Interchange/Engine/Private/InterchangeManager.cpp
--- a/Engine/Source/Runtime/Interchange/Engine/Private/InterchangeManager.cpp
+++ b/Engine/Source/Runtime/Interchange/Engine/Private/InterchangeManager.cpp
@@ -604,18 +604,37 @@ UInterchangePipelineBase* UE::Interchange::GeneratePipelineInstance(const FSoftO
 UInterchangeManager& UInterchangeManager::GetInterchangeManager()
 {
 	static TStrongObjectPtr<UInterchangeManager> InterchangeManager = nullptr;
-	
+
 	//This boolean will be true after we delete the singleton
 	static bool InterchangeManagerScopeOfLifeEnded = false;
 
 	if (!InterchangeManager.IsValid())
 	{
+		if ( InterchangeManagerScopeOfLifeEnded )
+		{
+			#if 1
+
+			//Avoid hard crash if someone call the manager after we delete it, but send a callstack to the crash manager
+			ensure(!InterchangeManagerScopeOfLifeEnded);
+
+			#else
+
+			// -> no, this is being called after "LogExit: Exiting."
+			//	at that point you cannot check or ensure or log or anything
+			//	do a low level break instead
+
+			PLATFORM_BREAK(); //__debugbreak();
+
+			// just return a nullptr ; you will crash after returning
+			// this return is just so you can step out in the debugger to see who's calling this
+			return *(InterchangeManager.Get());
+
+			#endif
+		}
+
 		//We cannot create a TStrongObjectPtr outside of the main thread, we also need a valid Transient package
 		check(IsInGameThread() && GetTransientPackage());
 
-		//Avoid hard crash if someone call the manager after we delete it, but send a callstack to the crash manager
-		ensure(!InterchangeManagerScopeOfLifeEnded);
-
 		InterchangeManager = TStrongObjectPtr<UInterchangeManager>(NewObject<UInterchangeManager>(GetTransientPackage(), NAME_None, EObjectFlags::RF_NoFlags));
 		
 		//We cancel any running task when we pre exit the engine
