Interchange: New created package now have the "PKG_NewlyCreated" flag set.

#rnx

diff --git a/Engine/Source/Runtime/Interchange/Engine/Private/Tasks/InterchangeTaskCreateAsset.cpp b/Engine/Source/Runtime/Interchange/Engine/Private/Tasks/InterchangeTaskCreateAsset.cpp
--- a/Engine/Source/Runtime/Interchange/Engine/Private/Tasks/InterchangeTaskCreateAsset.cpp
+++ b/Engine/Source/Runtime/Interchange/Engine/Private/Tasks/InterchangeTaskCreateAsset.cpp
@@ -216,6 +216,8 @@ void UE::Interchange::FTaskCreatePackage::DoTask(ENamedThreads::Type CurrentThre
 			return;
 		}
 
+		Pkg->SetPackageFlags(PKG_NewlyCreated);
+
 		//Import Asset describe by the node
 		UInterchangeFactoryBase::FCreateAssetParams CreateAssetParams;
 		CreateAssetParams.AssetName = AssetName;
