Interchange Blueprint pipeline base should not override GetBlueprintClass().

#rnx

diff --git a/Engine/Source/Runtime/Interchange/Engine/Public/InterchangeBlueprintPipelineBase.h b/Engine/Source/Runtime/Interchange/Engine/Public/InterchangeBlueprintPipelineBase.h
--- a/Engine/Source/Runtime/Interchange/Engine/Public/InterchangeBlueprintPipelineBase.h
+++ b/Engine/Source/Runtime/Interchange/Engine/Public/InterchangeBlueprintPipelineBase.h
@@ -22,7 +22,4 @@ public:
 		//We must make sure the GeneratedClass is generated after the blueprint is loaded
 		bRecompileOnLoad = true;
 	}
-#if WITH_EDITOR
-	virtual UClass* GetBlueprintClass() const override { return UInterchangePipelineBase::StaticClass(); }
-#endif
 };
\ No newline at end of file
