Interchange: Fix synchronous import to wait properly until the job is done before returning.

#rnx

diff --git a/Engine/Plugins/Interchange/Editor/Source/Pipelines/Private/InterchangeGraphInspectorPipeline.cpp b/Engine/Plugins/Interchange/Editor/Source/Pipelines/Private/InterchangeGraphInspectorPipeline.cpp
--- a/Engine/Plugins/Interchange/Editor/Source/Pipelines/Private/InterchangeGraphInspectorPipeline.cpp
+++ b/Engine/Plugins/Interchange/Editor/Source/Pipelines/Private/InterchangeGraphInspectorPipeline.cpp
@@ -3,9 +3,11 @@
 
 #include "CoreMinimal.h"
 
+#include "CoreGlobals.h"
 #include "Framework/Application/SlateApplication.h"
 #include "InterchangeSourceData.h"
 #include "Interfaces/IMainFrameModule.h"
+#include "Misc/App.h"
 #include "Misc/Paths.h"
 #include "Nodes/InterchangeBaseNodeContainer.h"
 #include "SInterchangeGraphInspectorWindow.h"
@@ -17,6 +19,12 @@
 
 void UInterchangeGraphInspectorPipeline::ExecutePreImportPipeline(UInterchangeBaseNodeContainer* BaseNodeContainer, const TArray<UInterchangeSourceData*>& SourceDatas)
 {
+	//Do not pop dialog if we are unattended or doing automation
+	if (FApp::IsUnattended() || IsRunningCommandlet())
+	{
+		return;
+	}
+
 	//Create and show the graph inspector UI dialog
 	TSharedPtr<SWindow> ParentWindow;
 	if (FModuleManager::Get().IsModuleLoaded("MainFrame"))
diff --git a/Engine/Source/Runtime/Interchange/Engine/Private/InterchangeManager.cpp b/Engine/Source/Runtime/Interchange/Engine/Private/InterchangeManager.cpp
--- a/Engine/Source/Runtime/Interchange/Engine/Private/InterchangeManager.cpp
+++ b/Engine/Source/Runtime/Interchange/Engine/Private/InterchangeManager.cpp
@@ -6,6 +6,7 @@
 #include "AssetRegistry/AssetDataTagMap.h"
 #include "AssetRegistry/AssetRegistryModule.h"
 #include "Async/Async.h"
+#include "CoreGlobals.h"
 #include "CoreMinimal.h"
 #include "Engine/Blueprint.h"
 #include "EngineAnalytics.h"
@@ -293,16 +294,8 @@ void UE::Interchange::FImportAsyncHelper::ReleaseTranslatorsSource()
 	}
 }
 
-void UE::Interchange::FImportAsyncHelper::InitCancel()
+FGraphEventArray UE::Interchange::FImportAsyncHelper::GetCompletionTaskGraphEvent()
 {
-	bCancel = true;
-	ReleaseTranslatorsSource();
-}
-
-void UE::Interchange::FImportAsyncHelper::CancelAndWaitUntilDoneSynchronously()
-{
-	bCancel = true;
-
 	FGraphEventArray TasksToComplete;
 
 	TasksToComplete.Append(TranslatorTasks);
@@ -313,6 +306,10 @@ void UE::Interchange::FImportAsyncHelper::CancelAndWaitUntilDoneSynchronously()
 		TasksToComplete.Add(ParsingTask);
 	}
 
+	//Parsing task must be done before the other tasks get added
+	FTaskGraphInterface::Get().WaitUntilTasksComplete(TasksToComplete, ENamedThreads::GameThread);
+	TasksToComplete.Reset();
+
 	TasksToComplete.Append(CreatePackageTasks);
 	TasksToComplete.Append(CreateAssetTasks);
 	TasksToComplete.Append(SceneTasks);
@@ -326,20 +323,20 @@ void UE::Interchange::FImportAsyncHelper::CancelAndWaitUntilDoneSynchronously()
 	{
 		TasksToComplete.Add(PreCompletionTask);
 	}
+	
 	if (CompletionTask.GetReference())
 	{
 		//Completion task will make sure any created asset before canceling will be mark for delete
 		TasksToComplete.Add(CompletionTask);
 	}
 
-	//Block until all task are completed, it should be fast since bCancel is true
-	if (TasksToComplete.Num())
-	{
-		FTaskGraphInterface::Get().WaitUntilTasksComplete(TasksToComplete, ENamedThreads::GameThread);
-	}
+	return TasksToComplete;
+}
 
-	AssetImportResult->SetDone();
-	SceneImportResult->SetDone();
+void UE::Interchange::FImportAsyncHelper::InitCancel()
+{
+	bCancel = true;
+	ReleaseTranslatorsSource();
 }
 
 void UE::Interchange::FImportAsyncHelper::CleanUp()
@@ -641,17 +638,12 @@ UInterchangeManager& UInterchangeManager::GetInterchangeManager()
 		//We cancel any running task when we pre exit the engine
 		FCoreDelegates::OnEnginePreExit.AddLambda([]()
 		{
-			//In editor the user cannot exit the editor if the interchange manager has active task.
-			//But if we are not running the editor its possible to get here, so block the main thread until all
-			//cancel tasks are done.
-			if(GIsEditor)
-			{
-				ensure(InterchangeManager->ImportTasks.Num() == 0);
-			}
-			else
-			{
-				InterchangeManager->CancelAllTasksSynchronously();
-			}
+			//If a user run a editor commandlet, we want to finish the import before the editor close.
+			//In editor mode, the user cannot close the editor if an import task is running, so we should not endup here.
+			const bool bCancel = !GIsEditor;
+			//Synchronously wait all task to finish
+			InterchangeManager->WaitUntilAllTasksDone(bCancel);
+
 			//Task should have been cancel in the Engine pre exit callback
 			ensure(InterchangeManager->ImportTasks.Num() == 0);
 			InterchangeManager->OnPreDestroyInterchangeManager.Broadcast();
@@ -957,7 +949,9 @@ void UInterchangeManager::StartQueuedTasks(bool bCancelAllTasks /*= false*/)
 
 bool UInterchangeManager::ImportAsset(const FString& ContentPath, const UInterchangeSourceData* SourceData, const FImportAssetParameters& ImportAssetParameters)
 {
-	return ImportAssetAsync( ContentPath, SourceData, ImportAssetParameters )->IsValid();
+	UE::Interchange::FAssetImportResultRef InterchangeResult = ImportAssetAsync(ContentPath, SourceData, ImportAssetParameters);
+	InterchangeResult->WaitUntilDone();
+	return InterchangeResult->IsValid();
 }
 
 UE::Interchange::FAssetImportResultRef UInterchangeManager::ImportAssetAsync(const FString& ContentPath, const UInterchangeSourceData* SourceData, const FImportAssetParameters& ImportAssetParameters)
@@ -969,7 +963,8 @@ bool UInterchangeManager::ImportScene(const FString& ContentPath, const UInterch
 {
 	using namespace UE::Interchange;
 	TTuple<FAssetImportResultRef, FSceneImportResultRef> ImportResults = ImportInternal(ContentPath, SourceData, ImportAssetParameters, UE::Interchange::EImportType::ImportType_Scene);
-
+	ImportResults.Get<0>()->WaitUntilDone();
+	ImportResults.Get<1>()->WaitUntilDone();
 	return ImportResults.Get<0>()->IsValid() && ImportResults.Get<1>()->IsValid();
 }
 
@@ -1129,7 +1124,8 @@ UInterchangeManager::ImportInternal(const FString& ContentPath, const UInterchan
 		const bool bShowPipelineStacksConfigurationDialog = !bIsUnattended
 															&& FInterchangeProjectSettingsUtils::ShouldShowPipelineStacksConfigurationDialog(bImportScene, *SourceData)
 															&& !bImportAllWithDefault
-															&& !bImportCanceled;
+															&& !bImportCanceled
+															&& !IsRunningCommandlet();
 #else
 		const bool bShowPipelineStacksConfigurationDialog = false;
 #endif
@@ -1678,25 +1674,26 @@ void UInterchangeManager::CancelAllTasks()
 	//Tasks should all finish quite fast now
 };
 
-void UInterchangeManager::CancelAllTasksSynchronously()
+void UInterchangeManager::WaitUntilAllTasksDone(bool bCancel)
 {
-	//Start the cancel process by cancelling all current task
-	CancelAllTasks();
+	check(IsInGameThread());
+	if (bCancel)
+	{
+		//Start the cancel process by cancelling all current task
+		CancelAllTasks();
+	}
 
-	//Now wait for each task to be completed on the main thread
 	while (ImportTasks.Num() > 0)
 	{
-		int32 ImportTaskCount = ImportTasks.Num();
 		TSharedPtr<UE::Interchange::FImportAsyncHelper, ESPMode::ThreadSafe> AsyncHelper = ImportTasks[0];
 		if (AsyncHelper.IsValid())
 		{
-			//Cancel any on going interchange activity this is blocking but necessary.
-			AsyncHelper->CancelAndWaitUntilDoneSynchronously();
-			ensure(ImportTaskCount > ImportTasks.Num());
 			TWeakPtr<UE::Interchange::FImportAsyncHelper, ESPMode::ThreadSafe> WeakAsyncHelper = AsyncHelper;
-			//free the async helper
+			FGraphEventArray TasksToComplete = AsyncHelper->GetCompletionTaskGraphEvent();
+			//Release the shared pointer before waiting to be sure the async helper can be destroy in the completion task
 			AsyncHelper = nullptr;
-			//We verify that the weak pointer is invalid after releasing the async helper
+			FTaskGraphInterface::Get().WaitUntilTasksComplete(TasksToComplete, ENamedThreads::GameThread);
+			//We verify that the weak pointer is invalid after the task completed
 			ensure(!WeakAsyncHelper.IsValid());
 		}
 	}
diff --git a/Engine/Source/Runtime/Interchange/Engine/Private/Tasks/InterchangeTaskCompletion.h b/Engine/Source/Runtime/Interchange/Engine/Private/Tasks/InterchangeTaskCompletion.h
--- a/Engine/Source/Runtime/Interchange/Engine/Private/Tasks/InterchangeTaskCompletion.h
+++ b/Engine/Source/Runtime/Interchange/Engine/Private/Tasks/InterchangeTaskCompletion.h
@@ -88,7 +88,8 @@ namespace UE
 			}
 			static FORCEINLINE ESubsequentsMode::Type GetSubsequentsMode()
 			{
-				return ESubsequentsMode::FireAndForget;
+				//In case we need to know when the task is done we need track subsequent to get a valid FGraphEventRef when we create the task
+				return ESubsequentsMode::TrackSubsequents;
 			}
 
 			FORCEINLINE TStatId GetStatId() const
diff --git a/Engine/Source/Runtime/Interchange/Engine/Public/InterchangeManager.h b/Engine/Source/Runtime/Interchange/Engine/Public/InterchangeManager.h
--- a/Engine/Source/Runtime/Interchange/Engine/Public/InterchangeManager.h
+++ b/Engine/Source/Runtime/Interchange/Engine/Public/InterchangeManager.h
@@ -242,8 +242,13 @@ namespace UE
 
 			void SendAnalyticImportEndData();
 			void ReleaseTranslatorsSource();
+
+			/**
+			 * Wait synchronously after graph parsing task is done and return the GraphEventArray up to the completion TaskGraphEvent
+			 */
+			FGraphEventArray GetCompletionTaskGraphEvent();
+
 			void InitCancel();
-			void CancelAndWaitUntilDoneSynchronously();
 
 			void CleanUp();
 		};
@@ -520,14 +525,15 @@ protected:
 	/**
 	 * This function cancel all task and finish them has fast as possible.
 	 * We use this if the user cancel the work or if the editor is exiting.
-	 * @note - This is a blocking call until the tasks are completed.
+	 * @note - This is a asynchronous call, tasks will be completed (cancel) soon.
 	 */
 	void CancelAllTasks();
 
 	/**
-	 * Same has CancelAllTasks, but will wait all task are done before exiting the function
+	 * Wait synchronously that all tasks are done
 	 */
-	void CancelAllTasksSynchronously();
+	void WaitUntilAllTasksDone(bool bCancel);
+
 
 	/**
 	 * If we set the mode to active we will setup the timer and add the thread that will block the GC.
