Interchange:
- Refactor pipeline and factory callback names.
- Add a pipeline call back before calling post edit change, this callback will be use by the pipeline to do named workflow like material ordering, vertex channel ordering or any other pipeline stuff we need to do after the asset is created but before calling PostEditChange.

#rnx

diff --git a/Engine/Plugins/Enterprise/DatasmithImporter/Source/DatasmithImporter/Private/DatasmithTextureImporter.cpp b/Engine/Plugins/Enterprise/DatasmithImporter/Source/DatasmithImporter/Private/DatasmithTextureImporter.cpp
--- a/Engine/Plugins/Enterprise/DatasmithImporter/Source/DatasmithImporter/Private/DatasmithTextureImporter.cpp
+++ b/Engine/Plugins/Enterprise/DatasmithImporter/Source/DatasmithImporter/Private/DatasmithTextureImporter.cpp
@@ -408,7 +408,7 @@ UE::Interchange::FAssetImportResultRef FDatasmithTextureImporter::CreateTextureA
 	return UInterchangeManager::GetInterchangeManager().ImportAssetAsync( ContentPath, ScopedSourceData.GetSourceData(), ImportAssetParameters );
 }
 
-void UDatasmithTexturePipeline::ExecutePreImportPipeline(UInterchangeBaseNodeContainer* BaseNodeContainer, const TArray<UInterchangeSourceData*>& SourceDatas)
+void UDatasmithTexturePipeline::ExecutePipeline(UInterchangeBaseNodeContainer* BaseNodeContainer, const TArray<UInterchangeSourceData*>& SourceDatas)
 {
 	if ( !TextureElement.IsValid() )
 	{
diff --git a/Engine/Plugins/Enterprise/DatasmithImporter/Source/DatasmithImporter/Public/DatasmithTextureImporter.h b/Engine/Plugins/Enterprise/DatasmithImporter/Source/DatasmithImporter/Public/DatasmithTextureImporter.h
--- a/Engine/Plugins/Enterprise/DatasmithImporter/Source/DatasmithImporter/Public/DatasmithTextureImporter.h
+++ b/Engine/Plugins/Enterprise/DatasmithImporter/Source/DatasmithImporter/Public/DatasmithTextureImporter.h
@@ -46,7 +46,7 @@ class UDatasmithTexturePipeline : public UInterchangePipelineBase
 	GENERATED_BODY()
 
 public:
-	virtual void ExecutePreImportPipeline(UInterchangeBaseNodeContainer* BaseNodeContainer, const TArray<UInterchangeSourceData*>& SourceDatas) override;
+	virtual void ExecutePipeline(UInterchangeBaseNodeContainer* BaseNodeContainer, const TArray<UInterchangeSourceData*>& SourceDatas) override;
 
 	bool IsScripted() override
 	{
diff --git a/Engine/Plugins/Experimental/Enterprise/DatasmithInterchange/Source/DatasmithInterchange/Private/InterchangeDatasmithLevelPipeline.cpp b/Engine/Plugins/Experimental/Enterprise/DatasmithInterchange/Source/DatasmithInterchange/Private/InterchangeDatasmithLevelPipeline.cpp
--- a/Engine/Plugins/Experimental/Enterprise/DatasmithInterchange/Source/DatasmithInterchange/Private/InterchangeDatasmithLevelPipeline.cpp
+++ b/Engine/Plugins/Experimental/Enterprise/DatasmithInterchange/Source/DatasmithInterchange/Private/InterchangeDatasmithLevelPipeline.cpp
@@ -19,11 +19,11 @@
 #include "Engine/StaticMeshActor.h"
 #include "Materials/MaterialInterface.h"
 
-void UInterchangeDatasmithLevelPipeline::ExecutePreImportPipeline(UInterchangeBaseNodeContainer* NodeContainer, const TArray<UInterchangeSourceData*>& InSourceDatas)
+void UInterchangeDatasmithLevelPipeline::ExecutePipeline(UInterchangeBaseNodeContainer* NodeContainer, const TArray<UInterchangeSourceData*>& InSourceDatas)
 {
 	using namespace UE::DatasmithInterchange;
 
-	Super::ExecutePreImportPipeline(NodeContainer, InSourceDatas);
+	Super::ExecutePipeline(NodeContainer, InSourceDatas);
 
 	// Add material factory dependencies for mesh actors where all overrides are filled with the same material
 	for (UInterchangeMeshActorFactoryNode* MeshActorFactoryNode : NodeUtils::GetNodes<UInterchangeMeshActorFactoryNode>(NodeContainer))
diff --git a/Engine/Plugins/Experimental/Enterprise/DatasmithInterchange/Source/DatasmithInterchange/Private/InterchangeDatasmithLevelPipeline.h b/Engine/Plugins/Experimental/Enterprise/DatasmithInterchange/Source/DatasmithInterchange/Private/InterchangeDatasmithLevelPipeline.h
--- a/Engine/Plugins/Experimental/Enterprise/DatasmithInterchange/Source/DatasmithInterchange/Private/InterchangeDatasmithLevelPipeline.h
+++ b/Engine/Plugins/Experimental/Enterprise/DatasmithInterchange/Source/DatasmithInterchange/Private/InterchangeDatasmithLevelPipeline.h
@@ -20,7 +20,7 @@ class UInterchangeDatasmithLevelPipeline : public UInterchangeGenericLevelPipeli
 	GENERATED_BODY()
 
 protected:
-	virtual void ExecutePreImportPipeline(UInterchangeBaseNodeContainer* BaseNodeContainer, const TArray<UInterchangeSourceData*>& SourceDatas) override;
+	virtual void ExecutePipeline(UInterchangeBaseNodeContainer* BaseNodeContainer, const TArray<UInterchangeSourceData*>& SourceDatas) override;
 	virtual void ExecutePostImportPipeline(const UInterchangeBaseNodeContainer* BaseNodeContainer, const FString& NodeKey, UObject* CreatedAsset, bool bIsAReimport) override;
 
 	virtual UInterchangeActorFactoryNode* CreateActorFactoryNode(const UInterchangeSceneNode* SceneNode, const UInterchangeBaseNode* TranslatedAssetNode) const override;
diff --git a/Engine/Plugins/Experimental/Enterprise/DatasmithInterchange/Source/DatasmithInterchange/Private/InterchangeDatasmithMaterialPipeline.cpp b/Engine/Plugins/Experimental/Enterprise/DatasmithInterchange/Source/DatasmithInterchange/Private/InterchangeDatasmithMaterialPipeline.cpp
--- a/Engine/Plugins/Experimental/Enterprise/DatasmithInterchange/Source/DatasmithInterchange/Private/InterchangeDatasmithMaterialPipeline.cpp
+++ b/Engine/Plugins/Experimental/Enterprise/DatasmithInterchange/Source/DatasmithInterchange/Private/InterchangeDatasmithMaterialPipeline.cpp
@@ -30,7 +30,7 @@ UInterchangeDatasmithMaterialPipeline::UInterchangeDatasmithMaterialPipeline()
 	MaterialImport = EInterchangeMaterialImportOption::ImportAsMaterials;
 }
 
-void UInterchangeDatasmithMaterialPipeline::ExecutePreImportPipeline(UInterchangeBaseNodeContainer* NodeContainer, const TArray<UInterchangeSourceData*>& InSourceDatas)
+void UInterchangeDatasmithMaterialPipeline::ExecutePipeline(UInterchangeBaseNodeContainer* NodeContainer, const TArray<UInterchangeSourceData*>& InSourceDatas)
 {
 	using namespace UE::DatasmithInterchange;
 
@@ -50,7 +50,7 @@ void UInterchangeDatasmithMaterialPipeline::ExecutePreImportPipeline(UInterchang
 			}
 		});
 
-	Super::ExecutePreImportPipeline(NodeContainer, InSourceDatas);
+	Super::ExecutePipeline(NodeContainer, InSourceDatas);
 
 	UpdateMaterialFactoryNodes(ShaderNodes);
 
diff --git a/Engine/Plugins/Experimental/Enterprise/DatasmithInterchange/Source/DatasmithInterchange/Private/InterchangeDatasmithMaterialPipeline.h b/Engine/Plugins/Experimental/Enterprise/DatasmithInterchange/Source/DatasmithInterchange/Private/InterchangeDatasmithMaterialPipeline.h
--- a/Engine/Plugins/Experimental/Enterprise/DatasmithInterchange/Source/DatasmithInterchange/Private/InterchangeDatasmithMaterialPipeline.h
+++ b/Engine/Plugins/Experimental/Enterprise/DatasmithInterchange/Source/DatasmithInterchange/Private/InterchangeDatasmithMaterialPipeline.h
@@ -24,7 +24,7 @@ class DATASMITHINTERCHANGE_API UInterchangeDatasmithMaterialPipeline : public UI
 	UInterchangeDatasmithMaterialPipeline();
 
 protected:
-	virtual void ExecutePreImportPipeline(UInterchangeBaseNodeContainer* BaseNodeContainer, const TArray<UInterchangeSourceData*>& SourceDatas) override;
+	virtual void ExecutePipeline(UInterchangeBaseNodeContainer* BaseNodeContainer, const TArray<UInterchangeSourceData*>& SourceDatas) override;
 	virtual void ExecutePostImportPipeline(const UInterchangeBaseNodeContainer* BaseNodeContainer, const FString& NodeKey, UObject* CreatedAsset, bool bIsAReimport) override;
 
 	virtual bool CanExecuteOnAnyThread(EInterchangePipelineTask PipelineTask) override
diff --git a/Engine/Plugins/Experimental/Enterprise/DatasmithInterchange/Source/DatasmithInterchange/Private/InterchangeDatasmithPipeline.cpp b/Engine/Plugins/Experimental/Enterprise/DatasmithInterchange/Source/DatasmithInterchange/Private/InterchangeDatasmithPipeline.cpp
--- a/Engine/Plugins/Experimental/Enterprise/DatasmithInterchange/Source/DatasmithInterchange/Private/InterchangeDatasmithPipeline.cpp
+++ b/Engine/Plugins/Experimental/Enterprise/DatasmithInterchange/Source/DatasmithInterchange/Private/InterchangeDatasmithPipeline.cpp
@@ -55,7 +55,7 @@ UInterchangeDatasmithPipeline::UInterchangeDatasmithPipeline()
 	MeshPipeline->CommonSkeletalMeshesAndAnimationsProperties = CommonSkeletalMeshesAndAnimationsProperties;
 }
 
-void UInterchangeDatasmithPipeline::ExecutePreImportPipeline(UInterchangeBaseNodeContainer* InBaseNodeContainer, const TArray<UInterchangeSourceData*>& SourceDatas)
+void UInterchangeDatasmithPipeline::ExecutePipeline(UInterchangeBaseNodeContainer* InBaseNodeContainer, const TArray<UInterchangeSourceData*>& SourceDatas)
 {
 	using namespace UE::DatasmithInterchange;
 
@@ -72,7 +72,7 @@ void UInterchangeDatasmithPipeline::ExecutePreImportPipeline(UInterchangeBaseNod
 		if (Pipeline)
 		{
 			Pipeline->SetResultsContainer(this->Results);
-			Pipeline->ScriptedExecutePreImportPipeline(this->BaseNodeContainer, SourceDatas);
+			Pipeline->ScriptedExecutePipeline(this->BaseNodeContainer, SourceDatas);
 		}
 	};
 
diff --git a/Engine/Plugins/Experimental/Enterprise/DatasmithInterchange/Source/DatasmithInterchange/Private/InterchangeDatasmithStaticMeshPipeline.cpp b/Engine/Plugins/Experimental/Enterprise/DatasmithInterchange/Source/DatasmithInterchange/Private/InterchangeDatasmithStaticMeshPipeline.cpp
--- a/Engine/Plugins/Experimental/Enterprise/DatasmithInterchange/Source/DatasmithInterchange/Private/InterchangeDatasmithStaticMeshPipeline.cpp
+++ b/Engine/Plugins/Experimental/Enterprise/DatasmithInterchange/Source/DatasmithInterchange/Private/InterchangeDatasmithStaticMeshPipeline.cpp
@@ -11,11 +11,11 @@
 #include "Engine/StaticMesh.h"
 #include "Materials/MaterialInterface.h"
 
-void UInterchangeDatasmithStaticMeshPipeline::ExecutePreImportPipeline(UInterchangeBaseNodeContainer* NodeContainer, const TArray<UInterchangeSourceData*>& InSourceDatas)
+void UInterchangeDatasmithStaticMeshPipeline::ExecutePipeline(UInterchangeBaseNodeContainer* NodeContainer, const TArray<UInterchangeSourceData*>& InSourceDatas)
 {
 	using namespace UE::DatasmithInterchange;
 
-	Super::ExecutePreImportPipeline(NodeContainer, InSourceDatas);
+	Super::ExecutePipeline(NodeContainer, InSourceDatas);
 
 	// Add material factory dependencies for meshes where all slots are filled with the same material
 	for (UInterchangeStaticMeshFactoryNode* MeshFactoryNode : NodeUtils::GetNodes<UInterchangeStaticMeshFactoryNode>(NodeContainer))
diff --git a/Engine/Plugins/Experimental/Enterprise/DatasmithInterchange/Source/DatasmithInterchange/Private/InterchangeDatasmithStaticMeshPipeline.h b/Engine/Plugins/Experimental/Enterprise/DatasmithInterchange/Source/DatasmithInterchange/Private/InterchangeDatasmithStaticMeshPipeline.h
--- a/Engine/Plugins/Experimental/Enterprise/DatasmithInterchange/Source/DatasmithInterchange/Private/InterchangeDatasmithStaticMeshPipeline.h
+++ b/Engine/Plugins/Experimental/Enterprise/DatasmithInterchange/Source/DatasmithInterchange/Private/InterchangeDatasmithStaticMeshPipeline.h
@@ -17,7 +17,7 @@ class DATASMITHINTERCHANGE_API UInterchangeDatasmithStaticMeshPipeline : public
 	GENERATED_BODY()
 
 protected:
-	virtual void ExecutePreImportPipeline(UInterchangeBaseNodeContainer* BaseNodeContainer, const TArray<UInterchangeSourceData*>& SourceDatas) override;
+	virtual void ExecutePipeline(UInterchangeBaseNodeContainer* BaseNodeContainer, const TArray<UInterchangeSourceData*>& SourceDatas) override;
 	virtual void ExecutePostImportPipeline(const UInterchangeBaseNodeContainer* BaseNodeContainer, const FString& NodeKey, UObject* CreatedAsset, bool bIsAReimport) override;
 
 	virtual bool CanExecuteOnAnyThread(EInterchangePipelineTask PipelineTask) override
diff --git a/Engine/Plugins/Experimental/Enterprise/DatasmithInterchange/Source/DatasmithInterchange/Private/InterchangeDatasmithTexturePipeline.cpp b/Engine/Plugins/Experimental/Enterprise/DatasmithInterchange/Source/DatasmithInterchange/Private/InterchangeDatasmithTexturePipeline.cpp
--- a/Engine/Plugins/Experimental/Enterprise/DatasmithInterchange/Source/DatasmithInterchange/Private/InterchangeDatasmithTexturePipeline.cpp
+++ b/Engine/Plugins/Experimental/Enterprise/DatasmithInterchange/Source/DatasmithInterchange/Private/InterchangeDatasmithTexturePipeline.cpp
@@ -11,11 +11,11 @@
 #include "InterchangeTextureNode.h"
 #include "IDatasmithSceneElements.h"
 
-void UInterchangeDatasmithTexturePipeline::ExecutePreImportPipeline(UInterchangeBaseNodeContainer* InNodeContainer, const TArray<UInterchangeSourceData*>& InSourceDatas)
+void UInterchangeDatasmithTexturePipeline::ExecutePipeline(UInterchangeBaseNodeContainer* InNodeContainer, const TArray<UInterchangeSourceData*>& InSourceDatas)
 {
 	using namespace UE::DatasmithInterchange;
 
-	Super::ExecutePreImportPipeline(InNodeContainer, InSourceDatas);
+	Super::ExecutePipeline(InNodeContainer, InSourceDatas);
 
 	for (UInterchangeTextureFactoryNode* TextureFactoryNode : NodeUtils::GetNodes<UInterchangeTextureFactoryNode>(InNodeContainer))
 	{
diff --git a/Engine/Plugins/Experimental/Enterprise/DatasmithInterchange/Source/DatasmithInterchange/Private/InterchangeDatasmithTexturePipeline.h b/Engine/Plugins/Experimental/Enterprise/DatasmithInterchange/Source/DatasmithInterchange/Private/InterchangeDatasmithTexturePipeline.h
--- a/Engine/Plugins/Experimental/Enterprise/DatasmithInterchange/Source/DatasmithInterchange/Private/InterchangeDatasmithTexturePipeline.h
+++ b/Engine/Plugins/Experimental/Enterprise/DatasmithInterchange/Source/DatasmithInterchange/Private/InterchangeDatasmithTexturePipeline.h
@@ -20,7 +20,7 @@ class DATASMITHINTERCHANGE_API UInterchangeDatasmithTexturePipeline : public UIn
 	GENERATED_BODY()
 
 public:
-	virtual void ExecutePreImportPipeline(UInterchangeBaseNodeContainer* BaseNodeContainer, const TArray<UInterchangeSourceData*>& SourceDatas) override;
+	virtual void ExecutePipeline(UInterchangeBaseNodeContainer* BaseNodeContainer, const TArray<UInterchangeSourceData*>& SourceDatas) override;
 
 private:
 	void PreImportTextureFactoryNode(UInterchangeBaseNodeContainer* BaseNodeContainer, UInterchangeTextureFactoryNode* TextureFactoryNode) const;
diff --git a/Engine/Plugins/Experimental/Enterprise/DatasmithInterchange/Source/DatasmithInterchange/Public/InterchangeDatasmithPipeline.h b/Engine/Plugins/Experimental/Enterprise/DatasmithInterchange/Source/DatasmithInterchange/Public/InterchangeDatasmithPipeline.h
--- a/Engine/Plugins/Experimental/Enterprise/DatasmithInterchange/Source/DatasmithInterchange/Public/InterchangeDatasmithPipeline.h
+++ b/Engine/Plugins/Experimental/Enterprise/DatasmithInterchange/Source/DatasmithInterchange/Public/InterchangeDatasmithPipeline.h
@@ -51,7 +51,7 @@ class DATASMITHINTERCHANGE_API UInterchangeDatasmithPipeline : public UInterchan
 	TObjectPtr<UInterchangeGenericCommonSkeletalMeshesAndAnimationsProperties> CommonSkeletalMeshesAndAnimationsProperties;
 
 protected:
-	virtual void ExecutePreImportPipeline(UInterchangeBaseNodeContainer* InBaseNodeContainer, const TArray<UInterchangeSourceData*>& SourceDatas) override;
+	virtual void ExecutePipeline(UInterchangeBaseNodeContainer* InBaseNodeContainer, const TArray<UInterchangeSourceData*>& SourceDatas) override;
 	virtual void ExecutePostImportPipeline(const UInterchangeBaseNodeContainer* InBaseNodeContainer, const FString& NodeKey, UObject* CreatedAsset, bool bIsAReimport) override;
 
 	virtual bool CanExecuteOnAnyThread(EInterchangePipelineTask PipelineTask) override
diff --git a/Engine/Plugins/Interchange/Editor/Source/Pipelines/Private/InterchangeGraphInspectorPipeline.cpp b/Engine/Plugins/Interchange/Editor/Source/Pipelines/Private/InterchangeGraphInspectorPipeline.cpp
--- a/Engine/Plugins/Interchange/Editor/Source/Pipelines/Private/InterchangeGraphInspectorPipeline.cpp
+++ b/Engine/Plugins/Interchange/Editor/Source/Pipelines/Private/InterchangeGraphInspectorPipeline.cpp
@@ -17,7 +17,7 @@
 
 #include UE_INLINE_GENERATED_CPP_BY_NAME(InterchangeGraphInspectorPipeline)
 
-void UInterchangeGraphInspectorPipeline::ExecutePreImportPipeline(UInterchangeBaseNodeContainer* BaseNodeContainer, const TArray<UInterchangeSourceData*>& SourceDatas)
+void UInterchangeGraphInspectorPipeline::ExecutePipeline(UInterchangeBaseNodeContainer* BaseNodeContainer, const TArray<UInterchangeSourceData*>& SourceDatas)
 {
 	//Do not pop dialog if we are unattended or doing automation
 	if (FApp::IsUnattended() || IsRunningCommandlet())
diff --git a/Engine/Plugins/Interchange/Editor/Source/Pipelines/Public/InterchangeGraphInspectorPipeline.h b/Engine/Plugins/Interchange/Editor/Source/Pipelines/Public/InterchangeGraphInspectorPipeline.h
--- a/Engine/Plugins/Interchange/Editor/Source/Pipelines/Public/InterchangeGraphInspectorPipeline.h
+++ b/Engine/Plugins/Interchange/Editor/Source/Pipelines/Public/InterchangeGraphInspectorPipeline.h
@@ -35,7 +35,7 @@ public:
 
 protected:
 
-	virtual void ExecutePreImportPipeline(UInterchangeBaseNodeContainer* BaseNodeContainer, const TArray<UInterchangeSourceData*>& SourceDatas) override;
+	virtual void ExecutePipeline(UInterchangeBaseNodeContainer* BaseNodeContainer, const TArray<UInterchangeSourceData*>& SourceDatas) override;
 
 	//virtual bool ExecutePostImportPipeline(const UInterchangeBaseNodeContainer* BaseNodeContainer, const FName& NodeKey, UObject* CreatedAsset) override;
 
diff --git a/Engine/Plugins/Interchange/Runtime/Source/Import/Private/Animation/InterchangeAnimSequenceFactory.cpp b/Engine/Plugins/Interchange/Runtime/Source/Import/Private/Animation/InterchangeAnimSequenceFactory.cpp
--- a/Engine/Plugins/Interchange/Runtime/Source/Import/Private/Animation/InterchangeAnimSequenceFactory.cpp
+++ b/Engine/Plugins/Interchange/Runtime/Source/Import/Private/Animation/InterchangeAnimSequenceFactory.cpp
@@ -986,10 +986,10 @@ UObject* UInterchangeAnimSequenceFactory::CreateAsset(const FCreateAssetParams&
 }
 
 /* This function is call in the completion task on the main thread, use it to call main thread post creation step for your assets*/
-void UInterchangeAnimSequenceFactory::PreImportPreCompletedCallback(const FImportPreCompletedCallbackParams& Arguments)
+void UInterchangeAnimSequenceFactory::BeginPreCompletedCallback(const FImportPreCompletedCallbackParams& Arguments)
 {
 	check(IsInGameThread());
-	Super::PreImportPreCompletedCallback(Arguments);
+	Super::BeginPreCompletedCallback(Arguments);
 
 	// TODO: make sure this works at runtime
 #if WITH_EDITORONLY_DATA
diff --git a/Engine/Plugins/Interchange/Runtime/Source/Import/Private/Animation/InterchangeAnimationTrackSetFactory.cpp b/Engine/Plugins/Interchange/Runtime/Source/Import/Private/Animation/InterchangeAnimationTrackSetFactory.cpp
--- a/Engine/Plugins/Interchange/Runtime/Source/Import/Private/Animation/InterchangeAnimationTrackSetFactory.cpp
+++ b/Engine/Plugins/Interchange/Runtime/Source/Import/Private/Animation/InterchangeAnimationTrackSetFactory.cpp
@@ -706,15 +706,15 @@ UObject* UInterchangeAnimationTrackSetFactory::CreateAsset(const FCreateAssetPar
 }
 
 /* This function is call in the completion task on the main thread, use it to call main thread post creation step for your assets*/
-void UInterchangeAnimationTrackSetFactory::PreImportPreCompletedCallback(const FImportPreCompletedCallbackParams& Arguments)
+void UInterchangeAnimationTrackSetFactory::BeginPreCompletedCallback(const FImportPreCompletedCallbackParams& Arguments)
 {
 	check(IsInGameThread());
-	Super::PreImportPreCompletedCallback(Arguments);
+	Super::BeginPreCompletedCallback(Arguments);
 
 	// TODO: Talk with sequence team about adding AssetImportData to ULevelSequence for re-import 
 }
 
-void UInterchangeAnimationTrackSetFactory::PostImportPreCompletedCallback(const FImportPreCompletedCallbackParams& Arguments)
+void UInterchangeAnimationTrackSetFactory::EndPreCompletedCallback(const FImportPreCompletedCallbackParams& Arguments)
 {
 	check(IsInGameThread());
 	if (ULevelSequence* LevelSequence = Cast<ULevelSequence>(Arguments.ImportedObject))
@@ -726,7 +726,7 @@ void UInterchangeAnimationTrackSetFactory::PostImportPreCompletedCallback(const
 		// Note: Sequencer team is aware of the issue and will look their design in a future release
 		UMovieSceneSignedObject::ResetImplicitScopedModifyDefer();
 	}
-	Super::PostImportPreCompletedCallback(Arguments);
+	Super::EndPreCompletedCallback(Arguments);
 }
 
 #undef LOCTEXT_NAMESPACE
\ No newline at end of file
diff --git a/Engine/Plugins/Interchange/Runtime/Source/Import/Private/Material/InterchangeMaterialFactory.cpp b/Engine/Plugins/Interchange/Runtime/Source/Import/Private/Material/InterchangeMaterialFactory.cpp
--- a/Engine/Plugins/Interchange/Runtime/Source/Import/Private/Material/InterchangeMaterialFactory.cpp
+++ b/Engine/Plugins/Interchange/Runtime/Source/Import/Private/Material/InterchangeMaterialFactory.cpp
@@ -484,10 +484,10 @@ UObject* UInterchangeMaterialFactory::CreateAsset(const FCreateAssetParams& Argu
 }
 
 /* This function is call in the completion task on the main thread, use it to call main thread post creation step for your assets*/
-void UInterchangeMaterialFactory::PreImportPreCompletedCallback(const FImportPreCompletedCallbackParams& Arguments)
+void UInterchangeMaterialFactory::BeginPreCompletedCallback(const FImportPreCompletedCallbackParams& Arguments)
 {
 	check(IsInGameThread());
-	Super::PreImportPreCompletedCallback(Arguments);
+	Super::BeginPreCompletedCallback(Arguments);
 
 	if (ensure(Arguments.ImportedObject && Arguments.SourceData))
 	{
@@ -1241,9 +1241,9 @@ UObject* UInterchangeMaterialFunctionFactory::CreateAsset(const FCreateAssetPara
 }
 
 /* This function is call in the completion task on the main thread, use it to call main thread post creation step for your assets*/
-void UInterchangeMaterialFunctionFactory::PreImportPreCompletedCallback(const FImportPreCompletedCallbackParams& Arguments)
+void UInterchangeMaterialFunctionFactory::BeginPreCompletedCallback(const FImportPreCompletedCallbackParams& Arguments)
 {
-	Super::PreImportPreCompletedCallback(Arguments);
+	Super::BeginPreCompletedCallback(Arguments);
 
 #if WITH_EDITORONLY_DATA
 	if (ensure(Arguments.ImportedObject && Arguments.SourceData))
diff --git a/Engine/Plugins/Interchange/Runtime/Source/Import/Private/Mesh/InterchangeSkeletalMeshFactory.cpp b/Engine/Plugins/Interchange/Runtime/Source/Import/Private/Mesh/InterchangeSkeletalMeshFactory.cpp
--- a/Engine/Plugins/Interchange/Runtime/Source/Import/Private/Mesh/InterchangeSkeletalMeshFactory.cpp
+++ b/Engine/Plugins/Interchange/Runtime/Source/Import/Private/Mesh/InterchangeSkeletalMeshFactory.cpp
@@ -1599,11 +1599,11 @@ void UInterchangeSkeletalMeshFactory::Cancel()
 }
 
 /* This function is call in the completion task on the main thread, use it to call main thread post creation step for your assets*/
-void UInterchangeSkeletalMeshFactory::PreImportPreCompletedCallback(const FImportPreCompletedCallbackParams& Arguments)
+void UInterchangeSkeletalMeshFactory::BeginPreCompletedCallback(const FImportPreCompletedCallbackParams& Arguments)
 {
 	TRACE_CPUPROFILER_EVENT_SCOPE("UInterchangeSkeletalMeshFactory::PreImportPreCompletedCallback")
 	check(IsInGameThread());
-	Super::PreImportPreCompletedCallback(Arguments);
+	Super::BeginPreCompletedCallback(Arguments);
 
 	//TODO make sure this work at runtime
 #if WITH_EDITORONLY_DATA
diff --git a/Engine/Plugins/Interchange/Runtime/Source/Import/Private/Mesh/InterchangeStaticMeshFactory.cpp b/Engine/Plugins/Interchange/Runtime/Source/Import/Private/Mesh/InterchangeStaticMeshFactory.cpp
--- a/Engine/Plugins/Interchange/Runtime/Source/Import/Private/Mesh/InterchangeStaticMeshFactory.cpp
+++ b/Engine/Plugins/Interchange/Runtime/Source/Import/Private/Mesh/InterchangeStaticMeshFactory.cpp
@@ -511,10 +511,10 @@ void UInterchangeStaticMeshFactory::SetupSourceModelsSettings(UStaticMesh& Stati
 #endif // WITH_EDITORONLY_DATA
 
 /* This function is call in the completion task on the main thread, use it to call main thread post creation step for your assets */
-void UInterchangeStaticMeshFactory::PreImportPreCompletedCallback(const FImportPreCompletedCallbackParams& Arguments)
+void UInterchangeStaticMeshFactory::BeginPreCompletedCallback(const FImportPreCompletedCallbackParams& Arguments)
 {
 	check(IsInGameThread());
-	Super::PreImportPreCompletedCallback(Arguments);
+	Super::BeginPreCompletedCallback(Arguments);
 
 	// TODO: make sure this works at runtime
 #if WITH_EDITORONLY_DATA
diff --git a/Engine/Plugins/Interchange/Runtime/Source/Import/Private/Scene/InterchangeSceneVariantSetsFactory.cpp b/Engine/Plugins/Interchange/Runtime/Source/Import/Private/Scene/InterchangeSceneVariantSetsFactory.cpp
--- a/Engine/Plugins/Interchange/Runtime/Source/Import/Private/Scene/InterchangeSceneVariantSetsFactory.cpp
+++ b/Engine/Plugins/Interchange/Runtime/Source/Import/Private/Scene/InterchangeSceneVariantSetsFactory.cpp
@@ -441,10 +441,10 @@ UObject* UInterchangeSceneVariantSetsFactory::CreateAsset(const FCreateAssetPara
 #endif //else !WITH_EDITOR || !WITH_EDITORONLY_DATA
 }
 
-void UInterchangeSceneVariantSetsFactory::PreImportPreCompletedCallback(const FImportPreCompletedCallbackParams& Arguments)
+void UInterchangeSceneVariantSetsFactory::BeginPreCompletedCallback(const FImportPreCompletedCallbackParams& Arguments)
 {
 	check(IsInGameThread());
-	Super::PreImportPreCompletedCallback(Arguments);
+	Super::BeginPreCompletedCallback(Arguments);
 
 	// TODO: Need to add an AssetImportData property to ULevelVariantSets
 }
diff --git a/Engine/Plugins/Interchange/Runtime/Source/Import/Private/Scene/InterchangeSkeletalMeshActorFactory.cpp b/Engine/Plugins/Interchange/Runtime/Source/Import/Private/Scene/InterchangeSkeletalMeshActorFactory.cpp
--- a/Engine/Plugins/Interchange/Runtime/Source/Import/Private/Scene/InterchangeSkeletalMeshActorFactory.cpp
+++ b/Engine/Plugins/Interchange/Runtime/Source/Import/Private/Scene/InterchangeSkeletalMeshActorFactory.cpp
@@ -46,7 +46,7 @@ void UInterchangeSkeletalMeshActorFactory::SetupSkeletalMeshActor(const UInterch
 	SkeletalMeshComponent->UnregisterComponent();
 }
 
-void UInterchangeSkeletalMeshActorFactory::PostImportPreCompletedCallback(const FImportPreCompletedCallbackParams& Arguments)
+void UInterchangeSkeletalMeshActorFactory::EndPreCompletedCallback(const FImportPreCompletedCallbackParams& Arguments)
 {
 	// Set the skeletal mesh on the component in the post import callback, once the skeletal mesh has been fully imported.
 
diff --git a/Engine/Plugins/Interchange/Runtime/Source/Import/Private/Texture/InterchangeTextureFactory.cpp b/Engine/Plugins/Interchange/Runtime/Source/Import/Private/Texture/InterchangeTextureFactory.cpp
--- a/Engine/Plugins/Interchange/Runtime/Source/Import/Private/Texture/InterchangeTextureFactory.cpp
+++ b/Engine/Plugins/Interchange/Runtime/Source/Import/Private/Texture/InterchangeTextureFactory.cpp
@@ -1654,9 +1654,9 @@ UObject* UInterchangeTextureFactory::CreateAsset(const FCreateAssetParams& Argum
 }
 
 /* This function is call in the completion task on the main thread, use it to call main thread post creation step for your assets*/
-void UInterchangeTextureFactory::PreImportPreCompletedCallback(const FImportPreCompletedCallbackParams& Arguments)
+void UInterchangeTextureFactory::BeginPreCompletedCallback(const FImportPreCompletedCallbackParams& Arguments)
 {
-	TRACE_CPUPROFILER_EVENT_SCOPE(UInterchangeTextureFactory::PreImportPreCompletedCallback);
+	TRACE_CPUPROFILER_EVENT_SCOPE(UInterchangeTextureFactory::BeginPreCompletedCallback);
 
 	check(IsInGameThread());
 
@@ -1718,7 +1718,7 @@ void UInterchangeTextureFactory::PreImportPreCompletedCallback(const FImportPreC
 		}
 	}
 
-	Super::PreImportPreCompletedCallback(Arguments);
+	Super::BeginPreCompletedCallback(Arguments);
 
 	//TODO make sure this work at runtime
 #if WITH_EDITORONLY_DATA
diff --git a/Engine/Plugins/Interchange/Runtime/Source/Import/Public/Animation/InterchangeAnimSequenceFactory.h b/Engine/Plugins/Interchange/Runtime/Source/Import/Public/Animation/InterchangeAnimSequenceFactory.h
--- a/Engine/Plugins/Interchange/Runtime/Source/Import/Public/Animation/InterchangeAnimSequenceFactory.h
+++ b/Engine/Plugins/Interchange/Runtime/Source/Import/Public/Animation/InterchangeAnimSequenceFactory.h
@@ -25,8 +25,8 @@ public:
 	virtual EInterchangeFactoryAssetType GetFactoryAssetType() override { return EInterchangeFactoryAssetType::Animations; }
 	virtual UObject* CreateEmptyAsset(const FCreateAssetParams& Arguments) override;
 	virtual UObject* CreateAsset(const FCreateAssetParams& Arguments) override;
-	virtual void PreImportPreCompletedCallback(const FImportPreCompletedCallbackParams& Arguments) override;
-	//virtual void PostImportPreCompletedCallback(const FImportPreCompletedCallbackParams& Arguments) const override;
+	virtual void BeginPreCompletedCallback(const FImportPreCompletedCallbackParams& Arguments) override;
+	//virtual void EndPreCompletedCallback(const FImportPreCompletedCallbackParams& Arguments) const override;
 	virtual bool CanExecuteOnAnyThread() const override
 	{
 		//Currently we cannot use the anim sequence controller outside of the game thread
diff --git a/Engine/Plugins/Interchange/Runtime/Source/Import/Public/Animation/InterchangeAnimationTrackSetFactory.h b/Engine/Plugins/Interchange/Runtime/Source/Import/Public/Animation/InterchangeAnimationTrackSetFactory.h
--- a/Engine/Plugins/Interchange/Runtime/Source/Import/Public/Animation/InterchangeAnimationTrackSetFactory.h
+++ b/Engine/Plugins/Interchange/Runtime/Source/Import/Public/Animation/InterchangeAnimationTrackSetFactory.h
@@ -22,8 +22,8 @@ public:
 	virtual EInterchangeFactoryAssetType GetFactoryAssetType() override { return EInterchangeFactoryAssetType::Animations; }
 	virtual UObject* CreateEmptyAsset(const FCreateAssetParams& Arguments) override;
 	virtual UObject* CreateAsset(const FCreateAssetParams& Arguments) override;
-	virtual void PreImportPreCompletedCallback(const FImportPreCompletedCallbackParams& Arguments) override;
-	virtual void PostImportPreCompletedCallback(const FImportPreCompletedCallbackParams& Arguments) override;
+	virtual void BeginPreCompletedCallback(const FImportPreCompletedCallbackParams& Arguments) override;
+	virtual void EndPreCompletedCallback(const FImportPreCompletedCallbackParams& Arguments) override;
 	virtual bool CanExecuteOnAnyThread() const override
 	{
 		//Currently we cannot use the anim sequence controller outside of the game thread
diff --git a/Engine/Plugins/Interchange/Runtime/Source/Import/Public/Material/InterchangeMaterialFactory.h b/Engine/Plugins/Interchange/Runtime/Source/Import/Public/Material/InterchangeMaterialFactory.h
--- a/Engine/Plugins/Interchange/Runtime/Source/Import/Public/Material/InterchangeMaterialFactory.h
+++ b/Engine/Plugins/Interchange/Runtime/Source/Import/Public/Material/InterchangeMaterialFactory.h
@@ -29,7 +29,7 @@ public:
 	virtual EInterchangeFactoryAssetType GetFactoryAssetType() override { return EInterchangeFactoryAssetType::Materials; }
 	virtual UObject* CreateEmptyAsset(const FCreateAssetParams& Arguments) override;
 	virtual UObject* CreateAsset(const FCreateAssetParams& Arguments) override;
-	virtual void PreImportPreCompletedCallback(const FImportPreCompletedCallbackParams& Arguments) override;
+	virtual void BeginPreCompletedCallback(const FImportPreCompletedCallbackParams& Arguments) override;
 	virtual bool GetSourceFilenames(const UObject* Object, TArray<FString>& OutSourceFilenames) const override;
 	virtual bool SetSourceFilename(const UObject* Object, const FString& SourceFilename, int32 SourceIndex) const override;
 
@@ -56,7 +56,7 @@ public:
 	virtual EInterchangeFactoryAssetType GetFactoryAssetType() override { return EInterchangeFactoryAssetType::Materials; }
 	virtual UObject* CreateEmptyAsset(const FCreateAssetParams& Arguments) override;
 	virtual UObject* CreateAsset(const FCreateAssetParams& Arguments) override;
-	virtual void PreImportPreCompletedCallback(const FImportPreCompletedCallbackParams& Arguments) override;
+	virtual void BeginPreCompletedCallback(const FImportPreCompletedCallbackParams& Arguments) override;
 
 	// Interchange factory base interface end
 	//////////////////////////////////////////////////////////////////////////
diff --git a/Engine/Plugins/Interchange/Runtime/Source/Import/Public/Mesh/InterchangeSkeletalMeshFactory.h b/Engine/Plugins/Interchange/Runtime/Source/Import/Public/Mesh/InterchangeSkeletalMeshFactory.h
--- a/Engine/Plugins/Interchange/Runtime/Source/Import/Public/Mesh/InterchangeSkeletalMeshFactory.h
+++ b/Engine/Plugins/Interchange/Runtime/Source/Import/Public/Mesh/InterchangeSkeletalMeshFactory.h
@@ -25,8 +25,8 @@ public:
 	virtual UObject* CreateEmptyAsset(const FCreateAssetParams& Arguments) override;
 	virtual UObject* CreateAsset(const FCreateAssetParams& Arguments) override;
 	virtual void Cancel() override;
-	virtual void PreImportPreCompletedCallback(const FImportPreCompletedCallbackParams& Arguments) override;
-	//virtual void PostImportPreCompletedCallback(const FImportPreCompletedCallbackParams& Arguments) const override;
+	virtual void BeginPreCompletedCallback(const FImportPreCompletedCallbackParams& Arguments) override;
+	//virtual void EndPreCompletedCallback(const FImportPreCompletedCallbackParams& Arguments) const override;
 	virtual bool GetSourceFilenames(const UObject* Object, TArray<FString>& OutSourceFilenames) const override;
 	virtual bool SetSourceFilename(const UObject* Object, const FString& SourceFilename, int32 SourceIndex) const override;
 	virtual bool SetReimportSourceIndex(const UObject* Object, int32 SourceIndex) const override;
diff --git a/Engine/Plugins/Interchange/Runtime/Source/Import/Public/Mesh/InterchangeStaticMeshFactory.h b/Engine/Plugins/Interchange/Runtime/Source/Import/Public/Mesh/InterchangeStaticMeshFactory.h
--- a/Engine/Plugins/Interchange/Runtime/Source/Import/Public/Mesh/InterchangeStaticMeshFactory.h
+++ b/Engine/Plugins/Interchange/Runtime/Source/Import/Public/Mesh/InterchangeStaticMeshFactory.h
@@ -32,7 +32,7 @@ public:
 	virtual EInterchangeFactoryAssetType GetFactoryAssetType() override { return EInterchangeFactoryAssetType::Meshes; }
 	virtual UObject* CreateEmptyAsset(const FCreateAssetParams& Arguments) override;
 	virtual UObject* CreateAsset(const FCreateAssetParams& Arguments) override;
-	virtual void PreImportPreCompletedCallback(const FImportPreCompletedCallbackParams& Arguments) override;
+	virtual void BeginPreCompletedCallback(const FImportPreCompletedCallbackParams& Arguments) override;
 	virtual bool GetSourceFilenames(const UObject* Object, TArray<FString>& OutSourceFilenames) const override;
 	virtual bool SetSourceFilename(const UObject* Object, const FString& SourceFilename, int32 SourceIndex) const override;
 
diff --git a/Engine/Plugins/Interchange/Runtime/Source/Import/Public/Scene/InterchangeSceneVariantSetsFactory.h b/Engine/Plugins/Interchange/Runtime/Source/Import/Public/Scene/InterchangeSceneVariantSetsFactory.h
--- a/Engine/Plugins/Interchange/Runtime/Source/Import/Public/Scene/InterchangeSceneVariantSetsFactory.h
+++ b/Engine/Plugins/Interchange/Runtime/Source/Import/Public/Scene/InterchangeSceneVariantSetsFactory.h
@@ -25,7 +25,7 @@ public:
 	virtual UClass* GetFactoryClass() const override;
 	virtual UObject* CreateEmptyAsset(const FCreateAssetParams& Arguments) override;
 	virtual UObject* CreateAsset(const FCreateAssetParams& Arguments) override;
-	virtual void PreImportPreCompletedCallback(const FImportPreCompletedCallbackParams& Arguments) override;
+	virtual void BeginPreCompletedCallback(const FImportPreCompletedCallbackParams& Arguments) override;
 	virtual bool CanExecuteOnAnyThread() const override
 	{
 		//Currently we cannot use the anim sequence controller outside of the game thread
diff --git a/Engine/Plugins/Interchange/Runtime/Source/Import/Public/Scene/InterchangeSkeletalMeshActorFactory.h b/Engine/Plugins/Interchange/Runtime/Source/Import/Public/Scene/InterchangeSkeletalMeshActorFactory.h
--- a/Engine/Plugins/Interchange/Runtime/Source/Import/Public/Scene/InterchangeSkeletalMeshActorFactory.h
+++ b/Engine/Plugins/Interchange/Runtime/Source/Import/Public/Scene/InterchangeSkeletalMeshActorFactory.h
@@ -31,7 +31,7 @@ public:
 		return false;
 	}
 
-	virtual void PostImportPreCompletedCallback(const FImportPreCompletedCallbackParams& Arguments) override;
+	virtual void EndPreCompletedCallback(const FImportPreCompletedCallbackParams& Arguments) override;
 
 	// Interchange factory base interface end
 	//////////////////////////////////////////////////////////////////////////
diff --git a/Engine/Plugins/Interchange/Runtime/Source/Import/Public/Texture/InterchangeTextureFactory.h b/Engine/Plugins/Interchange/Runtime/Source/Import/Public/Texture/InterchangeTextureFactory.h
--- a/Engine/Plugins/Interchange/Runtime/Source/Import/Public/Texture/InterchangeTextureFactory.h
+++ b/Engine/Plugins/Interchange/Runtime/Source/Import/Public/Texture/InterchangeTextureFactory.h
@@ -58,7 +58,7 @@ public:
 	virtual EInterchangeFactoryAssetType GetFactoryAssetType() override { return EInterchangeFactoryAssetType::Textures; }
 	virtual UObject* CreateEmptyAsset(const FCreateAssetParams& Arguments) override;
 	virtual UObject* CreateAsset(const FCreateAssetParams& Arguments) override;
-	virtual void PreImportPreCompletedCallback(const FImportPreCompletedCallbackParams& Arguments) override;
+	virtual void BeginPreCompletedCallback(const FImportPreCompletedCallbackParams& Arguments) override;
 	virtual bool GetSourceFilenames(const UObject* Object, TArray<FString>& OutSourceFilenames) const override;
 	virtual bool SetSourceFilename(const UObject* Object, const FString& SourceFilename, int32 SourceIndex) const override;
 
diff --git a/Engine/Plugins/Interchange/Runtime/Source/Pipelines/Private/InterchangeGenericAnimationPipeline.cpp b/Engine/Plugins/Interchange/Runtime/Source/Pipelines/Private/InterchangeGenericAnimationPipeline.cpp
--- a/Engine/Plugins/Interchange/Runtime/Source/Pipelines/Private/InterchangeGenericAnimationPipeline.cpp
+++ b/Engine/Plugins/Interchange/Runtime/Source/Pipelines/Private/InterchangeGenericAnimationPipeline.cpp
@@ -88,7 +88,7 @@ void UInterchangeGenericAnimationPipeline::AdjustSettingsForContext(EInterchange
 	}
 }
 
-void UInterchangeGenericAnimationPipeline::ExecutePreImportPipeline(UInterchangeBaseNodeContainer* InBaseNodeContainer, const TArray<UInterchangeSourceData*>& InSourceDatas)
+void UInterchangeGenericAnimationPipeline::ExecutePipeline(UInterchangeBaseNodeContainer* InBaseNodeContainer, const TArray<UInterchangeSourceData*>& InSourceDatas)
 {
 	if (!InBaseNodeContainer)
 	{
diff --git a/Engine/Plugins/Interchange/Runtime/Source/Pipelines/Private/InterchangeGenericAssetsPipeline.cpp b/Engine/Plugins/Interchange/Runtime/Source/Pipelines/Private/InterchangeGenericAssetsPipeline.cpp
--- a/Engine/Plugins/Interchange/Runtime/Source/Pipelines/Private/InterchangeGenericAssetsPipeline.cpp
+++ b/Engine/Plugins/Interchange/Runtime/Source/Pipelines/Private/InterchangeGenericAssetsPipeline.cpp
@@ -118,7 +118,7 @@ void UInterchangeGenericAssetsPipeline::AdjustSettingsForContext(EInterchangePip
 	}
 }
 
-void UInterchangeGenericAssetsPipeline::ExecutePreImportPipeline(UInterchangeBaseNodeContainer* InBaseNodeContainer, const TArray<UInterchangeSourceData*>& InSourceDatas)
+void UInterchangeGenericAssetsPipeline::ExecutePipeline(UInterchangeBaseNodeContainer* InBaseNodeContainer, const TArray<UInterchangeSourceData*>& InSourceDatas)
 {
 	check(CommonSkeletalMeshesAndAnimationsProperties);
 
@@ -200,15 +200,15 @@ void UInterchangeGenericAssetsPipeline::ExecutePreImportPipeline(UInterchangeBas
 
 	if (MaterialPipeline)
 	{
-		MaterialPipeline->ScriptedExecutePreImportPipeline(InBaseNodeContainer, InSourceDatas);
+		MaterialPipeline->ScriptedExecutePipeline(InBaseNodeContainer, InSourceDatas);
 	}
 	if (MeshPipeline)
 	{
-		MeshPipeline->ScriptedExecutePreImportPipeline(InBaseNodeContainer, InSourceDatas);
+		MeshPipeline->ScriptedExecutePipeline(InBaseNodeContainer, InSourceDatas);
 	}
 	if (AnimationPipeline)
 	{
-		AnimationPipeline->ScriptedExecutePreImportPipeline(InBaseNodeContainer, InSourceDatas);
+		AnimationPipeline->ScriptedExecutePipeline(InBaseNodeContainer, InSourceDatas);
 	}
 
 	ImplementUseSourceNameForAssetOption();
diff --git a/Engine/Plugins/Interchange/Runtime/Source/Pipelines/Private/InterchangeGenericMaterialPipeline.cpp b/Engine/Plugins/Interchange/Runtime/Source/Pipelines/Private/InterchangeGenericMaterialPipeline.cpp
--- a/Engine/Plugins/Interchange/Runtime/Source/Pipelines/Private/InterchangeGenericMaterialPipeline.cpp
+++ b/Engine/Plugins/Interchange/Runtime/Source/Pipelines/Private/InterchangeGenericMaterialPipeline.cpp
@@ -239,7 +239,7 @@ void UInterchangeGenericMaterialPipeline::AdjustSettingsForContext(EInterchangeP
 	}
 }
 
-void UInterchangeGenericMaterialPipeline::ExecutePreImportPipeline(UInterchangeBaseNodeContainer* InBaseNodeContainer, const TArray<UInterchangeSourceData*>& InSourceDatas)
+void UInterchangeGenericMaterialPipeline::ExecutePipeline(UInterchangeBaseNodeContainer* InBaseNodeContainer, const TArray<UInterchangeSourceData*>& InSourceDatas)
 {
 	if (!InBaseNodeContainer)
 	{
@@ -266,7 +266,7 @@ void UInterchangeGenericMaterialPipeline::ExecutePreImportPipeline(UInterchangeB
 	
 	if (TexturePipeline)
 	{
-		TexturePipeline->ScriptedExecutePreImportPipeline(InBaseNodeContainer, InSourceDatas);
+		TexturePipeline->ScriptedExecutePipeline(InBaseNodeContainer, InSourceDatas);
 	}
 
 	//Skip Material import if the toggle is off
diff --git a/Engine/Plugins/Interchange/Runtime/Source/Pipelines/Private/InterchangeGenericMeshPipeline.cpp b/Engine/Plugins/Interchange/Runtime/Source/Pipelines/Private/InterchangeGenericMeshPipeline.cpp
--- a/Engine/Plugins/Interchange/Runtime/Source/Pipelines/Private/InterchangeGenericMeshPipeline.cpp
+++ b/Engine/Plugins/Interchange/Runtime/Source/Pipelines/Private/InterchangeGenericMeshPipeline.cpp
@@ -104,7 +104,7 @@ void UInterchangeGenericMeshPipeline::PreDialogCleanup(const FName PipelineStack
 	PhysicsAsset = nullptr;
 }
 
-void UInterchangeGenericMeshPipeline::ExecutePreImportPipeline(UInterchangeBaseNodeContainer* InBaseNodeContainer, const TArray<UInterchangeSourceData*>& InSourceDatas)
+void UInterchangeGenericMeshPipeline::ExecutePipeline(UInterchangeBaseNodeContainer* InBaseNodeContainer, const TArray<UInterchangeSourceData*>& InSourceDatas)
 {
 	if (!InBaseNodeContainer)
 	{
diff --git a/Engine/Plugins/Interchange/Runtime/Source/Pipelines/Private/InterchangeGenericScenesPipeline.cpp b/Engine/Plugins/Interchange/Runtime/Source/Pipelines/Private/InterchangeGenericScenesPipeline.cpp
--- a/Engine/Plugins/Interchange/Runtime/Source/Pipelines/Private/InterchangeGenericScenesPipeline.cpp
+++ b/Engine/Plugins/Interchange/Runtime/Source/Pipelines/Private/InterchangeGenericScenesPipeline.cpp
@@ -28,7 +28,7 @@
 
 #include UE_INLINE_GENERATED_CPP_BY_NAME(InterchangeGenericScenesPipeline)
 
-void UInterchangeGenericLevelPipeline::ExecutePreImportPipeline(UInterchangeBaseNodeContainer* InBaseNodeContainer, const TArray<UInterchangeSourceData*>& InSourceDatas)
+void UInterchangeGenericLevelPipeline::ExecutePipeline(UInterchangeBaseNodeContainer* InBaseNodeContainer, const TArray<UInterchangeSourceData*>& InSourceDatas)
 {
 	if (!InBaseNodeContainer)
 	{
diff --git a/Engine/Plugins/Interchange/Runtime/Source/Pipelines/Private/InterchangeGenericTexturePipeline.cpp b/Engine/Plugins/Interchange/Runtime/Source/Pipelines/Private/InterchangeGenericTexturePipeline.cpp
--- a/Engine/Plugins/Interchange/Runtime/Source/Pipelines/Private/InterchangeGenericTexturePipeline.cpp
+++ b/Engine/Plugins/Interchange/Runtime/Source/Pipelines/Private/InterchangeGenericTexturePipeline.cpp
@@ -131,7 +131,7 @@ void UInterchangeGenericTexturePipeline::AdjustSettingsForContext(EInterchangePi
 	}
 }
 
-void UInterchangeGenericTexturePipeline::ExecutePreImportPipeline(UInterchangeBaseNodeContainer* InBaseNodeContainer, const TArray<UInterchangeSourceData*>& InSourceDatas)
+void UInterchangeGenericTexturePipeline::ExecutePipeline(UInterchangeBaseNodeContainer* InBaseNodeContainer, const TArray<UInterchangeSourceData*>& InSourceDatas)
 {
 	if (!InBaseNodeContainer)
 	{
diff --git a/Engine/Plugins/Interchange/Runtime/Source/Pipelines/Public/InterchangeGenericAnimationPipeline.h b/Engine/Plugins/Interchange/Runtime/Source/Pipelines/Public/InterchangeGenericAnimationPipeline.h
--- a/Engine/Plugins/Interchange/Runtime/Source/Pipelines/Public/InterchangeGenericAnimationPipeline.h
+++ b/Engine/Plugins/Interchange/Runtime/Source/Pipelines/Public/InterchangeGenericAnimationPipeline.h
@@ -110,7 +110,7 @@ public:
 	
 protected:
 
-	virtual void ExecutePreImportPipeline(UInterchangeBaseNodeContainer* InBaseNodeContainer, const TArray<UInterchangeSourceData*>& InSourceDatas) override;
+	virtual void ExecutePipeline(UInterchangeBaseNodeContainer* InBaseNodeContainer, const TArray<UInterchangeSourceData*>& InSourceDatas) override;
 
 	virtual bool CanExecuteOnAnyThread(EInterchangePipelineTask PipelineTask) override
 	{
diff --git a/Engine/Plugins/Interchange/Runtime/Source/Pipelines/Public/InterchangeGenericAssetsPipeline.h b/Engine/Plugins/Interchange/Runtime/Source/Pipelines/Public/InterchangeGenericAssetsPipeline.h
--- a/Engine/Plugins/Interchange/Runtime/Source/Pipelines/Public/InterchangeGenericAssetsPipeline.h
+++ b/Engine/Plugins/Interchange/Runtime/Source/Pipelines/Public/InterchangeGenericAssetsPipeline.h
@@ -80,7 +80,7 @@ public:
 	virtual void AdjustSettingsForContext(EInterchangePipelineContext ImportType, TObjectPtr<UObject> ReimportAsset) override;
 protected:
 
-	virtual void ExecutePreImportPipeline(UInterchangeBaseNodeContainer* InBaseNodeContainer, const TArray<UInterchangeSourceData*>& InSourceDatas) override;
+	virtual void ExecutePipeline(UInterchangeBaseNodeContainer* InBaseNodeContainer, const TArray<UInterchangeSourceData*>& InSourceDatas) override;
 
 	virtual void ExecutePostImportPipeline(const UInterchangeBaseNodeContainer* BaseNodeContainer, const FString& NodeKey, UObject* CreatedAsset, bool bIsAReimport) override;
 
diff --git a/Engine/Plugins/Interchange/Runtime/Source/Pipelines/Public/InterchangeGenericMaterialPipeline.h b/Engine/Plugins/Interchange/Runtime/Source/Pipelines/Public/InterchangeGenericMaterialPipeline.h
--- a/Engine/Plugins/Interchange/Runtime/Source/Pipelines/Public/InterchangeGenericMaterialPipeline.h
+++ b/Engine/Plugins/Interchange/Runtime/Source/Pipelines/Public/InterchangeGenericMaterialPipeline.h
@@ -87,7 +87,7 @@ public:
 	virtual void AdjustSettingsForContext(EInterchangePipelineContext ImportType, TObjectPtr<UObject> ReimportAsset) override;
 
 protected:
-	virtual void ExecutePreImportPipeline(UInterchangeBaseNodeContainer* InBaseNodeContainer, const TArray<UInterchangeSourceData*>& InSourceDatas) override;
+	virtual void ExecutePipeline(UInterchangeBaseNodeContainer* InBaseNodeContainer, const TArray<UInterchangeSourceData*>& InSourceDatas) override;
 	virtual void ExecutePostImportPipeline(const UInterchangeBaseNodeContainer* BaseNodeContainer, const FString& NodeKey, UObject* CreatedAsset, bool bIsAReimport) override;
 	virtual void SetReimportSourceIndex(UClass* ReimportObjectClass, const int32 SourceFileIndex) override;
 	/** END UInterchangePipelineBase overrides */
diff --git a/Engine/Plugins/Interchange/Runtime/Source/Pipelines/Public/InterchangeGenericMeshPipeline.h b/Engine/Plugins/Interchange/Runtime/Source/Pipelines/Public/InterchangeGenericMeshPipeline.h
--- a/Engine/Plugins/Interchange/Runtime/Source/Pipelines/Public/InterchangeGenericMeshPipeline.h
+++ b/Engine/Plugins/Interchange/Runtime/Source/Pipelines/Public/InterchangeGenericMeshPipeline.h
@@ -188,7 +188,7 @@ public:
 	virtual void PreDialogCleanup(const FName PipelineStackName) override;
 
 protected:
-	virtual void ExecutePreImportPipeline(UInterchangeBaseNodeContainer* InBaseNodeContainer, const TArray<UInterchangeSourceData*>& InSourceDatas) override;
+	virtual void ExecutePipeline(UInterchangeBaseNodeContainer* InBaseNodeContainer, const TArray<UInterchangeSourceData*>& InSourceDatas) override;
 
 	virtual void ExecutePostImportPipeline(const UInterchangeBaseNodeContainer* InBaseNodeContainer, const FString& NodeKey, UObject* CreatedAsset, bool bIsAReimport) override;
 
diff --git a/Engine/Plugins/Interchange/Runtime/Source/Pipelines/Public/InterchangeGenericScenesPipeline.h b/Engine/Plugins/Interchange/Runtime/Source/Pipelines/Public/InterchangeGenericScenesPipeline.h
--- a/Engine/Plugins/Interchange/Runtime/Source/Pipelines/Public/InterchangeGenericScenesPipeline.h
+++ b/Engine/Plugins/Interchange/Runtime/Source/Pipelines/Public/InterchangeGenericScenesPipeline.h
@@ -24,7 +24,7 @@ class INTERCHANGEPIPELINES_API UInterchangeGenericLevelPipeline : public UInterc
 
 protected:
 
-	virtual void ExecutePreImportPipeline(UInterchangeBaseNodeContainer* InBaseNodeContainer, const TArray<UInterchangeSourceData*>& InSourceDatas) override;
+	virtual void ExecutePipeline(UInterchangeBaseNodeContainer* InBaseNodeContainer, const TArray<UInterchangeSourceData*>& InSourceDatas) override;
 
 	virtual bool CanExecuteOnAnyThread(EInterchangePipelineTask PipelineTask) override
 	{
diff --git a/Engine/Plugins/Interchange/Runtime/Source/Pipelines/Public/InterchangeGenericTexturePipeline.h b/Engine/Plugins/Interchange/Runtime/Source/Pipelines/Public/InterchangeGenericTexturePipeline.h
--- a/Engine/Plugins/Interchange/Runtime/Source/Pipelines/Public/InterchangeGenericTexturePipeline.h
+++ b/Engine/Plugins/Interchange/Runtime/Source/Pipelines/Public/InterchangeGenericTexturePipeline.h
@@ -61,7 +61,7 @@ public:
 
 public:
 	virtual void AdjustSettingsForContext(EInterchangePipelineContext ImportType, TObjectPtr<UObject> ReimportAsset) override;
-	virtual void ExecutePreImportPipeline(UInterchangeBaseNodeContainer* InBaseNodeContainer, const TArray<UInterchangeSourceData*>& InSourceDatas) override;
+	virtual void ExecutePipeline(UInterchangeBaseNodeContainer* InBaseNodeContainer, const TArray<UInterchangeSourceData*>& InSourceDatas) override;
 	virtual void ExecutePostImportPipeline(const UInterchangeBaseNodeContainer* InBaseNodeContainer, const FString& NodeKey, UObject* CreatedAsset, bool bIsAReimport) override;
 
 protected:
diff --git a/Engine/Source/Runtime/Interchange/Core/Public/InterchangeFactoryBase.h b/Engine/Source/Runtime/Interchange/Core/Public/InterchangeFactoryBase.h
--- a/Engine/Source/Runtime/Interchange/Core/Public/InterchangeFactoryBase.h
+++ b/Engine/Source/Runtime/Interchange/Core/Public/InterchangeFactoryBase.h
@@ -94,7 +94,7 @@ public:
 	}
 
 	/**
-	 * Create an asset from a Node data, this function must be multithread safe, it cannot use member all the data must be pass in the FCreateAssetParams structure
+	 * Create an asset from a Node data, this function must be multithread safe.
 	 *
 	 * @param Arguments - The structure containing all necessary arguments, see the structure definition for the documentation
 	 * @return the created UObject or nullptr if there is an error, see LOG for any error detail.
@@ -165,21 +165,33 @@ public:
 ;
 	};
 
+	UE_DEPRECATED(5.2, "This function is replace by BeginPreCompletedCallback.")
+	virtual void PreImportPreCompletedCallback(const FImportPreCompletedCallbackParams& Arguments)
+	{
+		BeginPreCompletedCallback(Arguments);
+	}
+
 	/*
 	 * This function is call in the pre completion task on the main thread, use it to call main thread post creation step for your assets
 	 * @note - This function is called when starting the pre completion task (before PostEditChange is called for the asset).
 	 */
-	virtual void PreImportPreCompletedCallback(const FImportPreCompletedCallbackParams& Arguments)
+	virtual void BeginPreCompletedCallback(const FImportPreCompletedCallbackParams& Arguments)
 	{
 		check(IsInGameThread());
 		return;
 	}
 	
+	UE_DEPRECATED(5.2, "This function is replace by EndPreCompletedCallback.")
+	virtual void PostImportPreCompletedCallback(const FImportPreCompletedCallbackParams& Arguments)
+	{
+		EndPreCompletedCallback(Arguments);
+	}
+
 	/*
 	 * This function is call in the pre completion task on the main thread, use it to call main thread post creation step for your assets
 	 * @note - This function is called at the end of the pre completion task (after PostEditChange is called for the asset).
 	 */
-	virtual void PostImportPreCompletedCallback(const FImportPreCompletedCallbackParams& Arguments)
+	virtual void EndPreCompletedCallback(const FImportPreCompletedCallbackParams& Arguments)
 	{
 		check(IsInGameThread());
 		return;
diff --git a/Engine/Source/Runtime/Interchange/Core/Public/InterchangePipelineBase.h b/Engine/Source/Runtime/Interchange/Core/Public/InterchangePipelineBase.h
--- a/Engine/Source/Runtime/Interchange/Core/Public/InterchangePipelineBase.h
+++ b/Engine/Source/Runtime/Interchange/Core/Public/InterchangePipelineBase.h
@@ -31,8 +31,9 @@ template <typename OptionalType> struct TOptional;
 UENUM(BlueprintType)
 enum class EInterchangePipelineTask : uint8
 {
-	PreFactoryImport,
-	PostFactoryImport,
+	PostTranslator,
+	PostFactory,
+	PostImport,
 	Export
 };
 
@@ -107,31 +108,55 @@ class INTERCHANGECORE_API UInterchangePipelineBase : public UObject
 
 public:
 
-	/**
-	 * Non virtual helper to allow blueprint to implement event base function to implement a pre import pipeline,
-	 * It is call after the translation and before we parse the graph to call the factory. This is where factory node should be created
-	 * by the pipeline. Each factory node should be send to a a interchange factory to create an unreal asset.
-	 * @note - the Interchange manager is calling this function not the virtual one that is call by the default implementation.
-	 */
+	UE_DEPRECATED(5.2, "This function is replace by ScriptedExecutePipeline.")
 	UFUNCTION(BlueprintNativeEvent, BlueprintCallable, Category = "Interchange | Pipeline")
 	void ScriptedExecutePreImportPipeline(UInterchangeBaseNodeContainer* BaseNodeContainer, const TArray<UInterchangeSourceData*>& SourceDatas);
-	/** The default implementation (call if the blueprint do not have any implementation) will call the virtual ExecuteImportPipeline */
+
+	UE_DEPRECATED(5.2, "This internal public function is deprecated and not replace by any.")
 	void ScriptedExecutePreImportPipeline_Implementation(UInterchangeBaseNodeContainer* BaseNodeContainer, const TArray<UInterchangeSourceData*>& SourceDatas)
 	{
+		PRAGMA_DISABLE_DEPRECATION_WARNINGS
 		//By default we call the virtual import pipeline execution
 		ExecutePreImportPipeline(BaseNodeContainer, SourceDatas);
+		PRAGMA_ENABLE_DEPRECATION_WARNINGS
 	}
 
 	/**
-	 * Non virtual helper to allow blueprint to implement event base function to implement a post import pipeline,
-	 * It is call after we completely import an asset. PostEditChange is already called. Some assets uses asynchronous build.
-	 * This can be useful if you need builded data of an asset to finish the setup of another asset.
+	 * ScriptedExecutePipeline, is call after the translation and before we parse the graph to call the factory.
+	 * This is where factory node should be created by the pipeline.
+	 * Each factory node represent an unreal asset create that will be create by an interchange factory.
+	 * @note - the FTaskPipeline is calling this function not the virtual one that is call by the default implementation.
+	 */
+	UFUNCTION(BlueprintNativeEvent, BlueprintCallable, Category = "Interchange | Pipeline")
+	void ScriptedExecutePipeline(UInterchangeBaseNodeContainer* BaseNodeContainer, const TArray<UInterchangeSourceData*>& SourceDatas);
+	/** The default implementation (call if the blueprint do not have any implementation) will call the virtual ExecutePipeline */
+	void ScriptedExecutePipeline_Implementation(UInterchangeBaseNodeContainer* BaseNodeContainer, const TArray<UInterchangeSourceData*>& SourceDatas)
+	{
+		//By default we call the virtual import pipeline execution
+		ExecutePipeline(BaseNodeContainer, SourceDatas);
+	}
+
+	/**
+	 * ScriptedExecutePostFactoryPipeline is call after the factory create the unreal asset but before calling PostEditChange.
+	 * @note - the FTaskPreCompletion task is calling this function not the virtual one that is call by the default implementation.
+	 */
+	UFUNCTION(BlueprintNativeEvent, BlueprintCallable, Category = "Interchange | Pipeline")
+	void ScriptedExecutePostFactoryPipeline(const UInterchangeBaseNodeContainer* BaseNodeContainer, const FString& FactoryNodeKey, UObject* CreatedAsset, bool bIsAReimport);
+	/** The default implementation (call if the blueprint do not have any implementation) will call the virtual ExecutePostFactoryPipeline */
+	void ScriptedExecutePostFactoryPipeline_Implementation(const UInterchangeBaseNodeContainer* BaseNodeContainer, const FString& FactoryNodeKey, UObject* CreatedAsset, bool bIsAReimport)
+	{
+		//By default we call the virtual import pipeline execution
+		ExecutePostFactoryPipeline(BaseNodeContainer, FactoryNodeKey, CreatedAsset, bIsAReimport);
+	}
+	/**
+	 * ScriptedExecutePostImportPipeline is call after we completely import an asset, PostEditChange is already called.
+	 * This can be useful if you need an asset build data to finish the setup of another asset.
 	 * @example - PhysicsAsset need skeletal mesh render data to be build properly.
-	 * @note - the Interchange manager is calling this function not the virtual one that is call by the default implementation.
+	 * @note - the FTaskPipelinePostImport is calling this function not the virtual one that is call by the default implementation.
 	 */
 	UFUNCTION(BlueprintNativeEvent, BlueprintCallable, Category = "Interchange | Pipeline")
 	void ScriptedExecutePostImportPipeline(const UInterchangeBaseNodeContainer* BaseNodeContainer, const FString& FactoryNodeKey, UObject* CreatedAsset, bool bIsAReimport);
-	/** The default implementation (call if the blueprint do not have any implementation) will call the virtual ExecuteImportPipeline */
+	/** The default implementation (call if the blueprint do not have any implementation) will call the virtual ExecutePostImportPipeline */
 	void ScriptedExecutePostImportPipeline_Implementation(const UInterchangeBaseNodeContainer* BaseNodeContainer, const FString& FactoryNodeKey, UObject* CreatedAsset, bool bIsAReimport)
 	{
 		//By default we call the virtual import pipeline execution
@@ -273,19 +298,31 @@ public:
 	bool IsReimportContext() { return bIsReimportContext; }
 protected:
 
+	UE_DEPRECATED(5.2, "This function is replace by ExecutePipeline.")
+	virtual void ExecutePreImportPipeline(UInterchangeBaseNodeContainer* BaseNodeContainer, const TArray<UInterchangeSourceData*>& SourceDatas)
+	{
+		ExecutePipeline(BaseNodeContainer, SourceDatas);
+	}
+
 	/**
-	 * This function can modify the BaseNodeContainer to create a pipeline that will set the graph and the nodes options has it want it to be imported by the factories
-	 * The interchange manager is not calling this function directly. It is calling the blueprint native event in case this object is a blueprint derive object.
-	 * By default the scripted implementation is calling this virtual pipeline.
+	 * This function is called after the translation is done. It should create any factory node representing unreal assets to be imported by the factories
 	 */
-	virtual void ExecutePreImportPipeline(UInterchangeBaseNodeContainer* BaseNodeContainer, const TArray<UInterchangeSourceData*>& SourceDatas)
+	virtual void ExecutePipeline(UInterchangeBaseNodeContainer* BaseNodeContainer, const TArray<UInterchangeSourceData*>& SourceDatas)
+	{
+	}
+
+	/**
+	 * This function is called after the factory create an asset, but before PostEditChange is called.
+	 * It is always call on the game thread, code inside this function should not wait on other thread. (do not stall the game thread)
+	 * This is the place to do any change to the unreal asset before PostEditChange is call.
+	 */
+	virtual void ExecutePostFactoryPipeline(const UInterchangeBaseNodeContainer* BaseNodeContainer, const FString& NodeKey, UObject* CreatedAsset, bool bIsAReimport)
 	{
 	}
 
 	/**
-	 * This function can read the node data and apply some change to the imported asset. This is called after the factory creates the asset and configures the asset properties.
-	 * The interchange manager is not calling this function directly. It is calling the blueprint native event in case this object is a blueprint derived object.
-	 * By default the scripted implementation is calling this virtual pipeline.
+	 * This function is called after the unreal asset is completely import and PostEditChange was called.
+	 * @Note: Some unreal assets have asynchronous build, its possible they are still compiling.
 	 */
 	virtual void ExecutePostImportPipeline(const UInterchangeBaseNodeContainer* BaseNodeContainer, const FString& NodeKey, UObject* CreatedAsset, bool bIsAReimport)
 	{
diff --git a/Engine/Source/Runtime/Interchange/Engine/Private/InterchangeManager.cpp b/Engine/Source/Runtime/Interchange/Engine/Private/InterchangeManager.cpp
--- a/Engine/Source/Runtime/Interchange/Engine/Private/InterchangeManager.cpp
+++ b/Engine/Source/Runtime/Interchange/Engine/Private/InterchangeManager.cpp
@@ -299,7 +299,7 @@ FGraphEventArray UE::Interchange::FImportAsyncHelper::GetCompletionTaskGraphEven
 	FGraphEventArray TasksToComplete;
 
 	TasksToComplete.Append(TranslatorTasks);
-	TasksToComplete.Append(PipelinePreImportTasks);
+	TasksToComplete.Append(PipelineTasks);
 	
 	if (ParsingTask.GetReference())
 	{
@@ -917,13 +917,13 @@ void UInterchangeManager::StartQueuedTasks(bool bCancelAllTasks /*= false*/)
 				UInterchangePipelineBase* GraphPipeline = QueuedTaskData.AsyncHelper->Pipelines[GraphPipelineIndex];
 				TWeakObjectPtr<UInterchangePipelineBase> WeakPipelinePtr = GraphPipeline;
 				int32 GraphPipelineTaskIndex = INDEX_NONE;
-				GraphPipelineTaskIndex = QueuedTaskData.AsyncHelper->PipelinePreImportTasks.Add(TGraphTask<UE::Interchange::FTaskPipelinePreImport>::CreateTask(&PipelinePrerequistes).ConstructAndDispatchWhenReady(WeakPipelinePtr, WeakAsyncHelper));
+				GraphPipelineTaskIndex = QueuedTaskData.AsyncHelper->PipelineTasks.Add(TGraphTask<UE::Interchange::FTaskPipeline>::CreateTask(&PipelinePrerequistes).ConstructAndDispatchWhenReady(WeakPipelinePtr, WeakAsyncHelper));
 				//Ensure we run the pipeline in the same order we create the task, since pipeline modify the node container, its important that its not process in parallel, Adding the one we start to the prerequisites
 				//is the way to go here
-				PipelinePrerequistes.Add(QueuedTaskData.AsyncHelper->PipelinePreImportTasks[GraphPipelineTaskIndex]);
+				PipelinePrerequistes.Add(QueuedTaskData.AsyncHelper->PipelineTasks[GraphPipelineTaskIndex]);
 
 				//Add pipeline to the graph parsing prerequisites
-				GraphParsingPrerequistes.Add(QueuedTaskData.AsyncHelper->PipelinePreImportTasks[GraphPipelineTaskIndex]);
+				GraphParsingPrerequistes.Add(QueuedTaskData.AsyncHelper->PipelineTasks[GraphPipelineTaskIndex]);
 			}
 
 			if (GraphParsingPrerequistes.Num() > 0)
diff --git a/Engine/Source/Runtime/Interchange/Engine/Private/Tasks/InterchangeTaskCompletion.cpp b/Engine/Source/Runtime/Interchange/Engine/Private/Tasks/InterchangeTaskCompletion.cpp
--- a/Engine/Source/Runtime/Interchange/Engine/Private/Tasks/InterchangeTaskCompletion.cpp
+++ b/Engine/Source/Runtime/Interchange/Engine/Private/Tasks/InterchangeTaskCompletion.cpp
@@ -9,6 +9,7 @@
 #include "InterchangeEngineLogPrivate.h"
 #include "InterchangeFactoryBase.h"
 #include "InterchangeManager.h"
+#include "InterchangePipelineBase.h"
 #include "InterchangeResultsContainer.h"
 #include "Nodes/InterchangeFactoryBaseNode.h"
 #include "Stats/Stats.h"
@@ -82,7 +83,7 @@ void UE::Interchange::FTaskPreCompletion::DoTask(ENamedThreads::Type CurrentThre
 					Arguments.Pipelines = AsyncHelper->Pipelines;
 					Arguments.OriginalPipelines = AsyncHelper->OriginalPipelines;
 					Arguments.bIsReimport = ObjectInfo.bIsReimport;
-					ObjectInfo.Factory->PreImportPreCompletedCallback(Arguments);
+					ObjectInfo.Factory->BeginPreCompletedCallback(Arguments);
 				}
 
 				if (ImportedObject == nullptr || !IsValid(ImportedObject))
@@ -122,6 +123,14 @@ void UE::Interchange::FTaskPreCompletion::DoTask(ENamedThreads::Type CurrentThre
 					}
 				}
 
+				for (UInterchangePipelineBase* PipelineBase : AsyncHelper->Pipelines)
+				{
+					PipelineBase->ScriptedExecutePostFactoryPipeline(AsyncHelper->BaseNodeContainers[SourceIndex].Get()
+						, ObjectInfo.FactoryNode ? ObjectInfo.FactoryNode->GetUniqueID() : FString()
+						, ImportedObject
+						, ObjectInfo.bIsReimport);
+				}
+
 #if WITH_EDITOR
 				ImportedObject->PostEditChange();
 #endif
@@ -155,7 +164,7 @@ void UE::Interchange::FTaskPreCompletion::DoTask(ENamedThreads::Type CurrentThre
 					Arguments.Pipelines = AsyncHelper->Pipelines;
 					Arguments.OriginalPipelines = AsyncHelper->OriginalPipelines;
 					Arguments.bIsReimport = ObjectInfo.bIsReimport;
-					ObjectInfo.Factory->PostImportPreCompletedCallback(Arguments);
+					ObjectInfo.Factory->EndPreCompletedCallback(Arguments);
 				}
 			}
 		}
diff --git a/Engine/Source/Runtime/Interchange/Engine/Private/Tasks/InterchangeTaskPipeline.cpp b/Engine/Source/Runtime/Interchange/Engine/Private/Tasks/InterchangeTaskPipeline.cpp
--- a/Engine/Source/Runtime/Interchange/Engine/Private/Tasks/InterchangeTaskPipeline.cpp
+++ b/Engine/Source/Runtime/Interchange/Engine/Private/Tasks/InterchangeTaskPipeline.cpp
@@ -15,9 +15,9 @@
 
 
 
-void UE::Interchange::FTaskPipelinePreImport::DoTask(ENamedThreads::Type CurrentThread, const FGraphEventRef& MyCompletionGraphEvent)
+void UE::Interchange::FTaskPipeline::DoTask(ENamedThreads::Type CurrentThread, const FGraphEventRef& MyCompletionGraphEvent)
 {
-	TRACE_CPUPROFILER_EVENT_SCOPE("UE::Interchange::FTaskPipelinePreImport::DoTask")
+	TRACE_CPUPROFILER_EVENT_SCOPE("UE::Interchange::FTaskPipeline::DoTask")
 #if INTERCHANGE_TRACE_ASYNCHRONOUS_TASK_ENABLED
 	INTERCHANGE_TRACE_ASYNCHRONOUS_TASK(PipelinePreImport)
 #endif
@@ -45,7 +45,7 @@ void UE::Interchange::FTaskPipelinePreImport::DoTask(ENamedThreads::Type Current
 
 			if (ensure(AsyncHelper->BaseNodeContainers[GraphIndex].IsValid()))
 			{
-				Pipeline->ScriptedExecutePreImportPipeline(AsyncHelper->BaseNodeContainers[GraphIndex].Get(), AsyncHelper->SourceDatas);
+				Pipeline->ScriptedExecutePipeline(AsyncHelper->BaseNodeContainers[GraphIndex].Get(), AsyncHelper->SourceDatas);
 			}
 		}
 	}
diff --git a/Engine/Source/Runtime/Interchange/Engine/Private/Tasks/InterchangeTaskPipeline.h b/Engine/Source/Runtime/Interchange/Engine/Private/Tasks/InterchangeTaskPipeline.h
--- a/Engine/Source/Runtime/Interchange/Engine/Private/Tasks/InterchangeTaskPipeline.h
+++ b/Engine/Source/Runtime/Interchange/Engine/Private/Tasks/InterchangeTaskPipeline.h
@@ -14,13 +14,13 @@ namespace UE
 	namespace Interchange
 	{
 
-		class FTaskPipelinePreImport
+		class FTaskPipeline
 		{
 		private:
 			TWeakObjectPtr<UInterchangePipelineBase> PipelineBase;
 			TWeakPtr<FImportAsyncHelper, ESPMode::ThreadSafe> WeakAsyncHelper;
 		public:
-			FTaskPipelinePreImport(TWeakObjectPtr<UInterchangePipelineBase> InPipelineBase, TWeakPtr<FImportAsyncHelper, ESPMode::ThreadSafe> InAsyncHelper)
+			FTaskPipeline(TWeakObjectPtr<UInterchangePipelineBase> InPipelineBase, TWeakPtr<FImportAsyncHelper, ESPMode::ThreadSafe> InAsyncHelper)
 				: PipelineBase(InPipelineBase)
 				, WeakAsyncHelper(InAsyncHelper)
 			{
@@ -39,7 +39,7 @@ namespace UE
 					return ENamedThreads::GameThread;
 				}
 
-				return PipelineBase.Get()->ScriptedCanExecuteOnAnyThread(EInterchangePipelineTask::PreFactoryImport) ? ENamedThreads::AnyBackgroundThreadNormalTask : ENamedThreads::GameThread;
+				return PipelineBase.Get()->ScriptedCanExecuteOnAnyThread(EInterchangePipelineTask::PostTranslator) ? ENamedThreads::AnyBackgroundThreadNormalTask : ENamedThreads::GameThread;
 			}
 
 			static ESubsequentsMode::Type GetSubsequentsMode()
@@ -49,7 +49,7 @@ namespace UE
 
 			FORCEINLINE TStatId GetStatId() const
 			{
-				RETURN_QUICK_DECLARE_CYCLE_STAT(FTaskPipelinePreImport, STATGROUP_TaskGraphTasks);
+				RETURN_QUICK_DECLARE_CYCLE_STAT(FTaskPipeline, STATGROUP_TaskGraphTasks);
 			}
 
 			void DoTask(ENamedThreads::Type CurrentThread, const FGraphEventRef& MyCompletionGraphEvent);
@@ -90,7 +90,7 @@ namespace UE
 				}
 
 				//Ask the pipeline implementation
-				if (AsyncHelper->Pipelines[PipelineIndex]->ScriptedCanExecuteOnAnyThread(EInterchangePipelineTask::PostFactoryImport))
+				if (AsyncHelper->Pipelines[PipelineIndex]->ScriptedCanExecuteOnAnyThread(EInterchangePipelineTask::PostImport))
 				{
 					return ENamedThreads::AnyBackgroundThreadNormalTask;
 				}
diff --git a/Engine/Source/Runtime/Interchange/Engine/Public/InterchangeManager.h b/Engine/Source/Runtime/Interchange/Engine/Public/InterchangeManager.h
--- a/Engine/Source/Runtime/Interchange/Engine/Public/InterchangeManager.h
+++ b/Engine/Source/Runtime/Interchange/Engine/Public/InterchangeManager.h
@@ -199,7 +199,7 @@ namespace UE
 			TArray<UObject*> OriginalPipelines;
 
 			TArray<FGraphEventRef> TranslatorTasks;
-			TArray<FGraphEventRef> PipelinePreImportTasks;
+			TArray<FGraphEventRef> PipelineTasks;
 			TArray<FGraphEventRef> PipelinePostImportTasks;
 			FGraphEventRef ParsingTask;
 			TArray<FGraphEventRef> CreatePackageTasks;
