Remove bad API function that return a local temporary object.

#rnx

diff --git a/Engine/Source/Runtime/Interchange/Core/Public/Types/AttributeStorage.h b/Engine/Source/Runtime/Interchange/Core/Public/Types/AttributeStorage.h
--- a/Engine/Source/Runtime/Interchange/Core/Public/Types/AttributeStorage.h
+++ b/Engine/Source/Runtime/Interchange/Core/Public/Types/AttributeStorage.h
@@ -108,11 +108,6 @@ namespace UE
 				Key = FName(Other);
 			}
 
-			FORCEINLINE const TCHAR* operator*() const
-			{
-				return *(Key.ToString());
-			}
-
 			FORCEINLINE FAttributeKey& operator=(const FName& Other)
 			{
 				Key = Other;
