Make sure null pipeline in the import override pipelines list is handle correctly with no crash.

#rnx

diff --git a/Engine/Source/Runtime/Interchange/Engine/Private/InterchangeManager.cpp b/Engine/Source/Runtime/Interchange/Engine/Private/InterchangeManager.cpp
--- a/Engine/Source/Runtime/Interchange/Engine/Private/InterchangeManager.cpp
+++ b/Engine/Source/Runtime/Interchange/Engine/Private/InterchangeManager.cpp
@@ -1304,12 +1304,21 @@ UInterchangeManager::ImportInternal(const FString& ContentPath, const UInterchan
 	{
 		for (int32 GraphPipelineIndex = 0; GraphPipelineIndex < ImportAssetParameters.OverridePipelines.Num(); ++GraphPipelineIndex)
 		{
-			// Duplicate the override pipelines to protect the scripted users form making race conditions
-			UInterchangePipelineBase* GeneratedPipeline = DuplicateObject<UInterchangePipelineBase>(ImportAssetParameters.OverridePipelines[GraphPipelineIndex], GetTransientPackage());
-			AdjustPipelineSettingForContext(GeneratedPipeline);
-			AsyncHelper->Pipelines.Add(GeneratedPipeline);
-			AsyncHelper->OriginalPipelines.Add(GeneratedPipeline);
-			UE::Interchange::Private::FillPipelineAnalyticData(GeneratedPipeline, UniqueId, FString());
+			UInterchangePipelineBase* OverridePipeline = ImportAssetParameters.OverridePipelines[GraphPipelineIndex];
+			if (!OverridePipeline)
+			{
+				UE_LOG(LogInterchangeEngine, Error, TEXT("Interchange import: Override pipeline array contains a NULL pipeline. Script or code need to be fix to avoid this. "));
+				continue;
+			}
+			else
+			{
+				// Duplicate the override pipelines to protect the scripted users form making race conditions
+				UInterchangePipelineBase* GeneratedPipeline = DuplicateObject<UInterchangePipelineBase>(OverridePipeline, GetTransientPackage());
+				AdjustPipelineSettingForContext(GeneratedPipeline);
+				AsyncHelper->Pipelines.Add(GeneratedPipeline);
+				AsyncHelper->OriginalPipelines.Add(GeneratedPipeline);
+				UE::Interchange::Private::FillPipelineAnalyticData(GeneratedPipeline, UniqueId, FString());
+			}
 		}
 	}
 
