Interchange: Prevent notification to pop when doing automated import task.
#rnx

diff --git a/Engine/Source/Runtime/Interchange/Engine/Private/InterchangeManager.cpp b/Engine/Source/Runtime/Interchange/Engine/Private/InterchangeManager.cpp
--- a/Engine/Source/Runtime/Interchange/Engine/Private/InterchangeManager.cpp
+++ b/Engine/Source/Runtime/Interchange/Engine/Private/InterchangeManager.cpp
@@ -864,17 +864,37 @@ void UInterchangeManager::StartQueuedTasks(bool bCancelAllTasks /*= false*/)
 		}
 		else
 		{
-			FText TitleText = NSLOCTEXT("Interchange", "Asynchronous_import_start", "Importing");
-			FAsyncTaskNotificationConfig NotificationConfig;
-			NotificationConfig.bIsHeadless = false;
-			NotificationConfig.bKeepOpenOnFailure = true;
-			NotificationConfig.TitleText = TitleText;
-			NotificationConfig.LogCategory = UE::Interchange::Private::GetLogInterchangePtr();
-			NotificationConfig.bCanCancel.Set(true);
-			NotificationConfig.bKeepOpenOnFailure.Set(true);
-
-			Notification = MakeShared<FAsyncTaskNotification>(NotificationConfig);
-			Notification->SetNotificationState(FAsyncNotificationStateData(TitleText, FText::GetEmpty(), EAsyncTaskNotificationState::Pending));
+			bool bCanShowNotification = false;
+			{
+				int32 ImportTaskCount = ImportTasks.Num();
+				for (int32 TaskIndex = 0; TaskIndex < ImportTaskCount; ++TaskIndex)
+				{
+					TSharedPtr<UE::Interchange::FImportAsyncHelper, ESPMode::ThreadSafe> AsyncHelper = ImportTasks[TaskIndex];
+					if (AsyncHelper.IsValid())
+					{
+						//Allow notification if at least one task is not automated
+						if (!AsyncHelper->TaskData.bIsAutomated)
+						{
+							bCanShowNotification = true;
+							break;
+						}
+					}
+				}
+			}
+			if (bCanShowNotification)
+			{
+				FText TitleText = NSLOCTEXT("Interchange", "Asynchronous_import_start", "Importing");
+				FAsyncTaskNotificationConfig NotificationConfig;
+				NotificationConfig.bIsHeadless = false;
+				NotificationConfig.bKeepOpenOnFailure = true;
+				NotificationConfig.TitleText = TitleText;
+				NotificationConfig.LogCategory = UE::Interchange::Private::GetLogInterchangePtr();
+				NotificationConfig.bCanCancel.Set(true);
+				NotificationConfig.bKeepOpenOnFailure.Set(true);
+
+				Notification = MakeShared<FAsyncTaskNotification>(NotificationConfig);
+				Notification->SetNotificationState(FAsyncNotificationStateData(TitleText, FText::GetEmpty(), EAsyncTaskNotificationState::Pending));
+			}
 		}
 	};
 
