Interchange: Fix task parsing dependencies.

#rnx

diff --git a/Engine/Source/Runtime/Interchange/Engine/Private/Tasks/InterchangeTaskParsing.cpp b/Engine/Source/Runtime/Interchange/Engine/Private/Tasks/InterchangeTaskParsing.cpp
--- a/Engine/Source/Runtime/Interchange/Engine/Private/Tasks/InterchangeTaskParsing.cpp
+++ b/Engine/Source/Runtime/Interchange/Engine/Private/Tasks/InterchangeTaskParsing.cpp
@@ -320,8 +320,7 @@ void UE::Interchange::FTaskParsing::DoTask(ENamedThreads::Type CurrentThread, co
 	FGraphEventArray CompletionPrerequistes;
 	const int32 PoolWorkerThreadCount = FTaskGraphInterface::Get().GetNumWorkerThreads() / 2;
 	const int32 MaxNumWorker = FMath::Max(PoolWorkerThreadCount, 1);
-	FGraphEventArray GroupPrerequistes[2];
-	int32 CurrentGroup = 0;
+	FGraphEventArray GroupPrerequistes;
 	for (int32 TaskIndex = 0; TaskIndex < TaskDatas.Num(); ++TaskIndex)
 	{
 		FTaskData& TaskData = TaskDatas[TaskIndex];
@@ -335,23 +334,21 @@ void UE::Interchange::FTaskParsing::DoTask(ENamedThreads::Type CurrentThread, co
 				{
 					//Add has prerequisite
 					TaskData.Prerequisites.Add(TaskDatas[DepTaskIndex].GraphEventRef);
-					if (GroupPrerequistes[(CurrentGroup + 1) % 2].Num() >= MaxNumWorker)
-					{
-						TaskData.Prerequisites.Append(GroupPrerequistes[(CurrentGroup + 1) % 2]);
-					}
 				}
 			}
 		}
 
-		TaskData.GraphEventRef = CreateTasksFromData(TaskData);
-		if (GroupPrerequistes[CurrentGroup%2].Num() >= MaxNumWorker)
+		//Append the group prerequisite to the task data prerequisite if the group is full
+		//This allow to chain the group dependencies to control the number of task
+		if (GroupPrerequistes.Num() >= MaxNumWorker)
 		{
-			CurrentGroup++;
-			GroupPrerequistes[CurrentGroup % 2].Reset();
+			TaskData.Prerequisites.Append(GroupPrerequistes);
+			GroupPrerequistes.Reset();
 		}
-		GroupPrerequistes[CurrentGroup%2].Add(TaskData.GraphEventRef);
+		TaskData.GraphEventRef = CreateTasksFromData(TaskData);
+		GroupPrerequistes.Add(TaskData.GraphEventRef);
 	}
-	CompletionPrerequistes.Append(GroupPrerequistes[CurrentGroup % 2]);
+	CompletionPrerequistes.Append(GroupPrerequistes);
 
 	if (!RenameAssets.IsEmpty())
 	{
