Interchange: remove unuseful CVar Interchange.FeatureFlags.CustomPipelines
#rnx

diff --git a/Engine/Source/Runtime/Interchange/Core/Private/InterchangePipelineBase.cpp b/Engine/Source/Runtime/Interchange/Core/Private/InterchangePipelineBase.cpp
--- a/Engine/Source/Runtime/Interchange/Core/Private/InterchangePipelineBase.cpp
+++ b/Engine/Source/Runtime/Interchange/Core/Private/InterchangePipelineBase.cpp
@@ -10,14 +10,6 @@
 
 #include UE_INLINE_GENERATED_CPP_BY_NAME(InterchangePipelineBase)
 
-INTERCHANGECORE_API bool GInterchangeEnableCustomPipelines = false;
-
-static FAutoConsoleVariableRef CCvarInterchangeEnableCustomPipelines(
-	TEXT("Interchange.FeatureFlags.CustomPipelines"),
-	GInterchangeEnableCustomPipelines,
-	TEXT("[Experimental] Whether custom pipelines support is enabled."),
-	ECVF_Default);
-
 namespace UE
 {
 	namespace Interchange
