Refactor the interchange factory API to have function name that describe the purpose of the function.

#rnx

diff --git a/Engine/Plugins/Experimental/Enterprise/DatasmithInterchange/Source/DatasmithInterchange/Private/InterchangeDatasmithAreaLightFactory.cpp b/Engine/Plugins/Experimental/Enterprise/DatasmithInterchange/Source/DatasmithInterchange/Private/InterchangeDatasmithAreaLightFactory.cpp
--- a/Engine/Plugins/Experimental/Enterprise/DatasmithInterchange/Source/DatasmithInterchange/Private/InterchangeDatasmithAreaLightFactory.cpp
+++ b/Engine/Plugins/Experimental/Enterprise/DatasmithInterchange/Source/DatasmithInterchange/Private/InterchangeDatasmithAreaLightFactory.cpp
@@ -29,7 +29,7 @@ UClass* UInterchangeDatasmithAreaLightFactory::GetFactoryClass() const
 	return LightShapeBlueprint->GeneratedClass;
 }
 
-UObject* UInterchangeDatasmithAreaLightFactory::CreateSceneObject(const UInterchangeFactoryBase::FCreateSceneObjectsParams& CreateSceneObjectsParams)
+UObject* UInterchangeDatasmithAreaLightFactory::ImportSceneObject_GameThread(const UInterchangeFactoryBase::FImportSceneObjectsParams& CreateSceneObjectsParams)
 {
 	//using namespace UE::Interchange::ActorHelper;
 
diff --git a/Engine/Plugins/Experimental/Enterprise/DatasmithInterchange/Source/DatasmithInterchange/Private/InterchangeDatasmithSceneFactory.cpp b/Engine/Plugins/Experimental/Enterprise/DatasmithInterchange/Source/DatasmithInterchange/Private/InterchangeDatasmithSceneFactory.cpp
--- a/Engine/Plugins/Experimental/Enterprise/DatasmithInterchange/Source/DatasmithInterchange/Private/InterchangeDatasmithSceneFactory.cpp
+++ b/Engine/Plugins/Experimental/Enterprise/DatasmithInterchange/Source/DatasmithInterchange/Private/InterchangeDatasmithSceneFactory.cpp
@@ -15,7 +15,7 @@ UClass* UInterchangeDatasmithSceneFactory::GetFactoryClass() const
 }
 
 
-UObject* UInterchangeDatasmithSceneFactory::CreateEmptyAsset(const FCreateAssetParams& Arguments)
+UObject* UInterchangeDatasmithSceneFactory::ImportAssetObject_GameThread(const FImportAssetObjectParams& Arguments)
 {
 	UDatasmithScene* DatasmithScene = nullptr;
 
@@ -57,7 +57,7 @@ UObject* UInterchangeDatasmithSceneFactory::CreateEmptyAsset(const FCreateAssetP
 	return DatasmithScene;
 }
 
-UObject* UInterchangeDatasmithSceneFactory::CreateAsset(const FCreateAssetParams& Arguments)
+UObject* UInterchangeDatasmithSceneFactory::ImportAssetObject_Async(const FImportAssetObjectParams& Arguments)
 {
  #if !WITH_EDITORONLY_DATA
  	UE_LOG(LogInterchangeDatasmith, Error, TEXT("Cannot import datasmith scene asset in runtime, this is an editor only feature."));
@@ -119,7 +119,7 @@ UObject* UInterchangeDatasmithSceneFactory::CreateAsset(const FCreateAssetParams
 #endif //WITH_EDITORONLY_DATA
 }
 
-void UInterchangeDatasmithSceneFactory::PreImportPreCompletedCallback(const FImportPreCompletedCallbackParams& Arguments)
+void UInterchangeDatasmithSceneFactory::SetupObject_GameThread(const FSetupObjectParams& Arguments)
 {
 
 }
\ No newline at end of file
diff --git a/Engine/Plugins/Experimental/Enterprise/DatasmithInterchange/Source/DatasmithInterchange/Public/InterchangeDatasmithAreaLightFactory.h b/Engine/Plugins/Experimental/Enterprise/DatasmithInterchange/Source/DatasmithInterchange/Public/InterchangeDatasmithAreaLightFactory.h
--- a/Engine/Plugins/Experimental/Enterprise/DatasmithInterchange/Source/DatasmithInterchange/Public/InterchangeDatasmithAreaLightFactory.h
+++ b/Engine/Plugins/Experimental/Enterprise/DatasmithInterchange/Source/DatasmithInterchange/Public/InterchangeDatasmithAreaLightFactory.h
@@ -21,7 +21,7 @@ public:
 
 	virtual UClass* GetFactoryClass() const override;
 
-	virtual UObject* CreateSceneObject(const FCreateSceneObjectsParams& Arguments) override;
+	virtual UObject* ImportSceneObject_GameThread(const FImportSceneObjectsParams& Arguments) override;
 
 	// Interchange factory base interface end
 	//////////////////////////////////////////////////////////////////////////
diff --git a/Engine/Plugins/Experimental/Enterprise/DatasmithInterchange/Source/DatasmithInterchange/Public/InterchangeDatasmithSceneFactory.h b/Engine/Plugins/Experimental/Enterprise/DatasmithInterchange/Source/DatasmithInterchange/Public/InterchangeDatasmithSceneFactory.h
--- a/Engine/Plugins/Experimental/Enterprise/DatasmithInterchange/Source/DatasmithInterchange/Public/InterchangeDatasmithSceneFactory.h
+++ b/Engine/Plugins/Experimental/Enterprise/DatasmithInterchange/Source/DatasmithInterchange/Public/InterchangeDatasmithSceneFactory.h
@@ -21,9 +21,9 @@ public:
 	// Interchange factory base interface begin
 
 	virtual UClass* GetFactoryClass() const override;
-	virtual UObject* CreateEmptyAsset(const FCreateAssetParams& Arguments) override;
-	virtual UObject* CreateAsset(const FCreateAssetParams& Arguments) override;
-	virtual void PreImportPreCompletedCallback(const FImportPreCompletedCallbackParams& Arguments) override;
+	virtual UObject* ImportAssetObject_GameThread(const FImportAssetObjectParams& Arguments) override;
+	virtual UObject* ImportAssetObject_Async(const FImportAssetObjectParams& Arguments) override;
+	virtual void SetupObject_GameThread(const FSetupObjectParams& Arguments) override;
 	// Interchange factory base interface end
 	//////////////////////////////////////////////////////////////////////////
 
diff --git a/Engine/Plugins/Interchange/Runtime/Source/Import/Private/Animation/InterchangeAnimSequenceFactory.cpp b/Engine/Plugins/Interchange/Runtime/Source/Import/Private/Animation/InterchangeAnimSequenceFactory.cpp
--- a/Engine/Plugins/Interchange/Runtime/Source/Import/Private/Animation/InterchangeAnimSequenceFactory.cpp
+++ b/Engine/Plugins/Interchange/Runtime/Source/Import/Private/Animation/InterchangeAnimSequenceFactory.cpp
@@ -737,7 +737,7 @@ UClass* UInterchangeAnimSequenceFactory::GetFactoryClass() const
 	return UAnimSequence::StaticClass();
 }
 
-UObject* UInterchangeAnimSequenceFactory::CreateEmptyAsset(const FCreateAssetParams& Arguments)
+UObject* UInterchangeAnimSequenceFactory::ImportAssetObject_GameThread(const FImportAssetObjectParams& Arguments)
 {
 #if !WITH_EDITOR || !WITH_EDITORONLY_DATA
 
@@ -784,11 +784,11 @@ UObject* UInterchangeAnimSequenceFactory::CreateEmptyAsset(const FCreateAssetPar
 
 	AnimSequence->PreEditChange(nullptr);
 
-	return AnimSequence;
+	return ImportObjectSourceData(Arguments);
 #endif //else !WITH_EDITOR || !WITH_EDITORONLY_DATA
 }
 
-UObject* UInterchangeAnimSequenceFactory::CreateAsset(const FCreateAssetParams& Arguments)
+UObject* UInterchangeAnimSequenceFactory::ImportObjectSourceData(const FImportAssetObjectParams& Arguments)
 {
 #if !WITH_EDITOR || !WITH_EDITORONLY_DATA
 
@@ -986,10 +986,10 @@ UObject* UInterchangeAnimSequenceFactory::CreateAsset(const FCreateAssetParams&
 }
 
 /* This function is call in the completion task on the main thread, use it to call main thread post creation step for your assets*/
-void UInterchangeAnimSequenceFactory::BeginPreCompletedCallback(const FImportPreCompletedCallbackParams& Arguments)
+void UInterchangeAnimSequenceFactory::SetupObject_GameThread(const FSetupObjectParams& Arguments)
 {
 	check(IsInGameThread());
-	Super::BeginPreCompletedCallback(Arguments);
+	Super::SetupObject_GameThread(Arguments);
 
 	// TODO: make sure this works at runtime
 #if WITH_EDITORONLY_DATA
@@ -1030,7 +1030,7 @@ bool UInterchangeAnimSequenceFactory::SetSourceFilename(const UObject* Object, c
 	return false;
 }
 
-bool UInterchangeAnimSequenceFactory::IsBoneTrackAnimationValid(const UInterchangeAnimSequenceFactoryNode* AnimSequenceFactoryNode, const FCreateAssetParams& Arguments)
+bool UInterchangeAnimSequenceFactory::IsBoneTrackAnimationValid(const UInterchangeAnimSequenceFactoryNode* AnimSequenceFactoryNode, const FImportAssetObjectParams& Arguments)
 {
 	bool bResult = true;
 	FFrameRate FrameRate(30, 1);
diff --git a/Engine/Plugins/Interchange/Runtime/Source/Import/Private/Animation/InterchangeAnimationTrackSetFactory.cpp b/Engine/Plugins/Interchange/Runtime/Source/Import/Private/Animation/InterchangeAnimationTrackSetFactory.cpp
--- a/Engine/Plugins/Interchange/Runtime/Source/Import/Private/Animation/InterchangeAnimationTrackSetFactory.cpp
+++ b/Engine/Plugins/Interchange/Runtime/Source/Import/Private/Animation/InterchangeAnimationTrackSetFactory.cpp
@@ -560,7 +560,7 @@ UClass* UInterchangeAnimationTrackSetFactory::GetFactoryClass() const
 	return ULevelSequence::StaticClass();
 }
 
-UObject* UInterchangeAnimationTrackSetFactory::CreateEmptyAsset(const FCreateAssetParams& Arguments)
+UObject* UInterchangeAnimationTrackSetFactory::ImportAssetObject_GameThread(const FImportAssetObjectParams& Arguments)
 {
 #if !WITH_EDITOR || !WITH_EDITORONLY_DATA
 
@@ -612,11 +612,11 @@ UObject* UInterchangeAnimationTrackSetFactory::CreateEmptyAsset(const FCreateAss
 
 	LevelSequence->PreEditChange(nullptr);
 
-	return LevelSequence;
+	return ImportObjectSourceData(Arguments);
 #endif //else !WITH_EDITOR || !WITH_EDITORONLY_DATA
 }
 
-UObject* UInterchangeAnimationTrackSetFactory::CreateAsset(const FCreateAssetParams& Arguments)
+UObject* UInterchangeAnimationTrackSetFactory::ImportObjectSourceData(const FImportAssetObjectParams& Arguments)
 {
 #if !WITH_EDITOR || !WITH_EDITORONLY_DATA
 	// TODO: Can we import ULevelSequence at runtime
@@ -706,15 +706,15 @@ UObject* UInterchangeAnimationTrackSetFactory::CreateAsset(const FCreateAssetPar
 }
 
 /* This function is call in the completion task on the main thread, use it to call main thread post creation step for your assets*/
-void UInterchangeAnimationTrackSetFactory::BeginPreCompletedCallback(const FImportPreCompletedCallbackParams& Arguments)
+void UInterchangeAnimationTrackSetFactory::SetupObject_GameThread(const FSetupObjectParams& Arguments)
 {
 	check(IsInGameThread());
-	Super::BeginPreCompletedCallback(Arguments);
+	Super::SetupObject_GameThread(Arguments);
 
 	// TODO: Talk with sequence team about adding AssetImportData to ULevelSequence for re-import 
 }
 
-void UInterchangeAnimationTrackSetFactory::EndPreCompletedCallback(const FImportPreCompletedCallbackParams& Arguments)
+void UInterchangeAnimationTrackSetFactory::FinalizeObject_GameThread(const FSetupObjectParams& Arguments)
 {
 	check(IsInGameThread());
 	if (ULevelSequence* LevelSequence = Cast<ULevelSequence>(Arguments.ImportedObject))
@@ -726,7 +726,7 @@ void UInterchangeAnimationTrackSetFactory::EndPreCompletedCallback(const FImport
 		// Note: Sequencer team is aware of the issue and will look their design in a future release
 		UMovieSceneSignedObject::ResetImplicitScopedModifyDefer();
 	}
-	Super::EndPreCompletedCallback(Arguments);
+	Super::FinalizeObject_GameThread(Arguments);
 }
 
 #undef LOCTEXT_NAMESPACE
\ No newline at end of file
diff --git a/Engine/Plugins/Interchange/Runtime/Source/Import/Private/Material/InterchangeMaterialFactory.cpp b/Engine/Plugins/Interchange/Runtime/Source/Import/Private/Material/InterchangeMaterialFactory.cpp
--- a/Engine/Plugins/Interchange/Runtime/Source/Import/Private/Material/InterchangeMaterialFactory.cpp
+++ b/Engine/Plugins/Interchange/Runtime/Source/Import/Private/Material/InterchangeMaterialFactory.cpp
@@ -116,7 +116,7 @@ namespace UE
 					return 0; // Consider 0 as the default output to connect to since most expressions have a single output
 				}
 
-				void SetupFunctionCallExpression(const UInterchangeFactoryBase::FCreateAssetParams& Arguments, const UInterchangeMaterialExpressionFactoryNode& ExpressionNode, UMaterialExpressionMaterialFunctionCall* FunctionCallExpression)
+				void SetupFunctionCallExpression(const UInterchangeFactoryBase::FImportAssetObjectParams& Arguments, const UInterchangeMaterialExpressionFactoryNode& ExpressionNode, UMaterialExpressionMaterialFunctionCall* FunctionCallExpression)
 				{
 					FInterchangeImportMaterialAsyncHelper& AsyncHelper = FInterchangeImportMaterialAsyncHelper::GetInstance();
 					if (const UInterchangeMaterialFunctionCallExpressionFactoryNode* FunctionCallFactoryNode = Cast<UInterchangeMaterialFunctionCallExpressionFactoryNode>(&ExpressionNode))
@@ -139,7 +139,7 @@ namespace UE
 					AsyncHelper.UpdateFromFunctionResource(FunctionCallExpression);
 				}
 
-				void SetupTextureExpression(const UInterchangeFactoryBase::FCreateAssetParams& Arguments, const UInterchangeMaterialExpressionFactoryNode* ExpressionNode, UMaterialExpressionTextureBase* TextureExpression)
+				void SetupTextureExpression(const UInterchangeFactoryBase::FImportAssetObjectParams& Arguments, const UInterchangeMaterialExpressionFactoryNode* ExpressionNode, UMaterialExpressionTextureBase* TextureExpression)
 				{
 					using namespace UE::Interchange::Materials::Standard::Nodes::TextureSample;
 
@@ -179,7 +179,7 @@ namespace UE
 				class FMaterialExpressionBuilder
 				{
 				public:
-					FMaterialExpressionBuilder(UMaterial* InMaterial, UMaterialFunction* InMaterialFunction, const UInterchangeFactoryBase::FCreateAssetParams& InArguments)
+					FMaterialExpressionBuilder(UMaterial* InMaterial, UMaterialFunction* InMaterialFunction, const UInterchangeFactoryBase::FImportAssetObjectParams& InArguments)
 						: Material(InMaterial)
 						, MaterialFunction(InMaterialFunction)
 						, Arguments(InArguments)
@@ -327,7 +327,7 @@ namespace UE
 
 					UMaterial* Material = nullptr;
 					UMaterialFunction* MaterialFunction = nullptr;
-					const UInterchangeFactoryBase::FCreateAssetParams& Arguments;
+					const UInterchangeFactoryBase::FImportAssetObjectParams& Arguments;
 					TMap<FString, UMaterialExpression*> Expressions;
 				};
 #endif // #if WITH_EDITOR
@@ -341,7 +341,7 @@ UClass* UInterchangeMaterialFactory::GetFactoryClass() const
 	return UMaterialInterface::StaticClass();
 }
 
-UObject* UInterchangeMaterialFactory::CreateEmptyAsset(const FCreateAssetParams& Arguments)
+UObject* UInterchangeMaterialFactory::ImportAssetObject_GameThread(const FImportAssetObjectParams& Arguments)
 {
 	UObject* Material = nullptr;
 
@@ -417,7 +417,7 @@ UObject* UInterchangeMaterialFactory::CreateEmptyAsset(const FCreateAssetParams&
 	return Material;
 }
 
-UObject* UInterchangeMaterialFactory::CreateAsset(const FCreateAssetParams& Arguments)
+UObject* UInterchangeMaterialFactory::ImportAssetObject_Async(const FImportAssetObjectParams& Arguments)
 {
 	if (!Arguments.AssetNode || !Arguments.AssetNode->GetObjectClass()->IsChildOf(GetFactoryClass()))
 	{
@@ -505,10 +505,10 @@ UObject* UInterchangeMaterialFactory::CreateAsset(const FCreateAssetParams& Argu
 }
 
 /* This function is call in the completion task on the main thread, use it to call main thread post creation step for your assets*/
-void UInterchangeMaterialFactory::BeginPreCompletedCallback(const FImportPreCompletedCallbackParams& Arguments)
+void UInterchangeMaterialFactory::SetupObject_GameThread(const FSetupObjectParams& Arguments)
 {
 	check(IsInGameThread());
-	Super::BeginPreCompletedCallback(Arguments);
+	Super::SetupObject_GameThread(Arguments);
 
 	if (ensure(Arguments.ImportedObject && Arguments.SourceData))
 	{
@@ -573,7 +573,7 @@ bool UInterchangeMaterialFactory::SetSourceFilename(const UObject* Object, const
 }
 
 #if WITH_EDITOR
-void UInterchangeMaterialFactory::SetupMaterial(UMaterial* Material, const FCreateAssetParams& Arguments, const UInterchangeBaseMaterialFactoryNode* BaseMaterialFactoryNode)
+void UInterchangeMaterialFactory::SetupMaterial(UMaterial* Material, const FImportAssetObjectParams& Arguments, const UInterchangeBaseMaterialFactoryNode* BaseMaterialFactoryNode)
 {
 	using namespace UE::Interchange::Materials;
 	using namespace UE::Interchange::MaterialFactory::Internal;
@@ -1148,7 +1148,7 @@ UClass* UInterchangeMaterialFunctionFactory::GetFactoryClass() const
 	return UMaterialFunctionInterface::StaticClass();
 }
 
-UObject* UInterchangeMaterialFunctionFactory::CreateEmptyAsset(const FCreateAssetParams& Arguments)
+UObject* UInterchangeMaterialFunctionFactory::ImportAssetObject_GameThread(const FImportAssetObjectParams& Arguments)
 {
 	UObject* Material = nullptr;
 
@@ -1196,7 +1196,7 @@ UObject* UInterchangeMaterialFunctionFactory::CreateEmptyAsset(const FCreateAsse
 	return Material;
 }
 
-UObject* UInterchangeMaterialFunctionFactory::CreateAsset(const FCreateAssetParams& Arguments)
+UObject* UInterchangeMaterialFunctionFactory::ImportAssetObject_Async(const FImportAssetObjectParams& Arguments)
 {
 	if (!Arguments.AssetNode || !Arguments.AssetNode->GetObjectClass()->IsChildOf(GetFactoryClass()))
 	{
@@ -1262,9 +1262,9 @@ UObject* UInterchangeMaterialFunctionFactory::CreateAsset(const FCreateAssetPara
 }
 
 /* This function is call in the completion task on the main thread, use it to call main thread post creation step for your assets*/
-void UInterchangeMaterialFunctionFactory::BeginPreCompletedCallback(const FImportPreCompletedCallbackParams& Arguments)
+void UInterchangeMaterialFunctionFactory::SetupObject_GameThread(const FSetupObjectParams& Arguments)
 {
-	Super::BeginPreCompletedCallback(Arguments);
+	Super::SetupObject_GameThread(Arguments);
 
 #if WITH_EDITORONLY_DATA
 	if (ensure(Arguments.ImportedObject && Arguments.SourceData))
@@ -1287,7 +1287,7 @@ void UInterchangeMaterialFunctionFactory::BeginPreCompletedCallback(const FImpor
 }
 
 #if WITH_EDITOR
-void UInterchangeMaterialFunctionFactory::SetupMaterial(UMaterialFunction* MaterialFunction, const FCreateAssetParams& Arguments, const UInterchangeMaterialFunctionFactoryNode* MaterialFunctionFactoryNode)
+void UInterchangeMaterialFunctionFactory::SetupMaterial(UMaterialFunction* MaterialFunction, const FImportAssetObjectParams& Arguments, const UInterchangeMaterialFunctionFactoryNode* MaterialFunctionFactoryNode)
 {
 	using namespace UE::Interchange::MaterialFactory::Internal;
 	using namespace UE::Interchange::Materials;
diff --git a/Engine/Plugins/Interchange/Runtime/Source/Import/Private/Mesh/InterchangePhysicsAssetFactory.cpp b/Engine/Plugins/Interchange/Runtime/Source/Import/Private/Mesh/InterchangePhysicsAssetFactory.cpp
--- a/Engine/Plugins/Interchange/Runtime/Source/Import/Private/Mesh/InterchangePhysicsAssetFactory.cpp
+++ b/Engine/Plugins/Interchange/Runtime/Source/Import/Private/Mesh/InterchangePhysicsAssetFactory.cpp
@@ -16,7 +16,7 @@ UClass* UInterchangePhysicsAssetFactory::GetFactoryClass() const
 	return UPhysicsAsset::StaticClass();
 }
 
-UObject* UInterchangePhysicsAssetFactory::CreateEmptyAsset(const FCreateAssetParams& Arguments)
+UObject* UInterchangePhysicsAssetFactory::ImportAssetObject_GameThread(const FImportAssetObjectParams& Arguments)
 {
 	UObject* PhysicsAsset = nullptr;
 
@@ -58,7 +58,7 @@ UObject* UInterchangePhysicsAssetFactory::CreateEmptyAsset(const FCreateAssetPar
 	return PhysicsAsset;
 }
 
-UObject* UInterchangePhysicsAssetFactory::CreateAsset(const FCreateAssetParams& Arguments)
+UObject* UInterchangePhysicsAssetFactory::ImportAssetObject_Async(const FImportAssetObjectParams& Arguments)
 {
 #if !WITH_EDITORONLY_DATA
 
@@ -89,7 +89,7 @@ UObject* UInterchangePhysicsAssetFactory::CreateAsset(const FCreateAssetParams&
 	if (!ExistingAsset)
 	{
 		//NewObject is not thread safe, the asset registry directory watcher tick on the main thread can trig before we finish initializing the UObject and will crash
-		//The UObject should have been create by calling CreateEmptyAsset on the main thread.
+		//The UObject should have been create by calling ImportAssetObject_GameThread on the main thread.
 		check(IsInGameThread());
 		PhysicsAssetObject = NewObject<UObject>(Arguments.Parent, PhysicsAssetClass, *Arguments.AssetName, RF_Public | RF_Standalone);
 	}
diff --git a/Engine/Plugins/Interchange/Runtime/Source/Import/Private/Mesh/InterchangeSkeletalMeshFactory.cpp b/Engine/Plugins/Interchange/Runtime/Source/Import/Private/Mesh/InterchangeSkeletalMeshFactory.cpp
--- a/Engine/Plugins/Interchange/Runtime/Source/Import/Private/Mesh/InterchangeSkeletalMeshFactory.cpp
+++ b/Engine/Plugins/Interchange/Runtime/Source/Import/Private/Mesh/InterchangeSkeletalMeshFactory.cpp
@@ -406,7 +406,7 @@ namespace UE
 																  , FSkeletalMeshImportData& DestinationImportData
 																  , TArray<FMeshNodeContext>& MeshReferences
 																  , TArray<SkeletalMeshImportData::FBone>& RefBonesBinary
-																  , const UInterchangeSkeletalMeshFactory::FCreateAssetParams& Arguments
+																  , const UInterchangeSkeletalMeshFactory::FImportAssetObjectParams& Arguments
 																  , const IInterchangeSkeletalMeshPayloadInterface* SkeletalMeshTranslatorPayloadInterface
 																  , const bool bSkinControlPointToTimeZero
 																  , const UInterchangeBaseNodeContainer* NodeContainer
@@ -872,9 +872,9 @@ UClass* UInterchangeSkeletalMeshFactory::GetFactoryClass() const
 	return USkeletalMesh::StaticClass();
 }
 
-UObject* UInterchangeSkeletalMeshFactory::CreateEmptyAsset(const FCreateAssetParams& Arguments)
+UObject* UInterchangeSkeletalMeshFactory::ImportAssetObject_GameThread(const FImportAssetObjectParams& Arguments)
 {
-	TRACE_CPUPROFILER_EVENT_SCOPE("UInterchangeSkeletalMeshFactory::CreateEmptyAsset")
+	TRACE_CPUPROFILER_EVENT_SCOPE("UInterchangeSkeletalMeshFactory::ImportAssetObject_GameThread")
 #if !WITH_EDITOR || !WITH_EDITORONLY_DATA
 
 	UE_LOG(LogInterchangeImport, Error, TEXT("Cannot import skeletalMesh asset in runtime, this is an editor only feature."));
@@ -931,7 +931,7 @@ UObject* UInterchangeSkeletalMeshFactory::CreateEmptyAsset(const FCreateAssetPar
 #endif //else !WITH_EDITOR || !WITH_EDITORONLY_DATA
 }
 
-UObject* UInterchangeSkeletalMeshFactory::CreateAsset(const FCreateAssetParams& Arguments)
+UObject* UInterchangeSkeletalMeshFactory::ImportAssetObject_Async(const FImportAssetObjectParams& Arguments)
 {
 	TRACE_CPUPROFILER_EVENT_SCOPE("UInterchangeSkeletalMeshFactory::CreateAsset")
 
@@ -1618,11 +1618,11 @@ void UInterchangeSkeletalMeshFactory::Cancel()
 }
 
 /* This function is call in the completion task on the main thread, use it to call main thread post creation step for your assets*/
-void UInterchangeSkeletalMeshFactory::BeginPreCompletedCallback(const FImportPreCompletedCallbackParams& Arguments)
+void UInterchangeSkeletalMeshFactory::SetupObject_GameThread(const FSetupObjectParams& Arguments)
 {
 	TRACE_CPUPROFILER_EVENT_SCOPE("UInterchangeSkeletalMeshFactory::PreImportPreCompletedCallback")
 	check(IsInGameThread());
-	Super::BeginPreCompletedCallback(Arguments);
+	Super::SetupObject_GameThread(Arguments);
 
 	//TODO make sure this work at runtime
 #if WITH_EDITORONLY_DATA
diff --git a/Engine/Plugins/Interchange/Runtime/Source/Import/Private/Mesh/InterchangeSkeletonFactory.cpp b/Engine/Plugins/Interchange/Runtime/Source/Import/Private/Mesh/InterchangeSkeletonFactory.cpp
--- a/Engine/Plugins/Interchange/Runtime/Source/Import/Private/Mesh/InterchangeSkeletonFactory.cpp
+++ b/Engine/Plugins/Interchange/Runtime/Source/Import/Private/Mesh/InterchangeSkeletonFactory.cpp
@@ -17,7 +17,7 @@ UClass* UInterchangeSkeletonFactory::GetFactoryClass() const
 	return USkeleton::StaticClass();
 }
 
-UObject* UInterchangeSkeletonFactory::CreateEmptyAsset(const FCreateAssetParams& Arguments)
+UObject* UInterchangeSkeletonFactory::ImportAssetObject_GameThread(const FImportAssetObjectParams& Arguments)
 {
 	UObject* Skeleton = nullptr;
 
@@ -59,7 +59,7 @@ UObject* UInterchangeSkeletonFactory::CreateEmptyAsset(const FCreateAssetParams&
 	return Skeleton;
 }
 
-UObject* UInterchangeSkeletonFactory::CreateAsset(const FCreateAssetParams& Arguments)
+UObject* UInterchangeSkeletonFactory::ImportAssetObject_Async(const FImportAssetObjectParams& Arguments)
 {
 #if !WITH_EDITORONLY_DATA
 
diff --git a/Engine/Plugins/Interchange/Runtime/Source/Import/Private/Mesh/InterchangeStaticMeshFactory.cpp b/Engine/Plugins/Interchange/Runtime/Source/Import/Private/Mesh/InterchangeStaticMeshFactory.cpp
--- a/Engine/Plugins/Interchange/Runtime/Source/Import/Private/Mesh/InterchangeStaticMeshFactory.cpp
+++ b/Engine/Plugins/Interchange/Runtime/Source/Import/Private/Mesh/InterchangeStaticMeshFactory.cpp
@@ -51,7 +51,7 @@ UClass* UInterchangeStaticMeshFactory::GetFactoryClass() const
 }
 
 
-UObject* UInterchangeStaticMeshFactory::CreateEmptyAsset(const FCreateAssetParams& Arguments)
+UObject* UInterchangeStaticMeshFactory::ImportAssetObject_GameThread(const FImportAssetObjectParams& Arguments)
 {
 	UStaticMesh* StaticMesh = nullptr;
 	if (!Arguments.AssetNode || !Arguments.AssetNode->GetObjectClass()->IsChildOf(GetFactoryClass()))
@@ -108,7 +108,7 @@ UObject* UInterchangeStaticMeshFactory::CreateEmptyAsset(const FCreateAssetParam
 	return StaticMesh;
 }
 
-UObject* UInterchangeStaticMeshFactory::CreateAsset(const FCreateAssetParams& Arguments)
+UObject* UInterchangeStaticMeshFactory::ImportAssetObject_Async(const FImportAssetObjectParams& Arguments)
 {
 	if (!Arguments.AssetNode || !Arguments.AssetNode->GetObjectClass()->IsChildOf(GetFactoryClass()))
 	{
@@ -523,10 +523,10 @@ void UInterchangeStaticMeshFactory::SetupSourceModelsSettings(UStaticMesh& Stati
 #endif // WITH_EDITORONLY_DATA
 
 /* This function is call in the completion task on the main thread, use it to call main thread post creation step for your assets */
-void UInterchangeStaticMeshFactory::BeginPreCompletedCallback(const FImportPreCompletedCallbackParams& Arguments)
+void UInterchangeStaticMeshFactory::SetupObject_GameThread(const FSetupObjectParams& Arguments)
 {
 	check(IsInGameThread());
-	Super::BeginPreCompletedCallback(Arguments);
+	Super::SetupObject_GameThread(Arguments);
 
 	// TODO: make sure this works at runtime
 #if WITH_EDITORONLY_DATA
@@ -544,7 +544,7 @@ void UInterchangeStaticMeshFactory::BeginPreCompletedCallback(const FImportPreCo
 }
 
 
-TArray<UInterchangeStaticMeshFactory::FMeshPayload> UInterchangeStaticMeshFactory::GetMeshPayloads(const FCreateAssetParams& Arguments, const TArray<FString>& MeshUids) const
+TArray<UInterchangeStaticMeshFactory::FMeshPayload> UInterchangeStaticMeshFactory::GetMeshPayloads(const FImportAssetObjectParams& Arguments, const TArray<FString>& MeshUids) const
 {
 	TArray<FMeshPayload> Payloads;
 	Payloads.Reserve(MeshUids.Num());
@@ -627,7 +627,7 @@ TArray<UInterchangeStaticMeshFactory::FMeshPayload> UInterchangeStaticMeshFactor
 }
 
 
-bool UInterchangeStaticMeshFactory::AddConvexGeomFromVertices(const FCreateAssetParams& Arguments, const FMeshDescription& MeshDescription, const FTransform& Transform, FKAggregateGeom& AggGeom)
+bool UInterchangeStaticMeshFactory::AddConvexGeomFromVertices(const FImportAssetObjectParams& Arguments, const FMeshDescription& MeshDescription, const FTransform& Transform, FKAggregateGeom& AggGeom)
 {
 	FStaticMeshConstAttributes Attributes(MeshDescription);
 	TVertexAttributesRef<const FVector3f> VertexPositions = Attributes.GetVertexPositions();
@@ -651,7 +651,7 @@ bool UInterchangeStaticMeshFactory::AddConvexGeomFromVertices(const FCreateAsset
 }
 
 
-bool UInterchangeStaticMeshFactory::DecomposeConvexMesh(const FCreateAssetParams& Arguments, const FMeshDescription& MeshDescription, const FTransform& Transform, UBodySetup* BodySetup)
+bool UInterchangeStaticMeshFactory::DecomposeConvexMesh(const FImportAssetObjectParams& Arguments, const FMeshDescription& MeshDescription, const FTransform& Transform, UBodySetup* BodySetup)
 {
 #if WITH_EDITOR
 
@@ -781,7 +781,7 @@ static FVector3f GetTriangleNormal(const FTransform& Transform, TVertexAttribute
 }
 
 
-bool UInterchangeStaticMeshFactory::AddBoxGeomFromTris(const FCreateAssetParams& Arguments, const FMeshDescription& MeshDescription, const FTransform& Transform, FKAggregateGeom& AggGeom)
+bool UInterchangeStaticMeshFactory::AddBoxGeomFromTris(const FImportAssetObjectParams& Arguments, const FMeshDescription& MeshDescription, const FTransform& Transform, FKAggregateGeom& AggGeom)
 {
 	// Maintain an array of the planes we have encountered so far.
 	// We are expecting two instances of three unique plane orientations, one for each side of the box.
@@ -898,7 +898,7 @@ bool UInterchangeStaticMeshFactory::AddBoxGeomFromTris(const FCreateAssetParams&
 }
 
 
-bool UInterchangeStaticMeshFactory::AddSphereGeomFromVertices(const FCreateAssetParams& Arguments, const FMeshDescription& MeshDescription, const FTransform& Transform, FKAggregateGeom& AggGeom)
+bool UInterchangeStaticMeshFactory::AddSphereGeomFromVertices(const FImportAssetObjectParams& Arguments, const FMeshDescription& MeshDescription, const FTransform& Transform, FKAggregateGeom& AggGeom)
 {
 	FStaticMeshConstAttributes Attributes(MeshDescription);
 	TVertexAttributesRef<const FVector3f> VertexPositions = Attributes.GetVertexPositions();
@@ -970,7 +970,7 @@ bool UInterchangeStaticMeshFactory::AddSphereGeomFromVertices(const FCreateAsset
 }
 
 
-bool UInterchangeStaticMeshFactory::AddCapsuleGeomFromVertices(const FCreateAssetParams& Arguments, const FMeshDescription& MeshDescription, const FTransform& Transform, FKAggregateGeom& AggGeom)
+bool UInterchangeStaticMeshFactory::AddCapsuleGeomFromVertices(const FImportAssetObjectParams& Arguments, const FMeshDescription& MeshDescription, const FTransform& Transform, FKAggregateGeom& AggGeom)
 {
 	FStaticMeshConstAttributes Attributes(MeshDescription);
 	TVertexAttributesRef<const FVector3f> VertexPositions = Attributes.GetVertexPositions();
@@ -1038,7 +1038,7 @@ bool UInterchangeStaticMeshFactory::AddCapsuleGeomFromVertices(const FCreateAsse
 }
 
 
-bool UInterchangeStaticMeshFactory::ImportBoxCollision(const FCreateAssetParams& Arguments, UStaticMesh* StaticMesh, const UInterchangeStaticMeshLodDataNode* LodDataNode)
+bool UInterchangeStaticMeshFactory::ImportBoxCollision(const FImportAssetObjectParams& Arguments, UStaticMesh* StaticMesh, const UInterchangeStaticMeshLodDataNode* LodDataNode)
 {
 	using namespace UE::Interchange;
 
@@ -1086,7 +1086,7 @@ bool UInterchangeStaticMeshFactory::ImportBoxCollision(const FCreateAssetParams&
 }
 
 
-bool UInterchangeStaticMeshFactory::ImportCapsuleCollision(const FCreateAssetParams& Arguments, UStaticMesh* StaticMesh, const UInterchangeStaticMeshLodDataNode* LodDataNode)
+bool UInterchangeStaticMeshFactory::ImportCapsuleCollision(const FImportAssetObjectParams& Arguments, UStaticMesh* StaticMesh, const UInterchangeStaticMeshLodDataNode* LodDataNode)
 {
 	using namespace UE::Interchange;
 
@@ -1134,7 +1134,7 @@ bool UInterchangeStaticMeshFactory::ImportCapsuleCollision(const FCreateAssetPar
 }
 
 
-bool UInterchangeStaticMeshFactory::ImportSphereCollision(const FCreateAssetParams& Arguments, UStaticMesh* StaticMesh, const UInterchangeStaticMeshLodDataNode* LodDataNode)
+bool UInterchangeStaticMeshFactory::ImportSphereCollision(const FImportAssetObjectParams& Arguments, UStaticMesh* StaticMesh, const UInterchangeStaticMeshLodDataNode* LodDataNode)
 {
 	using namespace UE::Interchange;
 
@@ -1183,7 +1183,7 @@ bool UInterchangeStaticMeshFactory::ImportSphereCollision(const FCreateAssetPara
 }
 
 
-bool UInterchangeStaticMeshFactory::ImportConvexCollision(const FCreateAssetParams& Arguments, UStaticMesh* StaticMesh, const UInterchangeStaticMeshLodDataNode* LodDataNode)
+bool UInterchangeStaticMeshFactory::ImportConvexCollision(const FImportAssetObjectParams& Arguments, UStaticMesh* StaticMesh, const UInterchangeStaticMeshLodDataNode* LodDataNode)
 {
 	using namespace UE::Interchange;
 
@@ -1273,7 +1273,7 @@ bool UInterchangeStaticMeshFactory::ImportConvexCollision(const FCreateAssetPara
 }
 
 
-bool UInterchangeStaticMeshFactory::GenerateKDopCollision(const FCreateAssetParams& Arguments, UStaticMesh* StaticMesh)
+bool UInterchangeStaticMeshFactory::GenerateKDopCollision(const FImportAssetObjectParams& Arguments, UStaticMesh* StaticMesh)
 {
 #if WITH_EDITOR
 
@@ -1427,7 +1427,7 @@ bool UInterchangeStaticMeshFactory::GenerateKDopCollision(const FCreateAssetPara
 }
 
 
-bool UInterchangeStaticMeshFactory::ImportSockets(const FCreateAssetParams& Arguments, UStaticMesh* StaticMesh, const UInterchangeStaticMeshFactoryNode* FactoryNode)
+bool UInterchangeStaticMeshFactory::ImportSockets(const FImportAssetObjectParams& Arguments, UStaticMesh* StaticMesh, const UInterchangeStaticMeshFactoryNode* FactoryNode)
 {
 	TArray<FString> SocketUids;
 	FactoryNode->GetSocketUids(SocketUids);
diff --git a/Engine/Plugins/Interchange/Runtime/Source/Import/Private/Scene/InterchangeActorFactory.cpp b/Engine/Plugins/Interchange/Runtime/Source/Import/Private/Scene/InterchangeActorFactory.cpp
--- a/Engine/Plugins/Interchange/Runtime/Source/Import/Private/Scene/InterchangeActorFactory.cpp
+++ b/Engine/Plugins/Interchange/Runtime/Source/Import/Private/Scene/InterchangeActorFactory.cpp
@@ -14,7 +14,7 @@ UClass* UInterchangeActorFactory::GetFactoryClass() const
 	return AActor::StaticClass();
 }
 
-UObject* UInterchangeActorFactory::CreateSceneObject(const UInterchangeFactoryBase::FCreateSceneObjectsParams& CreateSceneObjectsParams)
+UObject* UInterchangeActorFactory::ImportSceneObject_GameThread(const UInterchangeFactoryBase::FImportSceneObjectsParams& CreateSceneObjectsParams)
 {
 	AActor* SpawnedActor = UE::Interchange::ActorHelper::SpawnFactoryActor(CreateSceneObjectsParams);
 
diff --git a/Engine/Plugins/Interchange/Runtime/Source/Import/Private/Scene/InterchangeActorHelper.cpp b/Engine/Plugins/Interchange/Runtime/Source/Import/Private/Scene/InterchangeActorHelper.cpp
--- a/Engine/Plugins/Interchange/Runtime/Source/Import/Private/Scene/InterchangeActorHelper.cpp
+++ b/Engine/Plugins/Interchange/Runtime/Source/Import/Private/Scene/InterchangeActorHelper.cpp
@@ -35,7 +35,7 @@ AActor* UE::Interchange::ActorHelper::GetSpawnedParentActor(const UInterchangeBa
 	return  ParentActor;
 }
 
-AActor* UE::Interchange::ActorHelper::SpawnFactoryActor(const UInterchangeFactoryBase::FCreateSceneObjectsParams& CreateSceneObjectsParams)
+AActor* UE::Interchange::ActorHelper::SpawnFactoryActor(const UInterchangeFactoryBase::FImportSceneObjectsParams& CreateSceneObjectsParams)
 {
 	FActorSpawnParameters SpawnParameters;
 	SpawnParameters.Name = FName(*CreateSceneObjectsParams.ObjectName);
diff --git a/Engine/Plugins/Interchange/Runtime/Source/Import/Private/Scene/InterchangeCineCameraActorFactory.cpp b/Engine/Plugins/Interchange/Runtime/Source/Import/Private/Scene/InterchangeCineCameraActorFactory.cpp
--- a/Engine/Plugins/Interchange/Runtime/Source/Import/Private/Scene/InterchangeCineCameraActorFactory.cpp
+++ b/Engine/Plugins/Interchange/Runtime/Source/Import/Private/Scene/InterchangeCineCameraActorFactory.cpp
@@ -19,7 +19,7 @@ UClass* UInterchangeCineCameraActorFactory::GetFactoryClass() const
 	return ACineCameraActor::StaticClass();
 }
 
-UObject* UInterchangeCineCameraActorFactory::CreateSceneObject(const UInterchangeFactoryBase::FCreateSceneObjectsParams& CreateSceneObjectsParams)
+UObject* UInterchangeCineCameraActorFactory::ImportSceneObject_GameThread(const UInterchangeFactoryBase::FImportSceneObjectsParams& CreateSceneObjectsParams)
 {
 	ACineCameraActor* SpawnedActor = Cast<ACineCameraActor>(UE::Interchange::ActorHelper::SpawnFactoryActor(CreateSceneObjectsParams));
 
diff --git a/Engine/Plugins/Interchange/Runtime/Source/Import/Private/Scene/InterchangeLightActorFactory.cpp b/Engine/Plugins/Interchange/Runtime/Source/Import/Private/Scene/InterchangeLightActorFactory.cpp
--- a/Engine/Plugins/Interchange/Runtime/Source/Import/Private/Scene/InterchangeLightActorFactory.cpp
+++ b/Engine/Plugins/Interchange/Runtime/Source/Import/Private/Scene/InterchangeLightActorFactory.cpp
@@ -18,7 +18,7 @@ UClass* UInterchangeLightActorFactory::GetFactoryClass() const
 	return ALight::StaticClass();
 }
 
-UObject* UInterchangeLightActorFactory::CreateSceneObject(const UInterchangeFactoryBase::FCreateSceneObjectsParams& CreateSceneObjectsParams)
+UObject* UInterchangeLightActorFactory::ImportSceneObject_GameThread(const UInterchangeFactoryBase::FImportSceneObjectsParams& CreateSceneObjectsParams)
 {
 	ALight * SpawnedActor = Cast<ALight>(UE::Interchange::ActorHelper::SpawnFactoryActor(CreateSceneObjectsParams));
 
diff --git a/Engine/Plugins/Interchange/Runtime/Source/Import/Private/Scene/InterchangeSceneVariantSetsFactory.cpp b/Engine/Plugins/Interchange/Runtime/Source/Import/Private/Scene/InterchangeSceneVariantSetsFactory.cpp
--- a/Engine/Plugins/Interchange/Runtime/Source/Import/Private/Scene/InterchangeSceneVariantSetsFactory.cpp
+++ b/Engine/Plugins/Interchange/Runtime/Source/Import/Private/Scene/InterchangeSceneVariantSetsFactory.cpp
@@ -294,7 +294,7 @@ UClass* UInterchangeSceneVariantSetsFactory::GetFactoryClass() const
 	return ULevelVariantSets::StaticClass();
 }
 
-UObject* UInterchangeSceneVariantSetsFactory::CreateEmptyAsset(const FCreateAssetParams& Arguments)
+UObject* UInterchangeSceneVariantSetsFactory::ImportAssetObject_GameThread(const FImportAssetObjectParams& Arguments)
 {
 #if !WITH_EDITOR || !WITH_EDITORONLY_DATA
 
@@ -350,11 +350,11 @@ UObject* UInterchangeSceneVariantSetsFactory::CreateEmptyAsset(const FCreateAsse
 		return nullptr;
 	}
 
-	return LevelVariantSets;
+	return ImportObjectSourceData(Arguments);
 #endif //else !WITH_EDITOR || !WITH_EDITORONLY_DATA
 }
 
-UObject* UInterchangeSceneVariantSetsFactory::CreateAsset(const FCreateAssetParams& Arguments)
+UObject* UInterchangeSceneVariantSetsFactory::ImportObjectSourceData(const FImportAssetObjectParams& Arguments)
 {
 #if !WITH_EDITOR || !WITH_EDITORONLY_DATA
 	// TODO: Can we import ULevelVariantSets at runtime
@@ -441,10 +441,10 @@ UObject* UInterchangeSceneVariantSetsFactory::CreateAsset(const FCreateAssetPara
 #endif //else !WITH_EDITOR || !WITH_EDITORONLY_DATA
 }
 
-void UInterchangeSceneVariantSetsFactory::BeginPreCompletedCallback(const FImportPreCompletedCallbackParams& Arguments)
+void UInterchangeSceneVariantSetsFactory::SetupObject_GameThread(const FSetupObjectParams& Arguments)
 {
 	check(IsInGameThread());
-	Super::BeginPreCompletedCallback(Arguments);
+	Super::SetupObject_GameThread(Arguments);
 
 	// TODO: Need to add an AssetImportData property to ULevelVariantSets
 }
diff --git a/Engine/Plugins/Interchange/Runtime/Source/Import/Private/Scene/InterchangeSkeletalMeshActorFactory.cpp b/Engine/Plugins/Interchange/Runtime/Source/Import/Private/Scene/InterchangeSkeletalMeshActorFactory.cpp
--- a/Engine/Plugins/Interchange/Runtime/Source/Import/Private/Scene/InterchangeSkeletalMeshActorFactory.cpp
+++ b/Engine/Plugins/Interchange/Runtime/Source/Import/Private/Scene/InterchangeSkeletalMeshActorFactory.cpp
@@ -15,7 +15,7 @@
 #include UE_INLINE_GENERATED_CPP_BY_NAME(InterchangeSkeletalMeshActorFactory)
 
 
-UObject* UInterchangeSkeletalMeshActorFactory::CreateSceneObject(const UInterchangeFactoryBase::FCreateSceneObjectsParams& CreateSceneObjectsParams)
+UObject* UInterchangeSkeletalMeshActorFactory::ImportSceneObject_GameThread(const UInterchangeFactoryBase::FImportSceneObjectsParams& CreateSceneObjectsParams)
 {
 	ASkeletalMeshActor* SpawnedActor = Cast<ASkeletalMeshActor>(UE::Interchange::ActorHelper::SpawnFactoryActor(CreateSceneObjectsParams));
 
@@ -46,7 +46,7 @@ void UInterchangeSkeletalMeshActorFactory::SetupSkeletalMeshActor(const UInterch
 	SkeletalMeshComponent->UnregisterComponent();
 }
 
-void UInterchangeSkeletalMeshActorFactory::EndPreCompletedCallback(const FImportPreCompletedCallbackParams& Arguments)
+void UInterchangeSkeletalMeshActorFactory::FinalizeObject_GameThread(const FSetupObjectParams& Arguments)
 {
 	// Set the skeletal mesh on the component in the post import callback, once the skeletal mesh has been fully imported.
 
diff --git a/Engine/Plugins/Interchange/Runtime/Source/Import/Private/Scene/InterchangeStaticMeshActorFactory.cpp b/Engine/Plugins/Interchange/Runtime/Source/Import/Private/Scene/InterchangeStaticMeshActorFactory.cpp
--- a/Engine/Plugins/Interchange/Runtime/Source/Import/Private/Scene/InterchangeStaticMeshActorFactory.cpp
+++ b/Engine/Plugins/Interchange/Runtime/Source/Import/Private/Scene/InterchangeStaticMeshActorFactory.cpp
@@ -17,7 +17,7 @@ UClass* UInterchangeStaticMeshActorFactory::GetFactoryClass() const
 	return AStaticMeshActor::StaticClass();
 }
 
-UObject* UInterchangeStaticMeshActorFactory::CreateSceneObject(const UInterchangeFactoryBase::FCreateSceneObjectsParams& CreateSceneObjectsParams)
+UObject* UInterchangeStaticMeshActorFactory::ImportSceneObject_GameThread(const UInterchangeFactoryBase::FImportSceneObjectsParams& CreateSceneObjectsParams)
 {
 	AStaticMeshActor* SpawnedActor = Cast<AStaticMeshActor>(UE::Interchange::ActorHelper::SpawnFactoryActor(CreateSceneObjectsParams));
 
diff --git a/Engine/Plugins/Interchange/Runtime/Source/Import/Private/Texture/InterchangeTextureFactory.cpp b/Engine/Plugins/Interchange/Runtime/Source/Import/Private/Texture/InterchangeTextureFactory.cpp
--- a/Engine/Plugins/Interchange/Runtime/Source/Import/Private/Texture/InterchangeTextureFactory.cpp
+++ b/Engine/Plugins/Interchange/Runtime/Source/Import/Private/Texture/InterchangeTextureFactory.cpp
@@ -1457,7 +1457,7 @@ UClass* UInterchangeTextureFactory::GetFactoryClass() const
 	return UTexture::StaticClass();
 }
 
-UObject* UInterchangeTextureFactory::CreateEmptyAsset(const FCreateAssetParams& Arguments)
+UObject* UInterchangeTextureFactory::ImportAssetObject_GameThread(const FImportAssetObjectParams& Arguments)
 {
 	using namespace  UE::Interchange::Private::InterchangeTextureFactory;
 
@@ -1534,7 +1534,7 @@ UObject* UInterchangeTextureFactory::CreateEmptyAsset(const FCreateAssetParams&
 }
 
 // The payload fetching and the heavy operations are done here
-UObject* UInterchangeTextureFactory::CreateAsset(const FCreateAssetParams& Arguments)
+UObject* UInterchangeTextureFactory::ImportAssetObject_Async(const FImportAssetObjectParams& Arguments)
 {
 	TRACE_CPUPROFILER_EVENT_SCOPE(UInterchangeTextureFactory::CreateAsset);
 
@@ -1666,7 +1666,7 @@ UObject* UInterchangeTextureFactory::CreateAsset(const FCreateAssetParams& Argum
 }
 
 /* This function is call in the completion task on the main thread, use it to call main thread post creation step for your assets*/
-void UInterchangeTextureFactory::BeginPreCompletedCallback(const FImportPreCompletedCallbackParams& Arguments)
+void UInterchangeTextureFactory::SetupObject_GameThread(const FSetupObjectParams& Arguments)
 {
 	TRACE_CPUPROFILER_EVENT_SCOPE(UInterchangeTextureFactory::BeginPreCompletedCallback);
 
@@ -1730,7 +1730,7 @@ void UInterchangeTextureFactory::BeginPreCompletedCallback(const FImportPreCompl
 		}
 	}
 
-	Super::BeginPreCompletedCallback(Arguments);
+	Super::SetupObject_GameThread(Arguments);
 
 	//TODO make sure this work at runtime
 #if WITH_EDITORONLY_DATA
diff --git a/Engine/Plugins/Interchange/Runtime/Source/Import/Public/Animation/InterchangeAnimSequenceFactory.h b/Engine/Plugins/Interchange/Runtime/Source/Import/Public/Animation/InterchangeAnimSequenceFactory.h
--- a/Engine/Plugins/Interchange/Runtime/Source/Import/Public/Animation/InterchangeAnimSequenceFactory.h
+++ b/Engine/Plugins/Interchange/Runtime/Source/Import/Public/Animation/InterchangeAnimSequenceFactory.h
@@ -23,20 +23,15 @@ public:
 
 	virtual UClass* GetFactoryClass() const override;
 	virtual EInterchangeFactoryAssetType GetFactoryAssetType() override { return EInterchangeFactoryAssetType::Animations; }
-	virtual UObject* CreateEmptyAsset(const FCreateAssetParams& Arguments) override;
-	virtual UObject* CreateAsset(const FCreateAssetParams& Arguments) override;
-	virtual void BeginPreCompletedCallback(const FImportPreCompletedCallbackParams& Arguments) override;
-	//virtual void EndPreCompletedCallback(const FImportPreCompletedCallbackParams& Arguments) const override;
-	virtual bool CanExecuteOnAnyThread() const override
-	{
-		//Currently we cannot use the anim sequence controller outside of the game thread
-		return false;
-	}
+	virtual UObject* ImportAssetObject_GameThread(const FImportAssetObjectParams& Arguments) override;
+	virtual void SetupObject_GameThread(const FSetupObjectParams& Arguments) override;
+
 	virtual bool GetSourceFilenames(const UObject* Object, TArray<FString>& OutSourceFilenames) const override;
 	virtual bool SetSourceFilename(const UObject* Object, const FString& SourceFilename, int32 SourceIndex) const override;
 
 private:
-	bool IsBoneTrackAnimationValid(const UInterchangeAnimSequenceFactoryNode* AnimSequenceFactoryNode, const FCreateAssetParams& Arguments);
+	UObject* ImportObjectSourceData(const FImportAssetObjectParams& Arguments);
+	bool IsBoneTrackAnimationValid(const UInterchangeAnimSequenceFactoryNode* AnimSequenceFactoryNode, const FImportAssetObjectParams& Arguments);
 
 	// Interchange factory base interface end
 	//////////////////////////////////////////////////////////////////////////
diff --git a/Engine/Plugins/Interchange/Runtime/Source/Import/Public/Animation/InterchangeAnimationTrackSetFactory.h b/Engine/Plugins/Interchange/Runtime/Source/Import/Public/Animation/InterchangeAnimationTrackSetFactory.h
--- a/Engine/Plugins/Interchange/Runtime/Source/Import/Public/Animation/InterchangeAnimationTrackSetFactory.h
+++ b/Engine/Plugins/Interchange/Runtime/Source/Import/Public/Animation/InterchangeAnimationTrackSetFactory.h
@@ -20,21 +20,17 @@ public:
 
 	virtual UClass* GetFactoryClass() const override;
 	virtual EInterchangeFactoryAssetType GetFactoryAssetType() override { return EInterchangeFactoryAssetType::Animations; }
-	virtual UObject* CreateEmptyAsset(const FCreateAssetParams& Arguments) override;
-	virtual UObject* CreateAsset(const FCreateAssetParams& Arguments) override;
-	virtual void BeginPreCompletedCallback(const FImportPreCompletedCallbackParams& Arguments) override;
-	virtual void EndPreCompletedCallback(const FImportPreCompletedCallbackParams& Arguments) override;
-	virtual bool CanExecuteOnAnyThread() const override
-	{
-		//Currently we cannot use the anim sequence controller outside of the game thread
-		return false;
-	}
+	virtual UObject* ImportAssetObject_GameThread(const FImportAssetObjectParams& Arguments) override;
+	virtual void SetupObject_GameThread(const FSetupObjectParams& Arguments) override;
+	virtual void FinalizeObject_GameThread(const FSetupObjectParams& Arguments) override;
 
 	// Interchange factory base interface end
 	//////////////////////////////////////////////////////////////////////////
 
 private:
 	const UInterchangeTranslatorBase* Translator = nullptr;
+
+	UObject* ImportObjectSourceData(const FImportAssetObjectParams& Arguments);
 };
 
 
diff --git a/Engine/Plugins/Interchange/Runtime/Source/Import/Public/Material/InterchangeMaterialFactory.h b/Engine/Plugins/Interchange/Runtime/Source/Import/Public/Material/InterchangeMaterialFactory.h
--- a/Engine/Plugins/Interchange/Runtime/Source/Import/Public/Material/InterchangeMaterialFactory.h
+++ b/Engine/Plugins/Interchange/Runtime/Source/Import/Public/Material/InterchangeMaterialFactory.h
@@ -28,9 +28,9 @@ public:
 
 	virtual UClass* GetFactoryClass() const override;
 	virtual EInterchangeFactoryAssetType GetFactoryAssetType() override { return EInterchangeFactoryAssetType::Materials; }
-	virtual UObject* CreateEmptyAsset(const FCreateAssetParams& Arguments) override;
-	virtual UObject* CreateAsset(const FCreateAssetParams& Arguments) override;
-	virtual void BeginPreCompletedCallback(const FImportPreCompletedCallbackParams& Arguments) override;
+	virtual UObject* ImportAssetObject_GameThread(const FImportAssetObjectParams& Arguments) override;
+	virtual UObject* ImportAssetObject_Async(const FImportAssetObjectParams& Arguments) override;
+	virtual void SetupObject_GameThread(const FSetupObjectParams& Arguments) override;
 	virtual bool GetSourceFilenames(const UObject* Object, TArray<FString>& OutSourceFilenames) const override;
 	virtual bool SetSourceFilename(const UObject* Object, const FString& SourceFilename, int32 SourceIndex) const override;
 
@@ -38,7 +38,7 @@ public:
 	//////////////////////////////////////////////////////////////////////////
 private:
 #if WITH_EDITOR
-	void SetupMaterial(UMaterial* Material, const FCreateAssetParams& Arguments, const UInterchangeBaseMaterialFactoryNode* MaterialFactoryNode);
+	void SetupMaterial(UMaterial* Material, const FImportAssetObjectParams& Arguments, const UInterchangeBaseMaterialFactoryNode* MaterialFactoryNode);
 #endif // #if WITH_EDITOR
 
 	void SetupMaterialInstance(UMaterialInstance* MaterialInstance, const UInterchangeBaseNodeContainer* NodeContainer, const UInterchangeBaseMaterialFactoryNode* MaterialFactoryNode);
@@ -55,16 +55,16 @@ public:
 
 	virtual UClass* GetFactoryClass() const override;
 	virtual EInterchangeFactoryAssetType GetFactoryAssetType() override { return EInterchangeFactoryAssetType::Materials; }
-	virtual UObject* CreateEmptyAsset(const FCreateAssetParams& Arguments) override;
-	virtual UObject* CreateAsset(const FCreateAssetParams& Arguments) override;
-	virtual void BeginPreCompletedCallback(const FImportPreCompletedCallbackParams& Arguments) override;
+	virtual UObject* ImportAssetObject_GameThread(const FImportAssetObjectParams& Arguments) override;
+	virtual UObject* ImportAssetObject_Async(const FImportAssetObjectParams& Arguments) override;
+	virtual void SetupObject_GameThread(const FSetupObjectParams& Arguments) override;
 
 	// Interchange factory base interface end
 	//////////////////////////////////////////////////////////////////////////
 
 private:
 #if WITH_EDITOR
-	void SetupMaterial(class UMaterialFunction* Material, const FCreateAssetParams& Arguments, const class UInterchangeMaterialFunctionFactoryNode* MaterialFactoryNode);
+	void SetupMaterial(class UMaterialFunction* Material, const FImportAssetObjectParams& Arguments, const class UInterchangeMaterialFunctionFactoryNode* MaterialFactoryNode);
 #endif
 };
 
diff --git a/Engine/Plugins/Interchange/Runtime/Source/Import/Public/Mesh/InterchangePhysicsAssetFactory.h b/Engine/Plugins/Interchange/Runtime/Source/Import/Public/Mesh/InterchangePhysicsAssetFactory.h
--- a/Engine/Plugins/Interchange/Runtime/Source/Import/Public/Mesh/InterchangePhysicsAssetFactory.h
+++ b/Engine/Plugins/Interchange/Runtime/Source/Import/Public/Mesh/InterchangePhysicsAssetFactory.h
@@ -20,10 +20,8 @@ public:
 
 	virtual UClass* GetFactoryClass() const override;
 	virtual EInterchangeFactoryAssetType GetFactoryAssetType() override { return EInterchangeFactoryAssetType::Physics; }
-	virtual UObject* CreateEmptyAsset(const FCreateAssetParams& Arguments) override;
-	virtual UObject* CreateAsset(const FCreateAssetParams& Arguments) override;
-	//virtual void PreImportPreCompletedCallback(const FImportPreCompletedCallbackParams& Arguments) const override;
-	//virtual void PostImportPreCompletedCallback(const FImportPreCompletedCallbackParams& Arguments) const override;
+	virtual UObject* ImportAssetObject_GameThread(const FImportAssetObjectParams& Arguments) override;
+	virtual UObject* ImportAssetObject_Async(const FImportAssetObjectParams& Arguments) override;
 
 	// Interchange factory base interface end
 	//////////////////////////////////////////////////////////////////////////
diff --git a/Engine/Plugins/Interchange/Runtime/Source/Import/Public/Mesh/InterchangeSkeletalMeshFactory.h b/Engine/Plugins/Interchange/Runtime/Source/Import/Public/Mesh/InterchangeSkeletalMeshFactory.h
--- a/Engine/Plugins/Interchange/Runtime/Source/Import/Public/Mesh/InterchangeSkeletalMeshFactory.h
+++ b/Engine/Plugins/Interchange/Runtime/Source/Import/Public/Mesh/InterchangeSkeletalMeshFactory.h
@@ -22,11 +22,10 @@ public:
 
 	virtual UClass* GetFactoryClass() const override;
 	virtual EInterchangeFactoryAssetType GetFactoryAssetType() override { return EInterchangeFactoryAssetType::Meshes; }
-	virtual UObject* CreateEmptyAsset(const FCreateAssetParams& Arguments) override;
-	virtual UObject* CreateAsset(const FCreateAssetParams& Arguments) override;
+	virtual UObject* ImportAssetObject_GameThread(const FImportAssetObjectParams& Arguments) override;
+	virtual UObject* ImportAssetObject_Async(const FImportAssetObjectParams& Arguments) override;
 	virtual void Cancel() override;
-	virtual void BeginPreCompletedCallback(const FImportPreCompletedCallbackParams& Arguments) override;
-	//virtual void EndPreCompletedCallback(const FImportPreCompletedCallbackParams& Arguments) const override;
+	virtual void SetupObject_GameThread(const FSetupObjectParams& Arguments) override;
 	virtual bool GetSourceFilenames(const UObject* Object, TArray<FString>& OutSourceFilenames) const override;
 	virtual bool SetSourceFilename(const UObject* Object, const FString& SourceFilename, int32 SourceIndex) const override;
 	virtual bool SetReimportSourceIndex(const UObject* Object, int32 SourceIndex) const override;
diff --git a/Engine/Plugins/Interchange/Runtime/Source/Import/Public/Mesh/InterchangeSkeletonFactory.h b/Engine/Plugins/Interchange/Runtime/Source/Import/Public/Mesh/InterchangeSkeletonFactory.h
--- a/Engine/Plugins/Interchange/Runtime/Source/Import/Public/Mesh/InterchangeSkeletonFactory.h
+++ b/Engine/Plugins/Interchange/Runtime/Source/Import/Public/Mesh/InterchangeSkeletonFactory.h
@@ -20,9 +20,8 @@ public:
 
 	virtual UClass* GetFactoryClass() const override;
 	virtual EInterchangeFactoryAssetType GetFactoryAssetType() override { return EInterchangeFactoryAssetType::Meshes; }
-	virtual UObject* CreateEmptyAsset(const FCreateAssetParams& Arguments) override;
-	virtual UObject* CreateAsset(const FCreateAssetParams& Arguments) override;
-	//virtual void PreImportPreCompletedCallback(const FImportPreCompletedCallbackParams& Arguments) const override;
+	virtual UObject* ImportAssetObject_GameThread(const FImportAssetObjectParams& Arguments) override;
+	virtual UObject* ImportAssetObject_Async(const FImportAssetObjectParams& Arguments) override;
 
 	// Interchange factory base interface end
 	//////////////////////////////////////////////////////////////////////////
diff --git a/Engine/Plugins/Interchange/Runtime/Source/Import/Public/Mesh/InterchangeStaticMeshFactory.h b/Engine/Plugins/Interchange/Runtime/Source/Import/Public/Mesh/InterchangeStaticMeshFactory.h
--- a/Engine/Plugins/Interchange/Runtime/Source/Import/Public/Mesh/InterchangeStaticMeshFactory.h
+++ b/Engine/Plugins/Interchange/Runtime/Source/Import/Public/Mesh/InterchangeStaticMeshFactory.h
@@ -30,9 +30,9 @@ public:
 
 	virtual UClass* GetFactoryClass() const override;
 	virtual EInterchangeFactoryAssetType GetFactoryAssetType() override { return EInterchangeFactoryAssetType::Meshes; }
-	virtual UObject* CreateEmptyAsset(const FCreateAssetParams& Arguments) override;
-	virtual UObject* CreateAsset(const FCreateAssetParams& Arguments) override;
-	virtual void BeginPreCompletedCallback(const FImportPreCompletedCallbackParams& Arguments) override;
+	virtual UObject* ImportAssetObject_GameThread(const FImportAssetObjectParams& Arguments) override;
+	virtual UObject* ImportAssetObject_Async(const FImportAssetObjectParams& Arguments) override;
+	virtual void SetupObject_GameThread(const FSetupObjectParams& Arguments) override;
 	virtual bool GetSourceFilenames(const UObject* Object, TArray<FString>& OutSourceFilenames) const override;
 	virtual bool SetSourceFilename(const UObject* Object, const FString& SourceFilename, int32 SourceIndex) const override;
 
@@ -48,20 +48,20 @@ private:
 		FTransform Transform = FTransform::Identity;
 	};
 
-	TArray<FMeshPayload> GetMeshPayloads(const FCreateAssetParams& Arguments, const TArray<FString>& MeshUids) const;
+	TArray<FMeshPayload> GetMeshPayloads(const FImportAssetObjectParams& Arguments, const TArray<FString>& MeshUids) const;
 
-	bool AddConvexGeomFromVertices(const FCreateAssetParams& Arguments, const FMeshDescription& MeshDescription, const FTransform& Transform, FKAggregateGeom& AggGeom);
-	bool DecomposeConvexMesh(const FCreateAssetParams& Arguments, const FMeshDescription& MeshDescription, const FTransform& Transform, UBodySetup* BodySetup);
-	bool AddBoxGeomFromTris(const FCreateAssetParams& Arguments, const FMeshDescription& MeshDescription, const FTransform& Transform, FKAggregateGeom& AggGeom);
-	bool AddSphereGeomFromVertices(const FCreateAssetParams& Arguments, const FMeshDescription& MeshDescription, const FTransform& Transform, FKAggregateGeom& AggGeom);
-	bool AddCapsuleGeomFromVertices(const FCreateAssetParams& Arguments, const FMeshDescription& MeshDescription, const FTransform& Transform, FKAggregateGeom& AggGeom);
+	bool AddConvexGeomFromVertices(const FImportAssetObjectParams& Arguments, const FMeshDescription& MeshDescription, const FTransform& Transform, FKAggregateGeom& AggGeom);
+	bool DecomposeConvexMesh(const FImportAssetObjectParams& Arguments, const FMeshDescription& MeshDescription, const FTransform& Transform, UBodySetup* BodySetup);
+	bool AddBoxGeomFromTris(const FImportAssetObjectParams& Arguments, const FMeshDescription& MeshDescription, const FTransform& Transform, FKAggregateGeom& AggGeom);
+	bool AddSphereGeomFromVertices(const FImportAssetObjectParams& Arguments, const FMeshDescription& MeshDescription, const FTransform& Transform, FKAggregateGeom& AggGeom);
+	bool AddCapsuleGeomFromVertices(const FImportAssetObjectParams& Arguments, const FMeshDescription& MeshDescription, const FTransform& Transform, FKAggregateGeom& AggGeom);
 
-	bool ImportBoxCollision(const FCreateAssetParams& Arguments, UStaticMesh* StaticMesh, const UInterchangeStaticMeshLodDataNode* LodDataNode);
-	bool ImportCapsuleCollision(const FCreateAssetParams& Arguments, UStaticMesh* StaticMesh, const UInterchangeStaticMeshLodDataNode* LodDataNode);
-	bool ImportSphereCollision(const FCreateAssetParams& Arguments, UStaticMesh* StaticMesh, const UInterchangeStaticMeshLodDataNode* LodDataNode);
-	bool ImportConvexCollision(const FCreateAssetParams& Arguments, UStaticMesh* StaticMesh, const UInterchangeStaticMeshLodDataNode* LodDataNode);
-	bool GenerateKDopCollision(const FCreateAssetParams& Arguments, UStaticMesh* StaticMesh);
-	bool ImportSockets(const FCreateAssetParams& Arguments, UStaticMesh* StaticMesh, const UInterchangeStaticMeshFactoryNode* FactoryNode);
+	bool ImportBoxCollision(const FImportAssetObjectParams& Arguments, UStaticMesh* StaticMesh, const UInterchangeStaticMeshLodDataNode* LodDataNode);
+	bool ImportCapsuleCollision(const FImportAssetObjectParams& Arguments, UStaticMesh* StaticMesh, const UInterchangeStaticMeshLodDataNode* LodDataNode);
+	bool ImportSphereCollision(const FImportAssetObjectParams& Arguments, UStaticMesh* StaticMesh, const UInterchangeStaticMeshLodDataNode* LodDataNode);
+	bool ImportConvexCollision(const FImportAssetObjectParams& Arguments, UStaticMesh* StaticMesh, const UInterchangeStaticMeshLodDataNode* LodDataNode);
+	bool GenerateKDopCollision(const FImportAssetObjectParams& Arguments, UStaticMesh* StaticMesh);
+	bool ImportSockets(const FImportAssetObjectParams& Arguments, UStaticMesh* StaticMesh, const UInterchangeStaticMeshFactoryNode* FactoryNode);
 
 	void CommitMeshDescriptions(UStaticMesh& StaticMesh, TArray<FMeshDescription>&& LodMeshDescriptions);
 
diff --git a/Engine/Plugins/Interchange/Runtime/Source/Import/Public/Scene/InterchangeActorFactory.h b/Engine/Plugins/Interchange/Runtime/Source/Import/Public/Scene/InterchangeActorFactory.h
--- a/Engine/Plugins/Interchange/Runtime/Source/Import/Public/Scene/InterchangeActorFactory.h
+++ b/Engine/Plugins/Interchange/Runtime/Source/Import/Public/Scene/InterchangeActorFactory.h
@@ -20,12 +20,7 @@ public:
 
 	virtual UClass* GetFactoryClass() const override;
 
-	virtual UObject* CreateSceneObject(const UInterchangeFactoryBase::FCreateSceneObjectsParams& CreateSceneObjectsParams) override;
-
-	virtual bool CanExecuteOnAnyThread() const override
-	{
-		return false;
-	}
+	virtual UObject* ImportSceneObject_GameThread(const UInterchangeFactoryBase::FImportSceneObjectsParams& CreateSceneObjectsParams) override;
 
 	// Interchange factory base interface end
 	//////////////////////////////////////////////////////////////////////////
diff --git a/Engine/Plugins/Interchange/Runtime/Source/Import/Public/Scene/InterchangeActorHelper.h b/Engine/Plugins/Interchange/Runtime/Source/Import/Public/Scene/InterchangeActorHelper.h
--- a/Engine/Plugins/Interchange/Runtime/Source/Import/Public/Scene/InterchangeActorHelper.h
+++ b/Engine/Plugins/Interchange/Runtime/Source/Import/Public/Scene/InterchangeActorHelper.h
@@ -21,7 +21,7 @@ namespace UE::Interchange::ActorHelper
 	 * Spawns and return a AAActor using the given SceneObjectsParam.
 	 * The actor's parent will be determine by the node hierarchy and its type by the factory node's ObjectClass.
 	 */
-	INTERCHANGEIMPORT_API AActor* SpawnFactoryActor(const UInterchangeFactoryBase::FCreateSceneObjectsParams& CreateSceneObjectsParams);
+	INTERCHANGEIMPORT_API AActor* SpawnFactoryActor(const UInterchangeFactoryBase::FImportSceneObjectsParams& CreateSceneObjectsParams);
 	
 	/**
 	 * Returns the factory node of the asset instanced by ActorFactoryNode.
diff --git a/Engine/Plugins/Interchange/Runtime/Source/Import/Public/Scene/InterchangeCineCameraActorFactory.h b/Engine/Plugins/Interchange/Runtime/Source/Import/Public/Scene/InterchangeCineCameraActorFactory.h
--- a/Engine/Plugins/Interchange/Runtime/Source/Import/Public/Scene/InterchangeCineCameraActorFactory.h
+++ b/Engine/Plugins/Interchange/Runtime/Source/Import/Public/Scene/InterchangeCineCameraActorFactory.h
@@ -20,12 +20,7 @@ public:
 
 	virtual UClass* GetFactoryClass() const override;
 
-	virtual UObject* CreateSceneObject(const UInterchangeFactoryBase::FCreateSceneObjectsParams& CreateSceneObjectsParams) override;
-
-	virtual bool CanExecuteOnAnyThread() const override
-	{
-		return false; 
-	}
+	virtual UObject* ImportSceneObject_GameThread(const UInterchangeFactoryBase::FImportSceneObjectsParams& CreateSceneObjectsParams) override;
 
 	// Interchange factory base interface end
 	//////////////////////////////////////////////////////////////////////////
diff --git a/Engine/Plugins/Interchange/Runtime/Source/Import/Public/Scene/InterchangeLightActorFactory.h b/Engine/Plugins/Interchange/Runtime/Source/Import/Public/Scene/InterchangeLightActorFactory.h
--- a/Engine/Plugins/Interchange/Runtime/Source/Import/Public/Scene/InterchangeLightActorFactory.h
+++ b/Engine/Plugins/Interchange/Runtime/Source/Import/Public/Scene/InterchangeLightActorFactory.h
@@ -17,12 +17,7 @@ public:
 
 	virtual UClass* GetFactoryClass() const override;
 
-	virtual UObject* CreateSceneObject(const UInterchangeFactoryBase::FCreateSceneObjectsParams& CreateSceneObjectsParams) override;
-
-	virtual bool CanExecuteOnAnyThread() const override
-	{
-		return false;
-	}
+	virtual UObject* ImportSceneObject_GameThread(const UInterchangeFactoryBase::FImportSceneObjectsParams& CreateSceneObjectsParams) override;
 
 	// Interchange factory base interface end
 	//////////////////////////////////////////////////////////////////////////
diff --git a/Engine/Plugins/Interchange/Runtime/Source/Import/Public/Scene/InterchangeSceneVariantSetsFactory.h b/Engine/Plugins/Interchange/Runtime/Source/Import/Public/Scene/InterchangeSceneVariantSetsFactory.h
--- a/Engine/Plugins/Interchange/Runtime/Source/Import/Public/Scene/InterchangeSceneVariantSetsFactory.h
+++ b/Engine/Plugins/Interchange/Runtime/Source/Import/Public/Scene/InterchangeSceneVariantSetsFactory.h
@@ -23,19 +23,15 @@ public:
 	// Interchange factory base interface begin
 
 	virtual UClass* GetFactoryClass() const override;
-	virtual UObject* CreateEmptyAsset(const FCreateAssetParams& Arguments) override;
-	virtual UObject* CreateAsset(const FCreateAssetParams& Arguments) override;
-	virtual void BeginPreCompletedCallback(const FImportPreCompletedCallbackParams& Arguments) override;
-	virtual bool CanExecuteOnAnyThread() const override
-	{
-		//Currently we cannot use the anim sequence controller outside of the game thread
-		return false;
-	}
+	virtual UObject* ImportAssetObject_GameThread(const FImportAssetObjectParams& Arguments) override;
+	
+	virtual void SetupObject_GameThread(const FSetupObjectParams& Arguments) override;
 
 	// Interchange factory base interface end
 	//////////////////////////////////////////////////////////////////////////
 
 private:
+	UObject* ImportObjectSourceData(const FImportAssetObjectParams& Arguments);
 	const UInterchangeTranslatorBase* Translator = nullptr;
 };
 
diff --git a/Engine/Plugins/Interchange/Runtime/Source/Import/Public/Scene/InterchangeSkeletalMeshActorFactory.h b/Engine/Plugins/Interchange/Runtime/Source/Import/Public/Scene/InterchangeSkeletalMeshActorFactory.h
--- a/Engine/Plugins/Interchange/Runtime/Source/Import/Public/Scene/InterchangeSkeletalMeshActorFactory.h
+++ b/Engine/Plugins/Interchange/Runtime/Source/Import/Public/Scene/InterchangeSkeletalMeshActorFactory.h
@@ -24,14 +24,9 @@ public:
 
 	virtual UClass* GetFactoryClass() const override;
 
-	virtual UObject* CreateSceneObject(const UInterchangeFactoryBase::FCreateSceneObjectsParams& CreateSceneObjectsParams) override;
+	virtual UObject* ImportSceneObject_GameThread(const UInterchangeFactoryBase::FImportSceneObjectsParams& CreateSceneObjectsParams) override;
 
-	virtual bool CanExecuteOnAnyThread() const override
-	{
-		return false;
-	}
-
-	virtual void EndPreCompletedCallback(const FImportPreCompletedCallbackParams& Arguments) override;
+	virtual void FinalizeObject_GameThread(const FSetupObjectParams& Arguments) override;
 
 	// Interchange factory base interface end
 	//////////////////////////////////////////////////////////////////////////
diff --git a/Engine/Plugins/Interchange/Runtime/Source/Import/Public/Scene/InterchangeStaticMeshActorFactory.h b/Engine/Plugins/Interchange/Runtime/Source/Import/Public/Scene/InterchangeStaticMeshActorFactory.h
--- a/Engine/Plugins/Interchange/Runtime/Source/Import/Public/Scene/InterchangeStaticMeshActorFactory.h
+++ b/Engine/Plugins/Interchange/Runtime/Source/Import/Public/Scene/InterchangeStaticMeshActorFactory.h
@@ -24,12 +24,8 @@ public:
 
 	virtual UClass* GetFactoryClass() const override;
 
-	virtual UObject* CreateSceneObject(const UInterchangeFactoryBase::FCreateSceneObjectsParams& CreateSceneObjectsParams) override;
+	virtual UObject* ImportSceneObject_GameThread(const UInterchangeFactoryBase::FImportSceneObjectsParams& CreateSceneObjectsParams) override;
 
-	virtual bool CanExecuteOnAnyThread() const override
-	{
-		return false; 
-	}
 
 	// Interchange factory base interface end
 	//////////////////////////////////////////////////////////////////////////
diff --git a/Engine/Plugins/Interchange/Runtime/Source/Import/Public/Texture/InterchangeTextureFactory.h b/Engine/Plugins/Interchange/Runtime/Source/Import/Public/Texture/InterchangeTextureFactory.h
--- a/Engine/Plugins/Interchange/Runtime/Source/Import/Public/Texture/InterchangeTextureFactory.h
+++ b/Engine/Plugins/Interchange/Runtime/Source/Import/Public/Texture/InterchangeTextureFactory.h
@@ -56,9 +56,9 @@ public:
 
 	virtual UClass* GetFactoryClass() const override;
 	virtual EInterchangeFactoryAssetType GetFactoryAssetType() override { return EInterchangeFactoryAssetType::Textures; }
-	virtual UObject* CreateEmptyAsset(const FCreateAssetParams& Arguments) override;
-	virtual UObject* CreateAsset(const FCreateAssetParams& Arguments) override;
-	virtual void BeginPreCompletedCallback(const FImportPreCompletedCallbackParams& Arguments) override;
+	virtual UObject* ImportAssetObject_GameThread(const FImportAssetObjectParams& Arguments) override;
+	virtual UObject* ImportAssetObject_Async(const FImportAssetObjectParams& Arguments) override;
+	virtual void SetupObject_GameThread(const FSetupObjectParams& Arguments) override;
 	virtual bool GetSourceFilenames(const UObject* Object, TArray<FString>& OutSourceFilenames) const override;
 	virtual bool SetSourceFilename(const UObject* Object, const FString& SourceFilename, int32 SourceIndex) const override;
 
diff --git a/Engine/Source/Runtime/Interchange/Core/Public/InterchangeFactoryBase.h b/Engine/Source/Runtime/Interchange/Core/Public/InterchangeFactoryBase.h
--- a/Engine/Source/Runtime/Interchange/Core/Public/InterchangeFactoryBase.h
+++ b/Engine/Source/Runtime/Interchange/Core/Public/InterchangeFactoryBase.h
@@ -4,15 +4,16 @@
 
 #include "CoreMinimal.h"
 #include "InterchangeResultsContainer.h"
+#include "Nodes/InterchangeFactoryBaseNode.h"
 #include "UObject/Class.h"
 #include "UObject/Object.h"
 #include "UObject/ObjectMacros.h"
+#include "UObject/SoftObjectPath.h"
 
 #include "InterchangeFactoryBase.generated.h"
 
 class ULevel;
 class UInterchangeBaseNodeContainer;
-class UInterchangeFactoryBaseNode;
 class UInterchangePipelineBase;
 class UInterchangeSourceData;
 class UInterchangeTranslatorBase;
@@ -28,6 +29,17 @@ enum class EInterchangeFactoryAssetType : uint8
 	Physics
 };
 
+/**
+ * UObject asset factory implementation:
+ * 1. ImportAssetObject_GameThread - Create the asset UObject, and you can also import source data and setup properties synchronously.
+ * 2. ImportAssetObject_Async - Import source data and setup properties asynchronously.
+ * 3. SetupObject_GameThread - Do any UObject setup required before the build (before PostEditChange), the UObject dependencies should exist and have all the source data and properties imported.
+ * 4. FinalizeObject_GameThread - Do any final UObject setup after the build (after PostEditChange)
+ * 
+ * Actor factory implementation
+ * 1. ImportSceneObject_GameThread - Create an actor in a level.
+ */
+
 UCLASS(BlueprintType, Blueprintable, Abstract)
 class INTERCHANGECORE_API UInterchangeFactoryBase : public UObject
 {
@@ -51,7 +63,7 @@ public:
 	/**
 	 * Parameters to pass to CreateAsset function
 	 */
-	struct FCreateAssetParams
+	struct FImportAssetObjectParams
 	{
 		/** The package where to create the asset, if null it will put it in the transient package */
 		UObject* Parent = nullptr;
@@ -78,38 +90,55 @@ public:
 		UObject* ReimportObject = nullptr;
 	};
 
+	UE_DEPRECATED(5.2, "This function is replaced by ImportAssetObject_GameThread.")
+	virtual UObject* CreateEmptyAsset(const FImportAssetObjectParams& Arguments)
+	{
+		return nullptr;
+	}
+	
 	/**
-	 * Create an empty asset from a Node data, this function will be call in the main thread in same time has we create the package.
-	 * The asset create here must have the internal async flag set because the object can be setup in an asynchronous thread and must be consider like an async object until
-	 * The completion task on the object is finish.
+	 * Create the asset UObject. You can optionally import source data (payload data) and configure the properties synchronously.
 	 *
-	 * @param Arguments - The structure containing all necessary arguments, see the structure definition for the documentation
-	 * @return the created UObject or nullptr if there is an error, see LOG for any error detail.
+	 * @param Arguments - The structure containing all necessary arguments to import the UObject and the package, see the structure definition for the documentation.
+	 * @return the imported UObject or nullptr if there is an error.
 	 *
-	 * @Note override function should verify the AssetNode in Arguments match the expected type for the factory.
+	 * @Note Mandatory to override this function to create the asset UObject. Not mandatory for level actor, use CreateSceneObject instead.
 	 */
-	virtual UObject* CreateEmptyAsset(const FCreateAssetParams& Arguments)
+	virtual UObject* ImportAssetObject_GameThread(const FImportAssetObjectParams& Arguments)
 	{
-		return nullptr;
+		PRAGMA_DISABLE_DEPRECATION_WARNINGS
+		return CreateEmptyAsset(Arguments);
+		PRAGMA_ENABLE_DEPRECATION_WARNINGS
 	}
 
+	UE_DEPRECATED(5.2, "This function is replaced by ImportAssetObject_Async.")
+	virtual UObject* CreateAsset(const FImportAssetObjectParams& Arguments)
+	{
+		//By default simply return the UObject created by ImportAssetObject_GameThread that was store in the asset node.
+		FSoftObjectPath ReferenceObject;
+		if (Arguments.AssetNode->GetCustomReferenceObject(ReferenceObject))
+		{
+			return ReferenceObject.TryLoad();
+		}
+		return nullptr;
+	}
 	/**
-	 * Create an asset from a Node data, this function must be multithread safe.
+	 * Override it to setup the UObject source data (payload data) and configure the properties asynchronously. Do not create the asset UObject asynchronously.
 	 *
-	 * @param Arguments - The structure containing all necessary arguments, see the structure definition for the documentation
-	 * @return the created UObject or nullptr if there is an error, see LOG for any error detail.
-	 *
-	 * @Note override function should verify the AssetNode in Arguments match the expected type for the factory.
+	 * @param Arguments - The structure containing all necessary data to imported the UObject, see the structure definition for the documentation.
+	 * @return the imported UObject or nullptr if there is an error.
 	 */
-	virtual UObject* CreateAsset(const FCreateAssetParams& Arguments)
+	virtual UObject* ImportAssetObject_Async(const FImportAssetObjectParams& Arguments)
 	{
-		return nullptr;
+		PRAGMA_DISABLE_DEPRECATION_WARNINGS
+		return CreateAsset(Arguments);
+		PRAGMA_ENABLE_DEPRECATION_WARNINGS
 	}
 
 	/**
 	 * Parameters to pass to SpawnActor function
 	 */
-	struct FCreateSceneObjectsParams
+	struct FImportSceneObjectsParams
 	{
 		/** The level in which to create the scene objects */
 		ULevel* Level = nullptr;
@@ -124,21 +153,23 @@ public:
 		const UInterchangeBaseNodeContainer* NodeContainer = nullptr;
 	};
 
+	UE_DEPRECATED(5.2, "This function is replaced by ImportSceneObject_GameThread.")
+	virtual UObject* CreateSceneObject(const FImportSceneObjectsParams& Arguments)
+	{
+		return nullptr;
+	}
+
 	/**
 	 * Creates the scene object from a Scene Node data.
 	 *
 	 * @param Arguments - The structure containing all necessary arguments, see the structure definition for the documentation.
 	 * @return The scene object or nullptr if the operation was unsuccessful.
 	 */
-	virtual UObject* CreateSceneObject(const FCreateSceneObjectsParams& Arguments)
+	virtual UObject* ImportSceneObject_GameThread(const FImportSceneObjectsParams& Arguments)
 	{
-		return nullptr;
-	}
-
-	/** Return true if the factory can create the asset asynchronously on any thread, false if it need to be on the main thread */
-	virtual bool CanExecuteOnAnyThread() const
-	{
-		return true;
+		PRAGMA_DISABLE_DEPRECATION_WARNINGS
+		return CreateSceneObject(Arguments);
+		PRAGMA_ENABLE_DEPRECATION_WARNINGS
 	}
 
 	/** Call when the user cancel the operation. */
@@ -147,7 +178,7 @@ public:
 	/**
 	 * Parameters to pass to CreateAsset function
 	 */
-	struct FImportPreCompletedCallbackParams
+	struct FSetupObjectParams
 	{
 		/** The source data, mainly use to set the asset import data file. TODO: we have to refactor UAssetImportData, the source data should be the base class for this now */
 		const UInterchangeSourceData* SourceData = nullptr;
@@ -165,36 +196,38 @@ public:
 ;
 	};
 
-	UE_DEPRECATED(5.2, "This function is replace by BeginPreCompletedCallback.")
-	virtual void PreImportPreCompletedCallback(const FImportPreCompletedCallbackParams& Arguments)
+	UE_DEPRECATED(5.2, "This function is replaced by SetupObject_GameThread.")
+	virtual void PreImportPreCompletedCallback(const FSetupObjectParams& Arguments)
 	{
-		BeginPreCompletedCallback(Arguments);
+		check(IsInGameThread());
 	}
 
 	/*
-	 * This function is call in the pre completion task on the main thread, use it to call main thread post creation step for your assets
+	 * Do any UObject setup required before the build (before PostEditChange) and after all dependency UObject have been imported.
 	 * @note - This function is called when starting the pre completion task (before PostEditChange is called for the asset).
 	 */
-	virtual void BeginPreCompletedCallback(const FImportPreCompletedCallbackParams& Arguments)
+	virtual void SetupObject_GameThread(const FSetupObjectParams& Arguments)
 	{
-		check(IsInGameThread());
-		return;
+		PRAGMA_DISABLE_DEPRECATION_WARNINGS
+		PreImportPreCompletedCallback(Arguments);
+		PRAGMA_ENABLE_DEPRECATION_WARNINGS
 	}
 	
-	UE_DEPRECATED(5.2, "This function is replace by EndPreCompletedCallback.")
-	virtual void PostImportPreCompletedCallback(const FImportPreCompletedCallbackParams& Arguments)
+	UE_DEPRECATED(5.2, "This function is replace by FinalizeObject_GameThread.")
+	virtual void PostImportPreCompletedCallback(const FSetupObjectParams& Arguments)
 	{
-		EndPreCompletedCallback(Arguments);
+		check(IsInGameThread());
 	}
 
 	/*
-	 * This function is call in the pre completion task on the main thread, use it to call main thread post creation step for your assets
+	 * Do any final UObject setup after the build (after PostEditChange)
 	 * @note - This function is called at the end of the pre completion task (after PostEditChange is called for the asset).
 	 */
-	virtual void EndPreCompletedCallback(const FImportPreCompletedCallbackParams& Arguments)
+	virtual void FinalizeObject_GameThread(const FSetupObjectParams& Arguments)
 	{
-		check(IsInGameThread());
-		return;
+		PRAGMA_DISABLE_DEPRECATION_WARNINGS
+		PostImportPreCompletedCallback(Arguments);
+		PRAGMA_ENABLE_DEPRECATION_WARNINGS
 	}
 
 	/**
diff --git a/Engine/Source/Runtime/Interchange/Core/Public/InterchangePipelineBase.h b/Engine/Source/Runtime/Interchange/Core/Public/InterchangePipelineBase.h
--- a/Engine/Source/Runtime/Interchange/Core/Public/InterchangePipelineBase.h
+++ b/Engine/Source/Runtime/Interchange/Core/Public/InterchangePipelineBase.h
@@ -305,7 +305,6 @@ protected:
 	UE_DEPRECATED(5.2, "This function is replace by ExecutePipeline.")
 	virtual void ExecutePreImportPipeline(UInterchangeBaseNodeContainer* BaseNodeContainer, const TArray<UInterchangeSourceData*>& SourceDatas)
 	{
-		ExecutePipeline(BaseNodeContainer, SourceDatas);
 	}
 
 	/**
@@ -313,6 +312,9 @@ protected:
 	 */
 	virtual void ExecutePipeline(UInterchangeBaseNodeContainer* BaseNodeContainer, const TArray<UInterchangeSourceData*>& SourceDatas)
 	{
+		PRAGMA_DISABLE_DEPRECATION_WARNINGS
+		ExecutePreImportPipeline(BaseNodeContainer, SourceDatas);
+		PRAGMA_ENABLE_DEPRECATION_WARNINGS
 	}
 
 	/**
diff --git a/Engine/Source/Runtime/Interchange/Engine/Private/Tasks/InterchangeTaskCompletion.cpp b/Engine/Source/Runtime/Interchange/Engine/Private/Tasks/InterchangeTaskCompletion.cpp
--- a/Engine/Source/Runtime/Interchange/Engine/Private/Tasks/InterchangeTaskCompletion.cpp
+++ b/Engine/Source/Runtime/Interchange/Engine/Private/Tasks/InterchangeTaskCompletion.cpp
@@ -73,7 +73,7 @@ void UE::Interchange::FTaskPreCompletion::DoTask(ENamedThreads::Type CurrentThre
 				//In case Some factory code cannot run outside of the main thread we offer this callback to finish the work before calling post edit change (building the asset)
 				if (bCallPostImportGameThreadCallback && ObjectInfo.Factory)
 				{
-					UInterchangeFactoryBase::FImportPreCompletedCallbackParams Arguments;
+					UInterchangeFactoryBase::FSetupObjectParams Arguments;
 					Arguments.ImportedObject = ImportedObject;
 					Arguments.SourceData = AsyncHelper->SourceDatas[SourceIndex];
 					Arguments.FactoryNode = ObjectInfo.FactoryNode;
@@ -83,7 +83,7 @@ void UE::Interchange::FTaskPreCompletion::DoTask(ENamedThreads::Type CurrentThre
 					Arguments.Pipelines = AsyncHelper->Pipelines;
 					Arguments.OriginalPipelines = AsyncHelper->OriginalPipelines;
 					Arguments.bIsReimport = ObjectInfo.bIsReimport;
-					ObjectInfo.Factory->BeginPreCompletedCallback(Arguments);
+					ObjectInfo.Factory->SetupObject_GameThread(Arguments);
 				}
 
 				if (ImportedObject == nullptr || !IsValid(ImportedObject))
@@ -155,7 +155,7 @@ void UE::Interchange::FTaskPreCompletion::DoTask(ENamedThreads::Type CurrentThre
 				//Its possible the build of the asset to be asynchronous, the factory must handle is own asset correctly
 				if (bCallPostImportGameThreadCallback && ObjectInfo.Factory)
 				{
-					UInterchangeFactoryBase::FImportPreCompletedCallbackParams Arguments;
+					UInterchangeFactoryBase::FSetupObjectParams Arguments;
 					Arguments.ImportedObject = ImportedObject;
 					Arguments.SourceData = AsyncHelper->SourceDatas[SourceIndex];
 					Arguments.FactoryNode = ObjectInfo.FactoryNode;
@@ -164,7 +164,7 @@ void UE::Interchange::FTaskPreCompletion::DoTask(ENamedThreads::Type CurrentThre
 					Arguments.Pipelines = AsyncHelper->Pipelines;
 					Arguments.OriginalPipelines = AsyncHelper->OriginalPipelines;
 					Arguments.bIsReimport = ObjectInfo.bIsReimport;
-					ObjectInfo.Factory->EndPreCompletedCallback(Arguments);
+					ObjectInfo.Factory->FinalizeObject_GameThread(Arguments);
 				}
 			}
 		}
diff --git a/Engine/Source/Runtime/Interchange/Engine/Private/Tasks/InterchangeTaskCreateAsset.cpp b/Engine/Source/Runtime/Interchange/Engine/Private/Tasks/InterchangeTaskCreateAsset.cpp
deleted file mode 100644
diff --git a/Engine/Source/Runtime/Interchange/Engine/Private/Tasks/InterchangeTaskCreateAsset.h b/Engine/Source/Runtime/Interchange/Engine/Private/Tasks/InterchangeTaskCreateAsset.h
deleted file mode 100644
diff --git a/Engine/Source/Runtime/Interchange/Engine/Private/Tasks/InterchangeTaskCreateSceneObjects.cpp b/Engine/Source/Runtime/Interchange/Engine/Private/Tasks/InterchangeTaskCreateSceneObjects.cpp
--- a/Engine/Source/Runtime/Interchange/Engine/Private/Tasks/InterchangeTaskCreateSceneObjects.cpp
+++ b/Engine/Source/Runtime/Interchange/Engine/Private/Tasks/InterchangeTaskCreateSceneObjects.cpp
@@ -65,7 +65,7 @@ void UE::Interchange::FTaskCreateSceneObjects::DoTask(ENamedThreads::Type Curren
 		FString NodeDisplayName = FactoryNode->GetDisplayLabel();
 		SanitizeObjectName(NodeDisplayName);
 
-		UInterchangeFactoryBase::FCreateSceneObjectsParams CreateSceneObjectsParams;
+		UInterchangeFactoryBase::FImportSceneObjectsParams CreateSceneObjectsParams;
 		CreateSceneObjectsParams.ObjectName = NodeDisplayName;
 		CreateSceneObjectsParams.FactoryNode = FactoryNode;
 		CreateSceneObjectsParams.Level = GWorld->GetCurrentLevel();
@@ -75,7 +75,7 @@ void UE::Interchange::FTaskCreateSceneObjects::DoTask(ENamedThreads::Type Curren
 			CreateSceneObjectsParams.NodeContainer = AsyncHelper->BaseNodeContainers[SourceIndex].Get();
 		}
 
-		UObject* SceneObject = Factory->CreateSceneObject(CreateSceneObjectsParams);
+		UObject* SceneObject = Factory->ImportSceneObject_GameThread(CreateSceneObjectsParams);
 		if (SceneObject)
 		{
 			FScopeLock Lock(&AsyncHelper->ImportedSceneObjectsPerSourceIndexLock);
diff --git a/Engine/Source/Runtime/Interchange/Engine/Private/Tasks/InterchangeTaskCreateSceneObjects.h b/Engine/Source/Runtime/Interchange/Engine/Private/Tasks/InterchangeTaskCreateSceneObjects.h
--- a/Engine/Source/Runtime/Interchange/Engine/Private/Tasks/InterchangeTaskCreateSceneObjects.h
+++ b/Engine/Source/Runtime/Interchange/Engine/Private/Tasks/InterchangeTaskCreateSceneObjects.h
@@ -43,7 +43,7 @@ namespace UE
 
 			TStatId GetStatId() const
 			{
-				RETURN_QUICK_DECLARE_CYCLE_STAT(FTaskCreateAsset, STATGROUP_TaskGraphTasks);
+				RETURN_QUICK_DECLARE_CYCLE_STAT(FTaskCreateSceneObjects, STATGROUP_TaskGraphTasks);
 			}
 
 			void DoTask(ENamedThreads::Type CurrentThread, const FGraphEventRef& MyCompletionGraphEvent);
diff --git a/Engine/Source/Runtime/Interchange/Engine/Private/Tasks/InterchangeTaskImportObject.cpp b/Engine/Source/Runtime/Interchange/Engine/Private/Tasks/InterchangeTaskImportObject.cpp
new file mode 100644
--- /dev/null
+++ b/Engine/Source/Runtime/Interchange/Engine/Private/Tasks/InterchangeTaskImportObject.cpp
@@ -0,0 +1,399 @@
+// Copyright Epic Games, Inc. All Rights Reserved.
+#include "InterchangeTaskImportObject.h"
+
+#include "Async/TaskGraphInterfaces.h"
+#include "CoreMinimal.h"
+#include "InterchangeAssetImportData.h"
+#include "InterchangeEngineLogPrivate.h"
+#include "InterchangeFactoryBase.h"
+#include "InterchangeManager.h"
+#include "InterchangeResult.h"
+#include "InterchangeSourceData.h"
+#include "InterchangeTranslatorBase.h"
+#include "Misc/Paths.h"
+#include "PackageUtils/PackageUtils.h"
+#include "Stats/Stats.h"
+#include "Templates/SharedPointer.h"
+#include "UObject/Object.h"
+#include "UObject/UObjectGlobals.h"
+#include "UObject/ObjectMacros.h"
+#include "UObject/WeakObjectPtrTemplates.h"
+
+namespace UE
+{
+	namespace Interchange
+	{
+		namespace Private
+		{
+			void InternalGetPackageName(const UE::Interchange::FImportAsyncHelper& AsyncHelper, const int32 SourceIndex, const FString& PackageBasePath, const UInterchangeFactoryBaseNode* FactoryNode, FString& OutPackageName, FString& OutAssetName)
+			{
+				TRACE_CPUPROFILER_EVENT_SCOPE("UE::Interchange::Private::InternalGetPackageName")
+				const UInterchangeSourceData* SourceData = AsyncHelper.SourceDatas[SourceIndex];
+				check(SourceData);
+				FString NodeDisplayName = FactoryNode->GetAssetName();
+
+				// Set the asset name and the package name
+				OutAssetName = NodeDisplayName;
+				SanitizeObjectName(OutAssetName);
+
+				FString SanitizedPackageBasePath = PackageBasePath;
+				SanitizeObjectPath(SanitizedPackageBasePath);
+
+				FString SubPath;
+				if (FactoryNode->GetCustomSubPath(SubPath))
+				{
+					SanitizeObjectPath(SubPath);
+				}
+
+				OutPackageName = FPaths::Combine(*SanitizedPackageBasePath, *SubPath, *OutAssetName);
+			}
+			bool ShouldReimportFactoryNode(UInterchangeFactoryBaseNode* FactoryNode, const UInterchangeBaseNodeContainer* NodeContainer, UObject* ReimportObject)
+			{
+				TRACE_CPUPROFILER_EVENT_SCOPE("UE::Interchange::Private::ShouldReimportFactoryNode")
+
+				if (!NodeContainer)
+				{
+					return false;
+				}
+				//Find all potential factory node
+				TArray<UInterchangeFactoryBaseNode*> PotentialFactoryNodes;
+				UClass* FactoryClass = FactoryNode->GetObjectClass();
+				NodeContainer->IterateNodesOfType<UInterchangeFactoryBaseNode>([FactoryClass, &PotentialFactoryNodes](const FString& NodeUniqueID, UInterchangeFactoryBaseNode* CurrentFactoryNode)
+					{
+						if (UClass* CurrentFactoryClass = CurrentFactoryNode->GetObjectClass())
+						{
+							if (CurrentFactoryClass->IsChildOf(FactoryClass))
+							{
+								PotentialFactoryNodes.Add(CurrentFactoryNode);
+							}
+						}
+					});
+
+				if (PotentialFactoryNodes.Num() == 1)
+				{
+					//There is only one factory node that will generate this class of UObject, no need to match the unique id or the name.
+					ensure(PotentialFactoryNodes[0] == FactoryNode);
+					return true;
+				}
+
+				//If the source was used to import multiple asset, see if the FactoryNode match the ReimportObject.
+				TArray<UObject*> SubObjects;
+				GetObjectsWithOuter(ReimportObject, SubObjects);
+				for (UObject* SubObject : SubObjects)
+				{
+					if(UInterchangeAssetImportData* OriginalAssetImportData = Cast<UInterchangeAssetImportData>(SubObject))
+					{
+						if (UInterchangeBaseNodeContainer* OriginalNodeContainer = OriginalAssetImportData->NodeContainer)
+						{
+							//Find the original factory node used by the last ReimportObject import
+							if (UInterchangeFactoryBaseNode* OriginalFactoryNode = OriginalNodeContainer->GetFactoryNode(OriginalAssetImportData->NodeUniqueID))
+							{
+								//Compare the original factory node UObject class to the factory node UObject class
+								if (OriginalFactoryNode->GetObjectClass()->IsChildOf(FactoryClass))
+								{
+									//Compare the original factory node, unique id to the factory node unique id
+									if (OriginalFactoryNode->GetUniqueID().Equals(FactoryNode->GetUniqueID()))
+									{
+										return true;
+									}
+									//Compare the original factory node, name to the factory node name
+									if (OriginalFactoryNode->GetDisplayLabel().Equals(FactoryNode->GetDisplayLabel()))
+									{
+										FString PackageSubPath;
+										FactoryNode->GetCustomSubPath(PackageSubPath);
+										FString OriginalPackageSubPath;
+										OriginalFactoryNode->GetCustomSubPath(PackageSubPath);
+										//Make sure both sub path are equal
+										if (PackageSubPath.Equals(OriginalPackageSubPath))
+										{
+											return true;
+										}
+									}
+								}
+							}
+						}
+					}
+				}
+				return false;
+			}
+
+		}//ns Private
+	}//ns Interchange
+}//ns UE
+
+void UE::Interchange::FTaskImportObject_GameThread::DoTask(ENamedThreads::Type CurrentThread, const FGraphEventRef& MyCompletionGraphEvent)
+{
+	TRACE_CPUPROFILER_EVENT_SCOPE("UE::Interchange::FTaskImportObject_GameThread::DoTask")
+#if INTERCHANGE_TRACE_ASYNCHRONOUS_TASK_ENABLED
+	INTERCHANGE_TRACE_ASYNCHRONOUS_TASK(CreatePackage)
+#endif
+	TSharedPtr<UE::Interchange::FImportAsyncHelper, ESPMode::ThreadSafe> AsyncHelper = WeakAsyncHelper.Pin();
+	check(AsyncHelper.IsValid());
+
+	//Verify if the task was cancel
+	if (AsyncHelper->bCancel)
+	{
+		return;
+	}
+
+	//The create package thread must always execute on the game thread
+	check(IsInGameThread());
+
+	// Create factory
+	UInterchangeFactoryBase* Factory = NewObject<UInterchangeFactoryBase>(GetTransientPackage(), FactoryClass);
+	Factory->SetResultsContainer(AsyncHelper->AssetImportResult->GetResults());
+
+	{
+		FScopeLock Lock(&AsyncHelper->CreatedFactoriesLock);
+		AsyncHelper->CreatedFactories.Add(FactoryNode->GetUniqueID(), Factory);
+	}
+
+	UPackage* Pkg = nullptr;
+	FString PackageName;
+	FString AssetName;
+	UObject* ReimportObject = AsyncHelper->TaskData.ReimportObject;
+	//If we do a reimport no need to create a package
+	if (ReimportObject)
+	{
+		UInterchangeBaseNodeContainer* NodeContainer = nullptr;
+		if (AsyncHelper->BaseNodeContainers.IsValidIndex(SourceIndex))
+		{
+			NodeContainer = AsyncHelper->BaseNodeContainers[SourceIndex].Get();
+		}
+		
+		if (Private::ShouldReimportFactoryNode(FactoryNode, NodeContainer, ReimportObject))
+		{
+			FactoryNode->SetDisplayLabel(ReimportObject->GetName());
+			FactoryNode->SetAssetName(ReimportObject->GetName());
+			Pkg = ReimportObject->GetPackage();
+			PackageName = Pkg->GetPathName();
+			AssetName = ReimportObject->GetName();
+			//Import Asset describe by the node
+			UInterchangeFactoryBase::FImportAssetObjectParams CreateAssetParams;
+			CreateAssetParams.AssetName = AssetName;
+			CreateAssetParams.AssetNode = FactoryNode;
+			CreateAssetParams.Parent = Pkg;
+			CreateAssetParams.SourceData = AsyncHelper->SourceDatas[SourceIndex];
+			CreateAssetParams.Translator = AsyncHelper->Translators[SourceIndex];
+			CreateAssetParams.NodeContainer = NodeContainer;
+			CreateAssetParams.ReimportObject = ReimportObject;
+			FactoryNode->SetCustomReferenceObject(FSoftObjectPath(ReimportObject));
+			//We call CreateEmptyAsset to ensure any resource use by an existing UObject is released on the game thread
+			Factory->ImportAssetObject_GameThread(CreateAssetParams);
+		}
+		else
+		{
+			//Skip this asset, reimport object original factory node does not match this factory node
+			return;
+		}
+	}
+	else
+	{
+		Private::InternalGetPackageName(*AsyncHelper, SourceIndex, PackageBasePath, FactoryNode, PackageName, AssetName);
+		// We can not create assets that share the name of a map file in the same location
+		if (UE::Interchange::FPackageUtils::IsMapPackageAsset(PackageName))
+		{
+			UInterchangeResultError_Generic* Message = Factory->AddMessage<UInterchangeResultError_Generic>();
+			Message->SourceAssetName = AsyncHelper->SourceDatas[SourceIndex]->GetFilename();
+			Message->DestinationAssetName = AssetName;
+			Message->AssetType = FactoryNode->GetObjectClass();
+			Message->Text = NSLOCTEXT("Interchange", "MapExistsWithSameName", "You cannot create an asset with this name, as there is already a map file with the same name in this folder.");
+
+			//Skip this asset
+			return;
+		}
+
+		Pkg = CreatePackage(*PackageName);
+		if (Pkg == nullptr)
+		{
+			UInterchangeResultError_Generic* Message = Factory->AddMessage<UInterchangeResultError_Generic>();
+			Message->SourceAssetName = AsyncHelper->SourceDatas[SourceIndex]->GetFilename();
+			Message->DestinationAssetName = AssetName;
+			Message->AssetType = FactoryNode->GetObjectClass();
+			Message->Text = FText::Format(NSLOCTEXT("Interchange", "CouldntCreatePackage", "It was not possible to create a package named '{0}'; the asset will not be imported."), FText::FromString(PackageName));
+
+			//Skip this asset
+			return;
+		}
+
+		Pkg->SetPackageFlags(PKG_NewlyCreated);
+
+		//Import Asset describe by the node
+		UInterchangeFactoryBase::FImportAssetObjectParams CreateAssetParams;
+		CreateAssetParams.AssetName = AssetName;
+		CreateAssetParams.AssetNode = FactoryNode;
+		CreateAssetParams.Parent = Pkg;
+		CreateAssetParams.SourceData = AsyncHelper->SourceDatas[SourceIndex];
+		CreateAssetParams.Translator = AsyncHelper->Translators[SourceIndex];
+		if (AsyncHelper->BaseNodeContainers.IsValidIndex(SourceIndex))
+		{
+			CreateAssetParams.NodeContainer = AsyncHelper->BaseNodeContainers[SourceIndex].Get();
+		}
+		CreateAssetParams.ReimportObject = ReimportObject;
+		//Make sure the asset UObject is created with the correct type on the main thread
+		UObject* NodeAsset = Factory->ImportAssetObject_GameThread(CreateAssetParams);
+		if (NodeAsset)
+		{
+			if (!NodeAsset->HasAnyInternalFlags(EInternalObjectFlags::Async))
+			{
+				//Since the async flag is not set we must be in the game thread
+				ensure(IsInGameThread());
+				NodeAsset->SetInternalFlags(EInternalObjectFlags::Async);
+			}
+			FScopeLock Lock(&AsyncHelper->ImportedAssetsPerSourceIndexLock);
+			TArray<UE::Interchange::FImportAsyncHelper::FImportedObjectInfo>& ImportedInfos = AsyncHelper->ImportedAssetsPerSourceIndex.FindOrAdd(SourceIndex);
+			UE::Interchange::FImportAsyncHelper::FImportedObjectInfo& AssetInfo = ImportedInfos.AddDefaulted_GetRef();
+			AssetInfo.ImportedObject = NodeAsset;
+			AssetInfo.Factory = Factory;
+			AssetInfo.FactoryNode = FactoryNode;
+			AssetInfo.bIsReimport = bool(ReimportObject != nullptr);
+			FactoryNode->SetCustomReferenceObject(FSoftObjectPath(NodeAsset));
+		}
+	}
+
+	// Make sure the destination package is loaded
+	Pkg->FullyLoad();
+	
+	{
+		FScopeLock Lock(&AsyncHelper->CreatedPackagesLock);
+		AsyncHelper->CreatedPackages.Add(PackageName, Pkg);
+	}
+}
+
+void UE::Interchange::FTaskImportObject_Async::DoTask(ENamedThreads::Type CurrentThread, const FGraphEventRef& MyCompletionGraphEvent)
+{
+#if INTERCHANGE_TRACE_ASYNCHRONOUS_TASK_ENABLED
+	INTERCHANGE_TRACE_ASYNCHRONOUS_TASK(CreateAsset)
+#endif
+	TSharedPtr<UE::Interchange::FImportAsyncHelper, ESPMode::ThreadSafe> AsyncHelper = WeakAsyncHelper.Pin();
+	check(AsyncHelper.IsValid());
+
+	//Verify if the task was cancel
+	if (AsyncHelper->bCancel)
+	{
+		return;
+	}
+
+	UInterchangeFactoryBase* Factory = nullptr;
+	{
+		FScopeLock Lock(&AsyncHelper->CreatedFactoriesLock);
+		Factory = AsyncHelper->CreatedFactories.FindChecked(FactoryNode->GetUniqueID());
+	}
+
+	UPackage* Pkg = nullptr;
+	FString PackageName;
+	FString AssetName;
+	Private::InternalGetPackageName(*AsyncHelper, SourceIndex, PackageBasePath, FactoryNode, PackageName, AssetName);
+	bool bSkipAsset = false;
+	UObject* ReimportObject = AsyncHelper->TaskData.ReimportObject;
+	if (ReimportObject)
+	{
+		UInterchangeBaseNodeContainer* NodeContainer = nullptr;
+		if (AsyncHelper->BaseNodeContainers.IsValidIndex(SourceIndex))
+		{
+			NodeContainer = AsyncHelper->BaseNodeContainers[SourceIndex].Get();
+		}
+
+		if (Private::ShouldReimportFactoryNode(FactoryNode, NodeContainer, ReimportObject))
+		{
+			Pkg = ReimportObject->GetPackage();
+			PackageName = Pkg->GetPathName();
+			AssetName = ReimportObject->GetName();
+		}
+		else
+		{
+			bSkipAsset = true;
+		}
+	}
+	else
+	{
+		FScopeLock Lock(&AsyncHelper->CreatedPackagesLock);
+		UPackage** PkgPtr = AsyncHelper->CreatedPackages.Find(PackageName);
+
+		if (!PkgPtr || !(*PkgPtr))
+		{
+			UInterchangeResultError_Generic* Message = Factory->AddMessage<UInterchangeResultError_Generic>();
+			Message->SourceAssetName = AsyncHelper->SourceDatas[SourceIndex]->GetFilename();
+			Message->DestinationAssetName = AssetName;
+			Message->AssetType = FactoryNode->GetObjectClass();
+			Message->Text = NSLOCTEXT("Interchange", "BadPackage", "It was not possible to create the asset as its package was not created correctly.");
+
+			return;
+		}
+
+		if (!AsyncHelper->SourceDatas.IsValidIndex(SourceIndex) || !AsyncHelper->Translators.IsValidIndex(SourceIndex))
+		{
+			UInterchangeResultError_Generic* Message = Factory->AddMessage<UInterchangeResultError_Generic>();
+			Message->DestinationAssetName = AssetName;
+			Message->AssetType = FactoryNode->GetObjectClass();
+			Message->Text = NSLOCTEXT("Interchange", "SourceDataOrTranslatorInvalid", "It was not possible to create the asset as its translator was not created correctly.");
+
+			return;
+		}
+
+		Pkg = *PkgPtr;
+	}
+
+	UObject* NodeAsset = nullptr;
+	if (bSkipAsset)
+	{
+		NodeAsset = nullptr;
+	}
+	else
+	{
+		UInterchangeTranslatorBase* Translator = AsyncHelper->Translators[SourceIndex];
+		//Import Asset describe by the node
+		UInterchangeFactoryBase::FImportAssetObjectParams CreateAssetParams;
+		CreateAssetParams.AssetName = AssetName;
+		CreateAssetParams.AssetNode = FactoryNode;
+		CreateAssetParams.Parent = Pkg;
+		CreateAssetParams.SourceData = AsyncHelper->SourceDatas[SourceIndex];
+		CreateAssetParams.Translator = Translator;
+		if (AsyncHelper->BaseNodeContainers.IsValidIndex(SourceIndex))
+		{
+			CreateAssetParams.NodeContainer = AsyncHelper->BaseNodeContainers[SourceIndex].Get();
+		}
+		CreateAssetParams.ReimportObject = ReimportObject;
+
+		NodeAsset = Factory->ImportAssetObject_Async(CreateAssetParams);
+	}
+	if (NodeAsset)
+	{
+		if (!bSkipAsset)
+		{
+			FScopeLock Lock(&AsyncHelper->ImportedAssetsPerSourceIndexLock);
+			TArray<UE::Interchange::FImportAsyncHelper::FImportedObjectInfo>& ImportedInfos = AsyncHelper->ImportedAssetsPerSourceIndex.FindOrAdd(SourceIndex);
+			UE::Interchange::FImportAsyncHelper::FImportedObjectInfo* AssetInfoPtr = ImportedInfos.FindByPredicate([NodeAsset](const UE::Interchange::FImportAsyncHelper::FImportedObjectInfo& CurInfo)
+			{
+				return CurInfo.ImportedObject == NodeAsset;
+			});
+
+			if (!AssetInfoPtr)
+			{
+				UE::Interchange::FImportAsyncHelper::FImportedObjectInfo& AssetInfo = ImportedInfos.AddDefaulted_GetRef();
+				AssetInfo.ImportedObject = NodeAsset;
+				AssetInfo.Factory = Factory;
+				AssetInfo.FactoryNode = FactoryNode;
+				AssetInfo.bIsReimport = bool(ReimportObject != nullptr);
+			}
+
+			// Fill in destination asset and type in any results which have been added previously by a translator or pipeline, now that we have a corresponding factory.
+			UInterchangeResultsContainer* Results = AsyncHelper->AssetImportResult->GetResults();
+			for (UInterchangeResult* Result : Results->GetResults())
+			{
+				if (!Result->InterchangeKey.IsEmpty() && (Result->DestinationAssetName.IsEmpty() || Result->AssetType == nullptr))
+				{
+					TArray<FString> TargetAssets;
+					FactoryNode->GetTargetNodeUids(TargetAssets);
+					if (TargetAssets.Contains(Result->InterchangeKey))
+					{
+						Result->DestinationAssetName = NodeAsset->GetPathName();
+						Result->AssetType = NodeAsset->GetClass();
+					}
+				}
+			}
+		}
+
+		FactoryNode->SetCustomReferenceObject(FSoftObjectPath(NodeAsset));
+	}
+}
\ No newline at end of file
diff --git a/Engine/Source/Runtime/Interchange/Engine/Private/Tasks/InterchangeTaskImportObject.h b/Engine/Source/Runtime/Interchange/Engine/Private/Tasks/InterchangeTaskImportObject.h
new file mode 100644
--- /dev/null
+++ b/Engine/Source/Runtime/Interchange/Engine/Private/Tasks/InterchangeTaskImportObject.h
@@ -0,0 +1,98 @@
+// Copyright Epic Games, Inc. All Rights Reserved.
+
+#pragma once
+
+#include "Async/TaskGraphInterfaces.h"
+#include "CoreMinimal.h"
+#include "InterchangeFactoryBase.h"
+#include "InterchangeManager.h"
+#include "Stats/Stats.h"
+#include "UObject/WeakObjectPtrTemplates.h"
+#include "Nodes/InterchangeFactoryBaseNode.h"
+
+namespace UE
+{
+	namespace Interchange
+	{
+
+		/**
+		 * This task create UPackage and UObject, Cook::PackageTracker::NotifyUObjectCreated is not thread safe, so we need to create the packages on the main thread
+		 */
+		class FTaskImportObject_GameThread
+		{
+		private:
+			FString PackageBasePath;
+			int32 SourceIndex;
+			TWeakPtr<FImportAsyncHelper, ESPMode::ThreadSafe> WeakAsyncHelper;
+			UInterchangeFactoryBaseNode* FactoryNode;
+			const UClass* FactoryClass;
+
+		public:
+			FTaskImportObject_GameThread(const FString& InPackageBasePath, const int32 InSourceIndex, TWeakPtr<FImportAsyncHelper, ESPMode::ThreadSafe> InAsyncHelper, UInterchangeFactoryBaseNode* InFactoryNode, const UClass* InFactoryClass)
+				: PackageBasePath(InPackageBasePath)
+				, SourceIndex(InSourceIndex)
+				, WeakAsyncHelper(InAsyncHelper)
+				, FactoryNode(InFactoryNode)
+				, FactoryClass(InFactoryClass)
+			{
+				check(FactoryNode);
+				check(FactoryClass);
+			}
+
+			ENamedThreads::Type GetDesiredThread() const
+			{
+				return ENamedThreads::GameThread;
+			}
+
+			static ESubsequentsMode::Type GetSubsequentsMode()
+			{
+				return ESubsequentsMode::TrackSubsequents;
+			}
+
+			TStatId GetStatId() const
+			{
+				RETURN_QUICK_DECLARE_CYCLE_STAT(FTaskImportObject_GameThread, STATGROUP_TaskGraphTasks);
+			}
+
+			void DoTask(ENamedThreads::Type CurrentThread, const FGraphEventRef& MyCompletionGraphEvent);
+		};
+
+		class FTaskImportObject_Async
+		{
+		private:
+			FString PackageBasePath;
+			int32 SourceIndex;
+			TWeakPtr<FImportAsyncHelper, ESPMode::ThreadSafe> WeakAsyncHelper;
+			UInterchangeFactoryBaseNode* FactoryNode;
+
+		public:
+			FTaskImportObject_Async(const FString& InPackageBasePath, const int32 InSourceIndex, TWeakPtr<FImportAsyncHelper, ESPMode::ThreadSafe> InAsyncHelper, UInterchangeFactoryBaseNode* InFactoryNode)
+				: PackageBasePath(InPackageBasePath)
+				, SourceIndex(InSourceIndex)
+				, WeakAsyncHelper(InAsyncHelper)
+				, FactoryNode(InFactoryNode)
+			{
+				check(FactoryNode);
+			}
+
+			ENamedThreads::Type GetDesiredThread() const
+			{
+				return ENamedThreads::AnyBackgroundThreadNormalTask;
+			}
+
+			static ESubsequentsMode::Type GetSubsequentsMode()
+			{
+				return ESubsequentsMode::TrackSubsequents;
+			}
+
+			TStatId GetStatId() const
+			{
+				RETURN_QUICK_DECLARE_CYCLE_STAT(FTaskImportObject_Async, STATGROUP_TaskGraphTasks);
+			}
+
+			void DoTask(ENamedThreads::Type CurrentThread, const FGraphEventRef& MyCompletionGraphEvent);
+		};
+
+
+	} //ns Interchange
+}//ns UE
diff --git a/Engine/Source/Runtime/Interchange/Engine/Private/Tasks/InterchangeTaskParsing.cpp b/Engine/Source/Runtime/Interchange/Engine/Private/Tasks/InterchangeTaskParsing.cpp
--- a/Engine/Source/Runtime/Interchange/Engine/Private/Tasks/InterchangeTaskParsing.cpp
+++ b/Engine/Source/Runtime/Interchange/Engine/Private/Tasks/InterchangeTaskParsing.cpp
@@ -9,8 +9,8 @@
 #include "InterchangeManager.h"
 #include "InterchangeSourceData.h"
 #include "InterchangeTaskCompletion.h"
-#include "InterchangeTaskCreateAsset.h"
 #include "InterchangeTaskCreateSceneObjects.h"
+#include "InterchangeTaskImportObject.h"
 #include "InterchangeTaskPipeline.h"
 #include "InterchangeTranslatorBase.h"
 #include "Misc/Paths.h"
@@ -245,7 +245,6 @@ void UE::Interchange::FTaskParsing::DoTask(ENamedThreads::Type CurrentThread, co
 		const int32 SourceIndex = TaskData.SourceIndex;
 		const UClass* const FactoryClass = TaskData.FactoryClass;
 		UInterchangeFactoryBaseNode* FactoryNode = TaskData.Nodes[0];
-		const bool bFactoryCanRunOnAnyThread = FactoryClass->GetDefaultObject<UInterchangeFactoryBase>()->CanExecuteOnAnyThread();
 
 		if (TaskData.bIsSceneNode)
 		{
@@ -293,15 +292,15 @@ void UE::Interchange::FTaskParsing::DoTask(ENamedThreads::Type CurrentThread, co
 			if (ensureMsgf(!CreatedTasksAssetNames.Contains(AssetFullPath),
 				TEXT("Found multiple task data with the same asset name (%s). Only one will be executed."), *AssetFullPath))
 			{
-				//Add create package task has a prerequisite of FTaskCreateAsset. Create package task is a game thread task
+				//Add create package task has a prerequisite of FTaskImportObject_Async. Create package task is a game thread task
 				FGraphEventArray CreatePackagePrerequistes;
 				int32 CreatePackageTaskIndex = AsyncHelper->CreatePackageTasks.Add(
-					TGraphTask<FTaskCreatePackage>::CreateTask(&(TaskData.Prerequisites)).ConstructAndDispatchWhenReady(PackageBasePath, SourceIndex, WeakAsyncHelper, FactoryNode, FactoryClass)
+					TGraphTask<FTaskImportObject_GameThread>::CreateTask(&(TaskData.Prerequisites)).ConstructAndDispatchWhenReady(PackageBasePath, SourceIndex, WeakAsyncHelper, FactoryNode, FactoryClass)
 				);
 				CreatePackagePrerequistes.Add(AsyncHelper->CreatePackageTasks[CreatePackageTaskIndex]);
 
 				int32 CreateTaskIndex = AsyncHelper->CreateAssetTasks.Add(
-					TGraphTask<FTaskCreateAsset>::CreateTask(&(CreatePackagePrerequistes)).ConstructAndDispatchWhenReady(PackageBasePath, SourceIndex, WeakAsyncHelper, FactoryNode, bFactoryCanRunOnAnyThread)
+					TGraphTask<FTaskImportObject_Async>::CreateTask(&(CreatePackagePrerequistes)).ConstructAndDispatchWhenReady(PackageBasePath, SourceIndex, WeakAsyncHelper, FactoryNode)
 				);
 
 				CreatedTasksAssetNames.Add(AssetFullPath);
