Fix CIS build
#rnx

diff --git a/Engine/Source/Editor/UnrealEd/Private/Fbx/FbxSceneImportFactory.cpp b/Engine/Source/Editor/UnrealEd/Private/Fbx/FbxSceneImportFactory.cpp
--- a/Engine/Source/Editor/UnrealEd/Private/Fbx/FbxSceneImportFactory.cpp
+++ b/Engine/Source/Editor/UnrealEd/Private/Fbx/FbxSceneImportFactory.cpp
@@ -65,6 +65,7 @@
 #include "HAL/FileManager.h"
 #include "LODUtilities.h"
 #include "ComponentReregisterContext.h"
+#include "Animation/Skeleton.h"
 
 #define LOCTEXT_NAMESPACE "FBXSceneImportFactory"
 
diff --git a/Engine/Source/Runtime/Interchange/Engine/Private/Tasks/InterchangeTaskImportObject.cpp b/Engine/Source/Runtime/Interchange/Engine/Private/Tasks/InterchangeTaskImportObject.cpp
--- a/Engine/Source/Runtime/Interchange/Engine/Private/Tasks/InterchangeTaskImportObject.cpp
+++ b/Engine/Source/Runtime/Interchange/Engine/Private/Tasks/InterchangeTaskImportObject.cpp
@@ -1,7 +1,9 @@
 // Copyright Epic Games, Inc. All Rights Reserved.
 #include "InterchangeTaskImportObject.h"
 
+#if WITH_EDITOR
 #include "AssetToolsModule.h"
+#endif //WITH_EDITOR
 #include "Async/TaskGraphInterfaces.h"
 #include "CoreMinimal.h"
 #include "InterchangeAssetImportData.h"
@@ -121,6 +123,7 @@ namespace UE
 
 			bool CanImportClass(UClass* Class)
 			{
+#if WITH_EDITOR
 				if (!Class)
 				{
 					return true;
@@ -134,6 +137,7 @@ namespace UE
 						return false;
 					}
 				}
+#endif //WITH_EDITOR
 				return true;
 			}
 		}//ns Private
