Fixed missing animation and materials in shipping binaries
   Root cause: 'check' is a noop in the compilation in shipping mode.
Fixed straightforward logic issue in material pipeline: material factory nodes were disabled although the material pipeline was set to true.

diff --git a/Engine/Plugins/Interchange/Runtime/Source/Pipelines/Private/InterchangeGenericMaterialPipeline.cpp b/Engine/Plugins/Interchange/Runtime/Source/Pipelines/Private/InterchangeGenericMaterialPipeline.cpp
--- a/Engine/Plugins/Interchange/Runtime/Source/Pipelines/Private/InterchangeGenericMaterialPipeline.cpp
+++ b/Engine/Plugins/Interchange/Runtime/Source/Pipelines/Private/InterchangeGenericMaterialPipeline.cpp
@@ -329,6 +329,7 @@ void UInterchangeGenericMaterialPipeline::ExecutePipeline(UInterchangeBaseNodeCo
 	{
 		SourceNode->GetCustomImportUnusedMaterial(bImportUnusedMaterial);
 	}
+	bImportUnusedMaterial |= bImportMaterials;
 
 #if !WITH_EDITOR
 	// Can't import materials at runtime, fallback to instances
diff --git a/Engine/Source/Runtime/Interchange/Core/Public/Nodes/InterchangeBaseNodeUtilities.h b/Engine/Source/Runtime/Interchange/Core/Public/Nodes/InterchangeBaseNodeUtilities.h
--- a/Engine/Source/Runtime/Interchange/Core/Public/Nodes/InterchangeBaseNodeUtilities.h
+++ b/Engine/Source/Runtime/Interchange/Core/Public/Nodes/InterchangeBaseNodeUtilities.h
@@ -365,8 +365,8 @@ public:
 		{
 			FAttributeKey IndexKey = GetKeyAttribute(CachedKeysAndValues.Num());
 			FAttributeKey AttributeKey = GetValueAttribute(InKey);
-			check(AttributesPtr->RegisterAttribute<KeyType>(IndexKey, InKey) == EAttributeStorageResult::Operation_Success);
-			check(AttributesPtr->RegisterAttribute<ValueType>(AttributeKey, InValue) == EAttributeStorageResult::Operation_Success);
+			ensure(AttributesPtr->RegisterAttribute<KeyType>(IndexKey, InKey) == EAttributeStorageResult::Operation_Success);
+			ensure(AttributesPtr->RegisterAttribute<ValueType>(AttributeKey, InValue) == EAttributeStorageResult::Operation_Success);
 			CachedKeysAndValues.AddByHash(
 				Hash
 				, InKey
