Added error message to result when a unsupported class is encountered.
Protected glTF import from image sources which are not part of the specs

diff --git a/Engine/Plugins/Interchange/Runtime/Source/Import/Private/Gltf/InterchangeGltfTranslator.cpp b/Engine/Plugins/Interchange/Runtime/Source/Import/Private/Gltf/InterchangeGltfTranslator.cpp
--- a/Engine/Plugins/Interchange/Runtime/Source/Import/Private/Gltf/InterchangeGltfTranslator.cpp
+++ b/Engine/Plugins/Interchange/Runtime/Source/Import/Private/Gltf/InterchangeGltfTranslator.cpp
@@ -1015,6 +1015,28 @@ bool UInterchangeGltfTranslator::Translate( UInterchangeBaseNodeContainer& NodeC
 		int32 TextureIndex = 0;
 		for ( const GLTF::FTexture& GltfTexture : GltfAsset.Textures )
 		{
+			// The glTF reader enforces the spec on the image format for buffers, URIs and file paths
+			// Skip the texture is the glTF reader has not recognized the format
+			if (GltfTexture.Source.Format == GLTF::FImage::EFormat::Unknown)
+			{
+				UInterchangeResultError_Generic* Message = AddMessage<UInterchangeResultError_Generic>();
+
+				FText TextMessage;
+				if (GltfTexture.Source.FilePath.IsEmpty())
+				{
+					Message->Text = FText::Format(LOCTEXT("TextureCreationFailed", "The image format of the buffer for texture {0} is not supported."), FText::FromString(GltfTexture.Name));
+				}
+				else
+				{
+					Message->SourceAssetName = GetSourceData()->GetFilename();
+					Message->Text = FText::Format(
+						LOCTEXT("TextureCreationFailed", "The extension of the image file, {0}, for texture {1} is not supported."),
+						FText::FromString(GltfTexture.Source.FilePath), FText::FromString(GltfTexture.Name));
+				}
+
+				continue;
+			}
+
 			UInterchangeTexture2DNode* TextureNode = UInterchangeTexture2DNode::Create(&NodeContainer, GltfTexture.UniqueId);
 			TextureNode->SetDisplayLabel(GltfTexture.Name);
 
diff --git a/Engine/Source/Runtime/Interchange/Engine/Private/Tasks/InterchangeTaskImportObject.cpp b/Engine/Source/Runtime/Interchange/Engine/Private/Tasks/InterchangeTaskImportObject.cpp
--- a/Engine/Source/Runtime/Interchange/Engine/Private/Tasks/InterchangeTaskImportObject.cpp
+++ b/Engine/Source/Runtime/Interchange/Engine/Private/Tasks/InterchangeTaskImportObject.cpp
@@ -121,23 +121,52 @@ namespace UE
 				return false;
 			}
 
-			bool CanImportClass(UClass* Class)
+			bool CanImportClass(UE::Interchange::FImportAsyncHelper& AsyncHelper, UInterchangeFactoryBaseNode& FactoryNode, int32 SourceIndex)
 			{
-#if WITH_EDITOR
-				if (!Class)
-				{
-					return true;
-				}
-				IAssetTools& AssetTools = FAssetToolsModule::GetModule().Get();
-				TSharedPtr<FPathPermissionList> AssetClassPermissionList = AssetTools.GetAssetClassPathPermissionList(EAssetClassAction::ImportAsset);
-				if (AssetClassPermissionList && AssetClassPermissionList->HasFiltering())
+			#if WITH_EDITOR
+				if (UClass* Class = FactoryNode.GetObjectClass())
 				{
-					if (!AssetClassPermissionList->PassesFilter(Class->GetPathName()))
+					if (AsyncHelper.AllowedClasses.Contains(Class))
+					{
+						return true;
+					}
+					else if (AsyncHelper.DeniedClasses.Contains(Class))
 					{
 						return false;
 					}
+					else
+					{
+						IAssetTools& AssetTools = FAssetToolsModule::GetModule().Get();
+						TSharedPtr<FPathPermissionList> AssetClassPermissionList = AssetTools.GetAssetClassPathPermissionList(EAssetClassAction::ImportAsset);
+						if (AssetClassPermissionList && AssetClassPermissionList->HasFiltering())
+						{
+							if (!AssetClassPermissionList->PassesFilter(Class->GetPathName()))
+							{
+								UInterchangeResultError_Generic* Message = AsyncHelper.AssetImportResult->GetResults()->Add<UInterchangeResultError_Generic>();
+								Message->SourceAssetName = AsyncHelper.SourceDatas[SourceIndex]->GetFilename();
+								FString AssetName = FactoryNode.GetAssetName();
+								UE::Interchange::SanitizeObjectName(AssetName);
+								Message->DestinationAssetName = AssetName;
+								Message->AssetType = Class;
+								Message->Text = FText::Format(NSLOCTEXT("Interchange", "NotAllowedClass", "The creation of asset of class '{0}' is not allowed in this project."), FText::FromString(Class->GetName()));
+
+								AsyncHelper.DeniedClasses.Add(Class);
+
+								return false;
+							}
+						}
+
+						AsyncHelper.AllowedClasses.Add(Class);
+
+						return true;
+					}
 				}
-#endif //WITH_EDITOR
+				else
+				{
+					return false;
+				}
+			#endif //WITH_EDITOR
+
 				return true;
 			}
 		}//ns Private
@@ -150,16 +179,17 @@ void UE::Interchange::FTaskImportObject_GameThread::DoTask(ENamedThreads::Type C
 #if INTERCHANGE_TRACE_ASYNCHRONOUS_TASK_ENABLED
 	INTERCHANGE_TRACE_ASYNCHRONOUS_TASK(CreatePackage)
 #endif
-	TSharedPtr<UE::Interchange::FImportAsyncHelper, ESPMode::ThreadSafe> AsyncHelper = WeakAsyncHelper.Pin();
+	using namespace UE::Interchange;
+
+	TSharedPtr<FImportAsyncHelper, ESPMode::ThreadSafe> AsyncHelper = WeakAsyncHelper.Pin();
 	check(AsyncHelper.IsValid());
 
-	//Verify if the task was cancel
-	if (AsyncHelper->bCancel
-		|| !FactoryNode
-		|| !UE::Interchange::Private::CanImportClass(FactoryNode->GetObjectClass()))
+	//Verify if the task was cancel or class to import is denied
+	if (AsyncHelper->bCancel || !FactoryNode || !Private::CanImportClass(*AsyncHelper, *FactoryNode, SourceIndex))
 	{
 		return;
 	}
+
 	//The create package thread must always execute on the game thread
 	check(IsInGameThread());
 
@@ -215,7 +245,7 @@ void UE::Interchange::FTaskImportObject_GameThread::DoTask(ENamedThreads::Type C
 	{
 		Private::InternalGetPackageName(*AsyncHelper, SourceIndex, PackageBasePath, FactoryNode, PackageName, AssetName);
 		// We can not create assets that share the name of a map file in the same location
-		if (UE::Interchange::FPackageUtils::IsMapPackageAsset(PackageName))
+		if (FPackageUtils::IsMapPackageAsset(PackageName))
 		{
 			UInterchangeResultError_Generic* Message = Factory->AddMessage<UInterchangeResultError_Generic>();
 			Message->SourceAssetName = AsyncHelper->SourceDatas[SourceIndex]->GetFilename();
@@ -265,8 +295,8 @@ void UE::Interchange::FTaskImportObject_GameThread::DoTask(ENamedThreads::Type C
 				NodeAsset->SetInternalFlags(EInternalObjectFlags::Async);
 			}
 			FScopeLock Lock(&AsyncHelper->ImportedAssetsPerSourceIndexLock);
-			TArray<UE::Interchange::FImportAsyncHelper::FImportedObjectInfo>& ImportedInfos = AsyncHelper->ImportedAssetsPerSourceIndex.FindOrAdd(SourceIndex);
-			UE::Interchange::FImportAsyncHelper::FImportedObjectInfo& AssetInfo = ImportedInfos.AddDefaulted_GetRef();
+			TArray<FImportAsyncHelper::FImportedObjectInfo>& ImportedInfos = AsyncHelper->ImportedAssetsPerSourceIndex.FindOrAdd(SourceIndex);
+			FImportAsyncHelper::FImportedObjectInfo& AssetInfo = ImportedInfos.AddDefaulted_GetRef();
 			AssetInfo.ImportedObject = NodeAsset;
 			AssetInfo.Factory = Factory;
 			AssetInfo.FactoryNode = FactoryNode;
@@ -289,13 +319,13 @@ void UE::Interchange::FTaskImportObject_Async::DoTask(ENamedThreads::Type Curren
 #if INTERCHANGE_TRACE_ASYNCHRONOUS_TASK_ENABLED
 	INTERCHANGE_TRACE_ASYNCHRONOUS_TASK(CreateAsset)
 #endif
-	TSharedPtr<UE::Interchange::FImportAsyncHelper, ESPMode::ThreadSafe> AsyncHelper = WeakAsyncHelper.Pin();
+		using namespace UE::Interchange;
+
+	TSharedPtr<FImportAsyncHelper, ESPMode::ThreadSafe> AsyncHelper = WeakAsyncHelper.Pin();
 	check(AsyncHelper.IsValid());
 
 	//Verify if the task was cancel
-	if (AsyncHelper->bCancel
-		|| !FactoryNode
-		|| !UE::Interchange::Private::CanImportClass(FactoryNode->GetObjectClass()))
+	if (AsyncHelper->bCancel || !FactoryNode || !Private::CanImportClass(*AsyncHelper, *FactoryNode, SourceIndex))
 	{
 		return;
 	}
@@ -388,15 +418,15 @@ void UE::Interchange::FTaskImportObject_Async::DoTask(ENamedThreads::Type Curren
 		if (!bSkipAsset)
 		{
 			FScopeLock Lock(&AsyncHelper->ImportedAssetsPerSourceIndexLock);
-			TArray<UE::Interchange::FImportAsyncHelper::FImportedObjectInfo>& ImportedInfos = AsyncHelper->ImportedAssetsPerSourceIndex.FindOrAdd(SourceIndex);
-			UE::Interchange::FImportAsyncHelper::FImportedObjectInfo* AssetInfoPtr = ImportedInfos.FindByPredicate([NodeAsset](const UE::Interchange::FImportAsyncHelper::FImportedObjectInfo& CurInfo)
+			TArray<FImportAsyncHelper::FImportedObjectInfo>& ImportedInfos = AsyncHelper->ImportedAssetsPerSourceIndex.FindOrAdd(SourceIndex);
+			FImportAsyncHelper::FImportedObjectInfo* AssetInfoPtr = ImportedInfos.FindByPredicate([NodeAsset](const FImportAsyncHelper::FImportedObjectInfo& CurInfo)
 			{
 				return CurInfo.ImportedObject == NodeAsset;
 			});
 
 			if (!AssetInfoPtr)
 			{
-				UE::Interchange::FImportAsyncHelper::FImportedObjectInfo& AssetInfo = ImportedInfos.AddDefaulted_GetRef();
+				FImportAsyncHelper::FImportedObjectInfo& AssetInfo = ImportedInfos.AddDefaulted_GetRef();
 				AssetInfo.ImportedObject = NodeAsset;
 				AssetInfo.Factory = Factory;
 				AssetInfo.FactoryNode = FactoryNode;
diff --git a/Engine/Source/Runtime/Interchange/Engine/Public/InterchangeManager.h b/Engine/Source/Runtime/Interchange/Engine/Public/InterchangeManager.h
--- a/Engine/Source/Runtime/Interchange/Engine/Public/InterchangeManager.h
+++ b/Engine/Source/Runtime/Interchange/Engine/Public/InterchangeManager.h
@@ -218,6 +218,12 @@ namespace UE
 			FCriticalSection CreatedFactoriesLock;
 			TMap<FString, UInterchangeFactoryBase*> CreatedFactories;
 
+			// Set of classes which creation has been denied
+			TSet<UClass*> DeniedClasses;
+
+			// Set of classes which creation is allowed
+			TSet<UClass*> AllowedClasses;
+
 			struct FImportedObjectInfo
 			{
 				UObject* ImportedObject = nullptr; // The object that was imported
