Interchange: Fix crash when importing multiple file containing staticmesh with the same name
#rnx

diff --git a/Engine/Source/Runtime/Interchange/Engine/Private/InterchangeManager.cpp b/Engine/Source/Runtime/Interchange/Engine/Private/InterchangeManager.cpp
--- a/Engine/Source/Runtime/Interchange/Engine/Private/InterchangeManager.cpp
+++ b/Engine/Source/Runtime/Interchange/Engine/Private/InterchangeManager.cpp
@@ -909,10 +909,9 @@ void UInterchangeManager::StartQueuedTasks(bool bCancelAllTasks /*= false*/)
 		}
 	};
 
-	//We need to leave some free task in the pool to avoid deadlock.
-	//Each import can use 2 tasks in same time if the build of the asset ddc use the same task pool (i.e. staticmesh, skeletalmesh, texture...)
-	const int32 PoolWorkerThreadCount = FTaskGraphInterface::Get().GetNumWorkerThreads()/2;
-	const int32 MaxNumWorker = FMath::Max(PoolWorkerThreadCount, 1);
+
+	//Hack, import the jobs one by one
+	const int32 MaxNumWorker = 1;
 	while (!QueuedTasks.IsEmpty() && (ImportTasks.Num() < MaxNumWorker || bCancelAllTasks))
 	{
 		FQueuedTaskData QueuedTaskData;
