Interchange: Change the interchange error message to a Display log.

#rnx

diff --git a/Engine/Source/Runtime/Interchange/Engine/Private/Tasks/InterchangeTaskImportObject.cpp b/Engine/Source/Runtime/Interchange/Engine/Private/Tasks/InterchangeTaskImportObject.cpp
--- a/Engine/Source/Runtime/Interchange/Engine/Private/Tasks/InterchangeTaskImportObject.cpp
+++ b/Engine/Source/Runtime/Interchange/Engine/Private/Tasks/InterchangeTaskImportObject.cpp
@@ -142,16 +142,8 @@ namespace UE
 						{
 							if (!AssetClassPermissionList->PassesFilter(Class->GetPathName()))
 							{
-								UInterchangeResultError_Generic* Message = AsyncHelper.AssetImportResult->GetResults()->Add<UInterchangeResultError_Generic>();
-								Message->SourceAssetName = AsyncHelper.SourceDatas[SourceIndex]->GetFilename();
-								FString AssetName = FactoryNode.GetAssetName();
-								UE::Interchange::SanitizeObjectName(AssetName);
-								Message->DestinationAssetName = AssetName;
-								Message->AssetType = Class;
-								Message->Text = FText::Format(NSLOCTEXT("Interchange", "NotAllowedClass", "The creation of asset of class '{0}' is not allowed in this project."), FText::FromString(Class->GetName()));
-
+								UE_LOG(LogInterchangeEngine, Display, TEXT("The creation of asset of class '%s' is not allowed in this project."), *Class->GetName());
 								AsyncHelper.DeniedClasses.Add(Class);
-
 								return false;
 							}
 						}
