Tentative fix for crash in UInterchangeTranslatorBase::AddMessage and in UInterchangeManager::CreateSourceData

#rnx

diff --git a/Engine/Source/Runtime/Interchange/Core/Public/InterchangeTranslatorBase.h b/Engine/Source/Runtime/Interchange/Core/Public/InterchangeTranslatorBase.h
--- a/Engine/Source/Runtime/Interchange/Core/Public/InterchangeTranslatorBase.h
+++ b/Engine/Source/Runtime/Interchange/Core/Public/InterchangeTranslatorBase.h
@@ -95,20 +95,23 @@ public:
 	template <typename T>
 	T* AddMessage() const
 	{
-		ensure(Results != nullptr);
-		ensure(SourceData != nullptr);
-		T* Item = Results->Add<T>();
-		Item->SourceAssetName = SourceData->GetFilename();
-		return Item;
+		if (ensure(Results) && ensure(SourceData))
+		{
+			T* Item = Results->Add<T>();
+			Item->SourceAssetName = SourceData->GetFilename();
+			return Item;
+		}
+		return nullptr;
 	}
 
 
 	void AddMessage(UInterchangeResult* Item) const
 	{
-		ensure(Results != nullptr);
-		ensure(SourceData != nullptr);
-		Results->Add(Item);
-		Item->SourceAssetName = SourceData->GetFilename();
+		if (ensure(Results) && ensure(SourceData))
+		{
+			Results->Add(Item);
+			Item->SourceAssetName = SourceData->GetFilename();
+		}
 	}
 	
 
diff --git a/Engine/Source/Runtime/Interchange/Engine/Private/InterchangeManager.cpp b/Engine/Source/Runtime/Interchange/Engine/Private/InterchangeManager.cpp
--- a/Engine/Source/Runtime/Interchange/Engine/Private/InterchangeManager.cpp
+++ b/Engine/Source/Runtime/Interchange/Engine/Private/InterchangeManager.cpp
@@ -1410,7 +1410,7 @@ bool UInterchangeManager::ExportScene(const UObject* World, bool bIsAutomated)
 UInterchangeSourceData* UInterchangeManager::CreateSourceData(const FString& InFileName)
 {
 	UInterchangeSourceData* SourceDataAsset = NewObject<UInterchangeSourceData>(GetTransientPackage(), NAME_None);
-	if(!InFileName.IsEmpty())
+	if(!InFileName.IsEmpty() && SourceDataAsset)
 	{
 		SourceDataAsset->SetFilename(InFileName);
 	}
