Interchange: Fix runtime InterchangeEngine dependency on SkeletalMeshUtilitiesCommon

#rnx

diff --git a/Engine/Source/Runtime/Interchange/Engine/InterchangeEngine.Build.cs b/Engine/Source/Runtime/Interchange/Engine/InterchangeEngine.Build.cs
--- a/Engine/Source/Runtime/Interchange/Engine/InterchangeEngine.Build.cs
+++ b/Engine/Source/Runtime/Interchange/Engine/InterchangeEngine.Build.cs
@@ -8,13 +8,13 @@ public class InterchangeEngine : ModuleRules
 	{
 		bTreatAsEngineModule = true;
 		PublicDependencyModuleNames.AddRange(
-				new string[] {
-					"Core",
-					"CoreUObject",
-					"Engine",
-					"InterchangeCore"
-				}
-			);
+			new string[] {
+				"Core",
+				"CoreUObject",
+				"Engine",
+				"InterchangeCore"
+			}
+		);
 
 		PrivateDependencyModuleNames.AddRange(
 			new string[] {
@@ -22,10 +22,19 @@ public class InterchangeEngine : ModuleRules
 				"SlateCore",
 				"Slate",
 				"DeveloperSettings",
-				"SkeletalMeshUtilitiesCommon",
 				"Json",
 				"JsonUtilities",
 			}
 		);
+
+		if (Target.bBuildEditor)
+		{
+			PrivateDependencyModuleNames.AddRange(
+				new string[]
+				{
+					"SkeletalMeshUtilitiesCommon",
+				}
+			);
+		}
 	}
 }
diff --git a/Engine/Source/Runtime/Interchange/Engine/Private/InterchangeMeshUtilities.cpp b/Engine/Source/Runtime/Interchange/Engine/Private/InterchangeMeshUtilities.cpp
--- a/Engine/Source/Runtime/Interchange/Engine/Private/InterchangeMeshUtilities.cpp
+++ b/Engine/Source/Runtime/Interchange/Engine/Private/InterchangeMeshUtilities.cpp
@@ -17,7 +17,10 @@
 #include "UObject/Class.h"
 #include "UObject/Object.h"
 #include "UObject/ObjectMacros.h"
+
+#if WITH_EDITOR
 #include "LODUtilities.h"
+#endif
 
 #include UE_INLINE_GENERATED_CPP_BY_NAME(InterchangeMeshUtilities)
 
diff --git a/Engine/Source/Runtime/Interchange/Engine/Public/InterchangeManager.h b/Engine/Source/Runtime/Interchange/Engine/Public/InterchangeManager.h
--- a/Engine/Source/Runtime/Interchange/Engine/Public/InterchangeManager.h
+++ b/Engine/Source/Runtime/Interchange/Engine/Public/InterchangeManager.h
@@ -284,7 +284,7 @@ public:
 };
 
 USTRUCT(BlueprintType)
-struct FImportAssetParameters
+struct INTERCHANGEENGINE_API FImportAssetParameters
 {
 	GENERATED_USTRUCT_BODY()
 
