Implemented ability to reimport a file into a level - #2
  - Added asset, InterchangeSceneImportAsset, used to track import and reimport of levels
  - Added translated node, factory node and factory for InterchangeSceneImportAsset
  - Modified InterchangeGenericScenesPipeline to support InterchangeSceneImportAsset
  - Modified pipeline assets for reimport workflows

diff --git a/Engine/Plugins/Interchange/Editor/Source/InterchangeEditor/InterchangeEditor.Build.cs b/Engine/Plugins/Interchange/Editor/Source/InterchangeEditor/InterchangeEditor.Build.cs
--- a/Engine/Plugins/Interchange/Editor/Source/InterchangeEditor/InterchangeEditor.Build.cs
+++ b/Engine/Plugins/Interchange/Editor/Source/InterchangeEditor/InterchangeEditor.Build.cs
@@ -14,15 +14,21 @@ namespace UnrealBuildTool.Rules
                     "Core",
 					"CoreUObject",
 					"Engine",
+					"UnrealEd",
 				}
 			);
 
 			PrivateDependencyModuleNames.AddRange(
 				new string[]
 				{
+					"AssetDefinition",
+					"ContentBrowser",
 					"InterchangeCore",
 					"InterchangeEngine",
+					"InterchangeImport",
 					"MessageLog",
+					"SlateCore",
+					"ToolMenus",
 				}
 			);
 		}
diff --git a/Engine/Plugins/Interchange/Editor/Source/InterchangeEditor/Private/AssetDefinition_InterchangeSceneImportAsset.cpp b/Engine/Plugins/Interchange/Editor/Source/InterchangeEditor/Private/AssetDefinition_InterchangeSceneImportAsset.cpp
new file mode 100644
--- /dev/null
+++ b/Engine/Plugins/Interchange/Editor/Source/InterchangeEditor/Private/AssetDefinition_InterchangeSceneImportAsset.cpp
@@ -0,0 +1,168 @@
+// Copyright Epic Games, Inc. All Rights Reserved.
+
+#include "AssetDefinition_InterchangeSceneImportAsset.h"
+
+#include "InterchangeAssetImportData.h"
+#include "InterchangeEditorModule.h"
+#include "InterchangeManager.h"
+#include "InterchangeSceneImportAsset.h"
+
+#include "ContentBrowserMenuContexts.h"
+#include "Misc/App.h"
+#include "Misc/DelayedAutoRegister.h"
+#include "ThumbnailRendering/SceneThumbnailInfo.h"
+#include "ToolMenus.h"
+#include "ToolMenuSection.h"
+#include "EditorReimportHandler.h"
+
+#define LOCTEXT_NAMESPACE "AssetDefinition_InterchangeSceneImportAsset"
+
+FAssetCategoryPath UAssetDefinition_InterchangeSceneImportAsset::Interchange(LOCTEXT("Interchange_Category_Path", "Interchange"));
+
+EAssetCommandResult UAssetDefinition_InterchangeSceneImportAsset::OpenAssets(const FAssetOpenArgs& /*OpenArgs*/) const
+{
+	return EAssetCommandResult::Handled;
+}
+
+UThumbnailInfo* UAssetDefinition_InterchangeSceneImportAsset::LoadThumbnailInfo(const FAssetData& InAssetData) const
+{
+	return UE::Editor::FindOrCreateThumbnailInfo(InAssetData.GetAsset(), USceneThumbnailInfo::StaticClass());
+}
+
+namespace MenuExtension_InterchangeSceneImportAsset
+{
+	void ExecuteReimportOneAsset(UInterchangeSceneImportAsset* Asset, const FString& FilePath)
+	{
+		using namespace UE::Interchange;
+
+		FScopedSourceData ScopedSourceData(FilePath);
+		const UInterchangeSourceData* SourceData = ScopedSourceData.GetSourceData();
+
+		UInterchangeManager& InterchangeManager = UInterchangeManager::GetInterchangeManager();
+		if (InterchangeManager.CanTranslateSourceData(SourceData))
+		{
+			FImportAssetParameters ImportAssetParameters;
+			ImportAssetParameters.bIsAutomated = GIsAutomationTesting || FApp::IsUnattended() || IsRunningCommandlet() || GIsRunningUnattendedScript;
+			ImportAssetParameters.ReimportAsset = Asset;
+			ImportAssetParameters.ReimportSourceIndex = INDEX_NONE;
+
+			TTuple<FAssetImportResultRef, FSceneImportResultRef> ImportResult = InterchangeManager.ImportSceneAsync(FString(), SourceData, ImportAssetParameters);
+
+			// TODO: Report if reimport failed
+			//return ImportResult;
+		}
+	}
+
+	void ExecuteReimport(const FToolMenuContext& InContext)
+	{
+		const UContentBrowserAssetContextMenuContext* Context = UContentBrowserAssetContextMenuContext::FindContextWithAssets(InContext);
+		TArray<UInterchangeSceneImportAsset*> SceneImportAssets = Context->LoadSelectedObjects<UInterchangeSceneImportAsset>();
+
+		for (UInterchangeSceneImportAsset* SceneImportAsset : SceneImportAssets)
+		{
+			if (SceneImportAsset)
+			{
+				ExecuteReimportOneAsset(SceneImportAsset, SceneImportAsset->AssetImportData->GetFirstFilename());
+			}
+		}
+	}
+
+	void ExecuteReimportWithFile(const FToolMenuContext& InContext)
+	{
+		const UContentBrowserAssetContextMenuContext* Context = UContentBrowserAssetContextMenuContext::FindContextWithAssets(InContext);
+		TArray<UInterchangeSceneImportAsset*> SceneImportAssets = Context->LoadSelectedObjects<UInterchangeSceneImportAsset>();
+		if (SceneImportAssets.Num() != 1 || !SceneImportAssets[0])
+		{
+			return;
+		}
+		
+		UInterchangeSceneImportAsset* SceneImportAsset = SceneImportAssets[0];
+
+		TArray<FString> OpenFilenames;
+		FReimportManager::Instance()->GetNewReimportPath(SceneImportAsset, OpenFilenames);
+		if (OpenFilenames.Num() == 1 && !OpenFilenames[0].IsEmpty())
+		{
+			ExecuteReimportOneAsset(SceneImportAsset, OpenFilenames[0]);
+		}
+
+	}
+
+	bool CanExecuteReimport(const FToolMenuContext& InContext)
+	{
+		const UContentBrowserAssetContextMenuContext* Context = UContentBrowserAssetContextMenuContext::FindContextWithAssets(InContext);
+		TArray<UInterchangeSceneImportAsset*> SceneImportAssets = Context->LoadSelectedObjects<UInterchangeSceneImportAsset>();
+
+		for (UInterchangeSceneImportAsset* SceneImportAsset : SceneImportAssets)
+		{
+			if (SceneImportAsset && SceneImportAsset->AssetImportData)
+			{
+				// TODO: Check that at least one file exists before returning true
+				if (SceneImportAsset->AssetImportData->ExtractFilenames().Num() > 0)
+				{
+					return true;
+				}
+			}
+		}
+
+		return false;
+	}
+
+	bool CanExecuteReimportWithFile(const FToolMenuContext& InContext)
+	{
+		const UContentBrowserAssetContextMenuContext* Context = UContentBrowserAssetContextMenuContext::FindContextWithAssets(InContext);
+		TArray<UInterchangeSceneImportAsset*> SceneImportAssets = Context->LoadSelectedObjects<UInterchangeSceneImportAsset>();
+
+		if (SceneImportAssets.Num() == 1)
+		{
+			const UInterchangeSceneImportAsset* SceneImportAsset = SceneImportAssets[0];
+			if (SceneImportAsset && SceneImportAsset->AssetImportData)
+			{
+				return true;
+			}
+		}
+
+		return false;
+	}
+
+	static FDelayedAutoRegisterHelper DelayedAutoRegister(EDelayedRegisterRunPhase::EndOfEngineInit, []{ 
+		UToolMenus::RegisterStartupCallback(FSimpleMulticastDelegate::FDelegate::CreateLambda([]()
+		{
+			FToolMenuOwnerScoped OwnerScoped(UE_MODULE_NAME);
+			UToolMenu* Menu = UE::ContentBrowser::ExtendToolMenu_AssetContextMenu(UInterchangeSceneImportAsset::StaticClass());
+		
+			if (!Menu->FindSection("Interchange"))
+			{
+				const FToolMenuInsert MenuInsert(NAME_None, EToolMenuInsertType::First);
+				Menu->AddSection("Interchange", LOCTEXT("InterchangeSceneImportAsset_Section", "Interchange"), MenuInsert);
+			}
+			FToolMenuSection& Section = Menu->FindOrAddSection("Interchange");
+			Section.AddDynamicEntry(NAME_None, FNewToolMenuSectionDelegate::CreateLambda([](FToolMenuSection& InSection)
+			{
+				{
+					const TAttribute<FText> Label = LOCTEXT("InterchangeSceneImportAsset_Reimport", "Reimport Scene");
+					const TAttribute<FText> ToolTip = LOCTEXT("InterchangeSceneImportAsset_ReimportTooltip", "Reimport the scene associated with each selected InterchangeSceneImportAsset.");
+					const FSlateIcon Icon = FSlateIcon(FAppStyle::GetAppStyleSetName(), "Icons.Import");
+
+					FToolUIAction UIAction;
+					UIAction.ExecuteAction = FToolMenuExecuteAction::CreateStatic(&ExecuteReimport);
+					UIAction.CanExecuteAction = FToolMenuCanExecuteAction::CreateStatic(&CanExecuteReimport);
+
+					InSection.AddMenuEntry("InterchangeSceneImportAsset_Reimport", Label, ToolTip, Icon, UIAction);
+				}
+				{
+					const TAttribute<FText> Label = LOCTEXT("InterchangeSceneImportAsset_ReimportWithFile", "Reimport Scene With File");
+					const TAttribute<FText> ToolTip = LOCTEXT("InterchangeSceneImportAsset_ReimportWithFile_Tooltip", "Reimport the scene associated with the selected InterchangeSceneImportAsset using a new file.");
+					const FSlateIcon Icon = FSlateIcon(FAppStyle::GetAppStyleSetName(), "Icons.Import");
+
+					FToolUIAction UIAction;
+					UIAction.ExecuteAction = FToolMenuExecuteAction::CreateStatic(&ExecuteReimportWithFile);
+					UIAction.CanExecuteAction = FToolMenuCanExecuteAction::CreateStatic(&CanExecuteReimportWithFile);
+
+					InSection.AddMenuEntry("InterchangeSceneImportAsset_ReimportWithFile", Label, ToolTip, Icon, UIAction);
+				}
+			}));
+		}));
+	});
+}
+
+#undef LOCTEXT_NAMESPACE
diff --git a/Engine/Plugins/Interchange/Editor/Source/InterchangeEditor/Private/AssetDefinition_InterchangeSceneImportAsset.h b/Engine/Plugins/Interchange/Editor/Source/InterchangeEditor/Private/AssetDefinition_InterchangeSceneImportAsset.h
new file mode 100644
--- /dev/null
+++ b/Engine/Plugins/Interchange/Editor/Source/InterchangeEditor/Private/AssetDefinition_InterchangeSceneImportAsset.h
@@ -0,0 +1,31 @@
+// Copyright Epic Games, Inc. All Rights Reserved.
+
+#pragma once
+
+#include "AssetDefinitionDefault.h"
+#include "InterchangeSceneImportAsset.h"
+
+#include "AssetDefinition_InterchangeSceneImportAsset.generated.h"
+
+UCLASS()
+class UAssetDefinition_InterchangeSceneImportAsset : public UAssetDefinitionDefault
+{
+	GENERATED_BODY()
+
+public:
+	// UAssetDefinition Begin
+	virtual FText GetAssetDisplayName() const override { return NSLOCTEXT("AssetDefinition", "AssetDefinition_InterchangeSceneImportAsset", "InterchangeSceneImportAsset"); }
+	virtual FLinearColor GetAssetColor() const override { return FLinearColor(FColor(255, 0, 0)); }
+	virtual TSoftClassPtr<UObject> GetAssetClass() const override { return UInterchangeSceneImportAsset::StaticClass(); }
+	virtual TConstArrayView<FAssetCategoryPath> GetAssetCategories() const override
+	{
+		static const auto Categories = { Interchange };
+		return Categories;
+	}
+	virtual UThumbnailInfo* LoadThumbnailInfo(const FAssetData& InAssetData) const override;
+	virtual EAssetCommandResult OpenAssets(const FAssetOpenArgs& OpenArgs) const override;
+	// UAssetDefinition End
+
+private:
+	static FAssetCategoryPath Interchange;
+};
diff --git a/Engine/Plugins/Interchange/Runtime/Source/FactoryNodes/Private/InterchangeSceneImportAssetFactoryNode.cpp b/Engine/Plugins/Interchange/Runtime/Source/FactoryNodes/Private/InterchangeSceneImportAssetFactoryNode.cpp
new file mode 100644
--- /dev/null
+++ b/Engine/Plugins/Interchange/Runtime/Source/FactoryNodes/Private/InterchangeSceneImportAssetFactoryNode.cpp
@@ -0,0 +1,12 @@
+// Copyright Epic Games, Inc. All Rights Reserved.
+
+#include "InterchangeSceneImportAssetFactoryNode.h"
+
+#include "InterchangeSceneImportAsset.h"
+
+#include UE_INLINE_GENERATED_CPP_BY_NAME(InterchangeSceneImportAssetFactoryNode)
+
+UClass* UInterchangeSceneImportAssetFactoryNode::GetObjectClass() const
+{
+	return UInterchangeSceneImportAsset::StaticClass();
+}
diff --git a/Engine/Plugins/Interchange/Runtime/Source/FactoryNodes/Public/InterchangeSceneImportAssetFactoryNode.h b/Engine/Plugins/Interchange/Runtime/Source/FactoryNodes/Public/InterchangeSceneImportAssetFactoryNode.h
new file mode 100644
--- /dev/null
+++ b/Engine/Plugins/Interchange/Runtime/Source/FactoryNodes/Public/InterchangeSceneImportAssetFactoryNode.h
@@ -0,0 +1,22 @@
+// Copyright Epic Games, Inc. All Rights Reserved.
+
+#pragma once
+
+#include "CoreMinimal.h"
+#include "InterchangeActorFactoryNode.h"
+
+#include "InterchangeSceneImportAssetFactoryNode.generated.h"
+
+UCLASS(BlueprintType)
+class INTERCHANGEFACTORYNODES_API UInterchangeSceneImportAssetFactoryNode : public UInterchangeFactoryBaseNode
+{
+	GENERATED_BODY()
+
+public:
+	// UInterchangeFactoryBaseNode Begin
+	virtual class UClass* GetObjectClass() const override;
+	// UInterchangeFactoryBaseNode End
+
+private:
+};
+
diff --git a/Engine/Plugins/Interchange/Runtime/Source/Import/Private/InterchangeImportModule.cpp b/Engine/Plugins/Interchange/Runtime/Source/Import/Private/InterchangeImportModule.cpp
--- a/Engine/Plugins/Interchange/Runtime/Source/Import/Private/InterchangeImportModule.cpp
+++ b/Engine/Plugins/Interchange/Runtime/Source/Import/Private/InterchangeImportModule.cpp
@@ -21,6 +21,7 @@
 #include "Modules/ModuleManager.h"
 #include "Scene/InterchangeActorFactory.h"
 #include "Scene/InterchangeCineCameraActorFactory.h"
+#include "Scene/InterchangeSceneImportAssetFactory.h"
 #include "Scene/InterchangeLightActorFactory.h"
 #include "Scene/InterchangeSceneVariantSetsFactory.h"
 #include "Scene/InterchangeStaticMeshActorFactory.h"
@@ -87,6 +88,7 @@ void FInterchangeImportModule::StartupModule()
 		InterchangeManager.RegisterFactory(UInterchangeSkeletalMeshActorFactory::StaticClass());
 		InterchangeManager.RegisterFactory(UInterchangeSceneVariantSetsFactory::StaticClass());
 		InterchangeManager.RegisterFactory(UInterchangeLightActorFactory::StaticClass());
+		InterchangeManager.RegisterFactory(UInterchangeSceneImportAssetFactory::StaticClass());
 	};
 	
 	if (GEngine)
diff --git a/Engine/Plugins/Interchange/Runtime/Source/Import/Private/Scene/InterchangeActorHelper.cpp b/Engine/Plugins/Interchange/Runtime/Source/Import/Private/Scene/InterchangeActorHelper.cpp
--- a/Engine/Plugins/Interchange/Runtime/Source/Import/Private/Scene/InterchangeActorHelper.cpp
+++ b/Engine/Plugins/Interchange/Runtime/Source/Import/Private/Scene/InterchangeActorHelper.cpp
@@ -144,11 +144,11 @@ AActor* UE::Interchange::ActorHelper::SpawnFactoryActor(const UInterchangeFactor
 				RootComponent->SetMobility(TargetMobility);
 			}
 		}
-	}
 
-	if (ParentActor || !SpawnedActor->IsAttachedTo(ParentActor))
-	{
-		SpawnedActor->AttachToActor(ParentActor, FAttachmentTransformRules::KeepWorldTransform);
+		if (ParentActor || !SpawnedActor->IsAttachedTo(ParentActor))
+		{
+			SpawnedActor->AttachToActor(ParentActor, FAttachmentTransformRules::KeepWorldTransform);
+		}
 	}
 
 	return SpawnedActor;
@@ -227,4 +227,3 @@ void UE::Interchange::ActorHelper::ApplyAllCustomAttributes(const UInterchangeFa
 		CreateSceneObjectsParams.FactoryNode->ApplyAllCustomAttributeToObject(&ObjectToUpdate);
 	}
 }
-#pragma optimize ("", on)
diff --git a/Engine/Plugins/Interchange/Runtime/Source/Import/Private/Scene/InterchangeSceneImportAssetFactory.cpp b/Engine/Plugins/Interchange/Runtime/Source/Import/Private/Scene/InterchangeSceneImportAssetFactory.cpp
new file mode 100644
--- /dev/null
+++ b/Engine/Plugins/Interchange/Runtime/Source/Import/Private/Scene/InterchangeSceneImportAssetFactory.cpp
@@ -0,0 +1,184 @@
+// Copyright Epic Games, Inc. All Rights Reserved.
+
+#include "Scene/InterchangeSceneImportAssetFactory.h"
+
+#include "InterchangeAssetImportData.h"
+#include "InterchangeImportCommon.h"
+#include "InterchangeImportLog.h"
+#include "InterchangeSceneImportAsset.h"
+#include "InterchangeSceneImportAssetFactoryNode.h"
+#include "InterchangeResult.h"
+#include "InterchangeSourceData.h"
+#include "InterchangeTranslatorBase.h"
+#include "InterchangeVariantSetNode.h"
+#include "Nodes/InterchangeBaseNode.h"
+#include "Nodes/InterchangeBaseNodeContainer.h"
+#include "UObject/Class.h"
+
+#define LOCTEXT_NAMESPACE "InterchangeSceneImportAssetFactory"
+
+namespace UE::Interchange::Private::InterchangeSceneImportAssetFactory
+{
+	const UInterchangeSceneImportAssetFactoryNode* GetFactoryNode(const UInterchangeFactoryBase::FImportAssetObjectParams& Arguments, UClass* TargetClass)
+	{
+		if (!Arguments.NodeContainer || !Arguments.AssetNode || !Arguments.AssetNode->GetObjectClass()->IsChildOf(TargetClass))
+		{
+			return nullptr;
+		}
+
+		const UInterchangeSceneImportAssetFactoryNode* FactoryNode = Cast<UInterchangeSceneImportAssetFactoryNode>(Arguments.AssetNode);
+		if (FactoryNode == nullptr)
+		{
+			return nullptr;
+		}
+
+		return FactoryNode;
+	}
+
+	UObject* FindOrCreateAsset(const UInterchangeFactoryBase::FImportAssetObjectParams& Arguments, UClass* TargetClass)
+	{
+		UObject* TargetAsset = Arguments.ReimportObject ? Arguments.ReimportObject : StaticFindObject(nullptr, Arguments.Parent, *Arguments.AssetName);
+
+		// Create a new asset or return existing asset, if possible
+		if (!TargetAsset)
+		{
+			check(IsInGameThread());
+			TargetAsset = NewObject<UObject>(Arguments.Parent, TargetClass, *Arguments.AssetName, RF_Public | RF_Standalone);
+		}
+		else if (!TargetAsset->GetClass()->IsChildOf(TargetClass))
+		{
+			TargetAsset = nullptr;
+		}
+
+		return TargetAsset;
+	}
+}
+
+UClass* UInterchangeSceneImportAssetFactory::GetFactoryClass() const
+{
+	return UInterchangeSceneImportAsset::StaticClass();
+}
+
+bool UInterchangeSceneImportAssetFactory::GetSourceFilenames(const UObject* Object, TArray<FString>& OutSourceFilenames) const
+{
+#if WITH_EDITORONLY_DATA
+	if (const UInterchangeSceneImportAsset* SceneImportAsset = Cast<UInterchangeSceneImportAsset>(Object))
+	{
+		return UE::Interchange::FFactoryCommon::GetSourceFilenames(SceneImportAsset->AssetImportData, OutSourceFilenames);
+	}
+#endif
+
+	return false;
+}
+
+bool UInterchangeSceneImportAssetFactory::SetSourceFilename(const UObject* Object, const FString& SourceFilename, int32 SourceIndex) const
+{
+#if WITH_EDITORONLY_DATA
+	if (const UInterchangeSceneImportAsset* SceneImportAsset = Cast<UInterchangeSceneImportAsset>(Object))
+	{
+		return UE::Interchange::FFactoryCommon::SetSourceFilename(SceneImportAsset->AssetImportData, SourceFilename, SourceIndex);
+	}
+#endif
+
+	return false;
+}
+
+UInterchangeFactoryBase::FImportAssetResult UInterchangeSceneImportAssetFactory::BeginImportAsset_GameThread(const FImportAssetObjectParams& Arguments)
+{
+	using namespace UE::Interchange::Private::InterchangeSceneImportAssetFactory;
+
+	UClass* TargetClass = GetFactoryClass();
+
+	if (GetFactoryNode(Arguments, TargetClass) == nullptr)
+	{
+		return {};
+	}
+
+	UInterchangeSceneImportAsset* SceneImportAsset = Cast<UInterchangeSceneImportAsset>(FindOrCreateAsset(Arguments, TargetClass));
+
+	if (!SceneImportAsset)
+	{
+		UE_LOG(LogInterchangeImport, Warning, TEXT("Could not create InterchangeSceneImportAsset asset %s"), *Arguments.AssetName);
+	}
+
+	FImportAssetResult Result;
+	Result.ImportedObject = SceneImportAsset;
+	
+	return Result;
+}
+
+UObject* UInterchangeSceneImportAssetFactory::ImportAssetObject_Async(const FImportAssetObjectParams& Arguments)
+{
+	using namespace UE::Interchange::Private::InterchangeSceneImportAssetFactory;
+
+	UClass* TargetClass = GetFactoryClass();
+
+	const UInterchangeSceneImportAssetFactoryNode* FactoryNode = GetFactoryNode(Arguments, TargetClass);
+	if (!FactoryNode)
+	{
+		return nullptr;
+	}
+
+	UInterchangeSceneImportAsset* SceneImportAsset = Cast<UInterchangeSceneImportAsset>(FindOrCreateAsset(Arguments, TargetClass));
+
+	if (!ensure(SceneImportAsset))
+	{
+		UInterchangeResultError_Generic* Message = AddMessage<UInterchangeResultError_Generic>();
+		Message->SourceAssetName = Arguments.SourceData->GetFilename();
+		Message->DestinationAssetName = Arguments.AssetName;
+		Message->AssetType = UInterchangeSceneImportAsset::StaticClass();
+		Message->Text = FText::Format(LOCTEXT("CreateAssetFailed", "Could not create nor find SceneImportAsset asset {0}."), FText::FromString(Arguments.AssetName));
+		return nullptr;
+	}
+
+	/** Apply all FactoryNode custom attributes to the level sequence asset */
+	FactoryNode->ApplyAllCustomAttributeToObject(SceneImportAsset);
+
+	//Getting the file Hash will cache it into the source data
+	Arguments.SourceData->GetFileContentHash();
+
+	return SceneImportAsset;
+}
+
+void UInterchangeSceneImportAssetFactory::SetupObject_GameThread(const FSetupObjectParams& Arguments)
+{
+	check(IsInGameThread());
+	Super::SetupObject_GameThread(Arguments);
+
+#if WITH_EDITORONLY_DATA
+	using namespace UE::Interchange;
+
+	if (ensure(Arguments.ImportedObject && Arguments.SourceData))
+	{
+		// We must call the Update of the asset source file in the main thread because UAssetImportData::Update execute some delegate we do not control
+		UInterchangeSceneImportAsset* SceneImportAsset = Cast<UInterchangeSceneImportAsset>(Arguments.ImportedObject);
+		if (ensure(SceneImportAsset))
+		{
+			if (Arguments.bIsReimport)
+			{
+				ensure(SceneImportAsset->AssetImportData);
+				SceneImportAsset->UpdateSceneObjects();
+			}
+
+			FFactoryCommon::FUpdateImportAssetDataParameters Parameters(SceneImportAsset, SceneImportAsset->AssetImportData, Arguments.SourceData, Arguments.NodeUniqueID, Arguments.NodeContainer, Arguments.OriginalPipelines);
+			SceneImportAsset->AssetImportData = Cast<UInterchangeAssetImportData>(FFactoryCommon::UpdateImportAssetData(Parameters));
+		}
+
+	}
+#endif
+}
+
+void UInterchangeSceneImportAssetFactory::FinalizeObject_GameThread(const FSetupObjectParams& Arguments)
+{
+	check(IsInGameThread());
+	Super::FinalizeObject_GameThread(Arguments);
+
+#if WITH_EDITOR
+	if (UInterchangeSceneImportAsset* SceneImportAsset = Cast<UInterchangeSceneImportAsset>(Arguments.ImportedObject))
+	{
+		SceneImportAsset->RegisterWorldRenameCallbacks();
+	}
+#endif
+}
+
+#undef LOCTEXT_NAMESPACE
\ No newline at end of file
diff --git a/Engine/Plugins/Interchange/Runtime/Source/Import/Public/Scene/InterchangeSceneImportAssetFactory.h b/Engine/Plugins/Interchange/Runtime/Source/Import/Public/Scene/InterchangeSceneImportAssetFactory.h
new file mode 100644
--- /dev/null
+++ b/Engine/Plugins/Interchange/Runtime/Source/Import/Public/Scene/InterchangeSceneImportAssetFactory.h
@@ -0,0 +1,33 @@
+// Copyright Epic Games, Inc. All Rights Reserved.
+
+#pragma once
+
+#include "CoreMinimal.h"
+#include "InterchangeFactoryBase.h"
+#include "UObject/Object.h"
+#include "UObject/ObjectMacros.h"
+
+#include "InterchangeSceneImportAssetFactory.generated.h"
+
+UCLASS(BlueprintType, Experimental)
+class INTERCHANGEIMPORT_API UInterchangeSceneImportAssetFactory : public UInterchangeFactoryBase
+{
+	GENERATED_BODY()
+public:
+
+	//////////////////////////////////////////////////////////////////////////
+	// Interchange factory base interface begin
+
+	virtual UClass* GetFactoryClass() const override;
+	virtual FImportAssetResult BeginImportAsset_GameThread(const FImportAssetObjectParams& Arguments) override;
+	virtual UObject* ImportAssetObject_Async(const FImportAssetObjectParams& Arguments) override;
+	virtual void SetupObject_GameThread(const FSetupObjectParams& Arguments) override;
+	virtual void FinalizeObject_GameThread(const FSetupObjectParams& Arguments) override;
+	virtual bool GetSourceFilenames(const UObject* Object, TArray<FString>& OutSourceFilenames) const override;
+	virtual bool SetSourceFilename(const UObject* Object, const FString& SourceFilename, int32 SourceIndex) const override;
+
+	// Interchange factory base interface end
+	//////////////////////////////////////////////////////////////////////////
+};
+
+
diff --git a/Engine/Plugins/Interchange/Runtime/Source/Pipelines/Private/InterchangeGenericScenesPipeline.cpp b/Engine/Plugins/Interchange/Runtime/Source/Pipelines/Private/InterchangeGenericScenesPipeline.cpp
--- a/Engine/Plugins/Interchange/Runtime/Source/Pipelines/Private/InterchangeGenericScenesPipeline.cpp
+++ b/Engine/Plugins/Interchange/Runtime/Source/Pipelines/Private/InterchangeGenericScenesPipeline.cpp
@@ -6,6 +6,8 @@
 #include "InterchangeCameraNode.h"
 #include "InterchangeCineCameraFactoryNode.h"
 #include "InterchangeCommonPipelineDataFactoryNode.h"
+#include "InterchangeSceneImportAsset.h"
+#include "InterchangeSceneImportAssetFactoryNode.h"
 #include "InterchangeLightNode.h"
 #include "InterchangeLightFactoryNode.h"
 #include "InterchangeMeshActorFactoryNode.h"
@@ -99,6 +101,33 @@ void UInterchangeGenericLevelPipeline::ExecutePipeline(UInterchangeBaseNodeConta
 		}
 	});
 
+#if WITH_EDITORONLY_DATA
+	// Add the SceneImportData factory node
+	{
+		ensure(!SceneImportFactoryNode);
+
+		const FString FilePath = FPaths::ConvertRelativePathToFull(InSourceDatas[0]->GetFilename());
+		const FString DisplayLabel = TEXT("SceneImport_") + FPaths::GetBaseFilename(FilePath);
+		const FString NodeUid = TEXT("SceneImport_") + FilePath;
+		const FString FactoryNodeUid = UInterchangeFactoryBaseNode::BuildFactoryNodeUid(NodeUid);
+		ensure(!BaseNodeContainer->IsNodeUidValid(FactoryNodeUid));
+		SceneImportFactoryNode = NewObject<UInterchangeSceneImportAssetFactoryNode>(BaseNodeContainer, NAME_None);
+
+		SceneImportFactoryNode->InitializeNode(FactoryNodeUid, DisplayLabel, EInterchangeNodeContainerType::FactoryData);
+
+		// Add dependency to all the factory nodes created so far
+		BaseNodeContainer->IterateNodesOfType<UInterchangeFactoryBaseNode>(
+			[this](const FString& NodeUid, UInterchangeFactoryBaseNode* Node)
+			{
+				// Make sure the UInterchangeSceneImportAsset is the last asset created
+				this->SceneImportFactoryNode->AddFactoryDependencyUid(NodeUid);
+			}
+		);
+
+		BaseNodeContainer->AddNode(SceneImportFactoryNode);
+	}
+#endif
+
 	for (const UInterchangeSceneNode* SceneNode : SceneNodes)
 	{
 		if (SceneNode)
@@ -112,7 +141,7 @@ void UInterchangeGenericLevelPipeline::ExecutePipeline(UInterchangeBaseNodeConta
 					bool bSkipNode = true;
 					if (SpecializeTypes.Contains(UE::Interchange::FSceneNodeStaticData::GetJointSpecializeTypeString()))
 					{
-						//check if its the rootjoint (we want to create an actor for the rootjoint)
+						//check if its the root joint (we want to create an actor for the root joint)
 						FString CurrentNodesParentUid = SceneNode->GetParentUid();
 						const UInterchangeBaseNode* ParentNode = BaseNodeContainer->GetNode(CurrentNodesParentUid);
 						if (const UInterchangeSceneNode* ParentSceneNode = Cast<UInterchangeSceneNode>(ParentNode))
@@ -138,19 +167,19 @@ void UInterchangeGenericLevelPipeline::ExecutePipeline(UInterchangeBaseNodeConta
 	//Find all translated scene variant sets
 	TArray<UInterchangeSceneVariantSetsNode*> SceneVariantSetNodes;
 
-InBaseNodeContainer->IterateNodesOfType<UInterchangeSceneVariantSetsNode>([&SceneVariantSetNodes](const FString& NodeUid, UInterchangeSceneVariantSetsNode* Node)
-	{
-		SceneVariantSetNodes.Add(Node);
-	});
+	InBaseNodeContainer->IterateNodesOfType<UInterchangeSceneVariantSetsNode>([&SceneVariantSetNodes](const FString& NodeUid, UInterchangeSceneVariantSetsNode* Node)
+		{
+			SceneVariantSetNodes.Add(Node);
+		});
 
-for (const UInterchangeSceneVariantSetsNode* SceneVariantSetNode : SceneVariantSetNodes)
-{
-	if (SceneVariantSetNode)
+	for (const UInterchangeSceneVariantSetsNode* SceneVariantSetNode : SceneVariantSetNodes)
 	{
-		ExecuteSceneVariantSetNodePreImport(*SceneVariantSetNode);
+		if (SceneVariantSetNode)
+		{
+			ExecuteSceneVariantSetNodePreImport(*SceneVariantSetNode);
+		}
 	}
 }
-}
 
 void UInterchangeGenericLevelPipeline::ExecuteSceneNodePreImport(const FTransform& GlobalOffsetTransform, const UInterchangeSceneNode* SceneNode)
 {
@@ -265,6 +294,14 @@ void UInterchangeGenericLevelPipeline::ExecuteSceneNodePreImport(const FTransfor
 	{
 		SetUpFactoryNode(ActorFactoryNode, SceneNode, TranslatedAssetNode);
 	}
+
+	//Make sure all actor factory nodes have the specified strategy
+	ActorFactoryNode->SetReimportStrategyFlags(ReimportStrategy);
+
+#if WITH_EDITORONLY_DATA
+	// Add dependency to newly created factory node
+	SceneImportFactoryNode->AddFactoryDependencyUid(ActorFactoryNode->GetUniqueID());
+#endif
 }
 
 UInterchangeActorFactoryNode* UInterchangeGenericLevelPipeline::CreateActorFactoryNode(const UInterchangeSceneNode* SceneNode, const UInterchangeBaseNode* TranslatedAssetNode) const
@@ -560,3 +597,29 @@ void UInterchangeGenericLevelPipeline::ExecuteSceneVariantSetNodePreImport(const
 
 	BaseNodeContainer->AddNode(FactoryNode);
 }
+
+void UInterchangeGenericLevelPipeline::ExecutePostImportPipeline(const UInterchangeBaseNodeContainer* InBaseNodeContainer, const FString& NodeKey, UObject* CreatedAsset, bool bIsAReimport)
+{
+	Super::ExecutePostImportPipeline(InBaseNodeContainer, NodeKey, CreatedAsset, bIsAReimport);
+
+#if WITH_EDITORONLY_DATA
+	//We do not use the provided base container since ExecutePreImportPipeline cache it
+	//We just make sure the same one is pass in parameter
+	if (!InBaseNodeContainer || !ensure(BaseNodeContainer == InBaseNodeContainer) || !CreatedAsset)
+	{
+		return;
+	}
+
+	if (UInterchangeSceneImportAsset* SceneImportAsset = Cast<UInterchangeSceneImportAsset>(CreatedAsset))
+	{
+		const UInterchangeSceneImportAssetFactoryNode* FactoryNode = Cast<UInterchangeSceneImportAssetFactoryNode>(BaseNodeContainer->GetFactoryNode(NodeKey));
+		if (!ensure(FactoryNode))
+		{
+			return;
+		}
+
+		SceneImportAsset->UpdateSceneObjects();
+	}
+#endif
+}
+
diff --git a/Engine/Plugins/Interchange/Runtime/Source/Pipelines/Public/InterchangeGenericScenesPipeline.h b/Engine/Plugins/Interchange/Runtime/Source/Pipelines/Public/InterchangeGenericScenesPipeline.h
--- a/Engine/Plugins/Interchange/Runtime/Source/Pipelines/Public/InterchangeGenericScenesPipeline.h
+++ b/Engine/Plugins/Interchange/Runtime/Source/Pipelines/Public/InterchangeGenericScenesPipeline.h
@@ -14,22 +14,28 @@
 class UInterchangeActorFactoryNode;
 class UInterchangeSceneNode;
 class UInterchangeSceneVariantSetsNode;
-
-
+class UInterchangeSceneImportAssetFactoryNode;
 
 UCLASS(BlueprintType, editinlinenew)
 class INTERCHANGEPIPELINES_API UInterchangeGenericLevelPipeline : public UInterchangePipelineBase
 {
 	GENERATED_BODY()
+public:
+	/* Allow user to choose the re-import strategy when importing into level. */
+	UPROPERTY(EditAnywhere, BlueprintReadWrite, Category = "Actors", meta = (AdjustPipelineAndRefreshDetailOnChange = "True"))
+	EReimportStrategyFlags ReimportStrategy = EReimportStrategyFlags::ApplyNoProperties;
 
 protected:
 
+	/** BEGIN UInterchangePipelineBase overrides */
 	virtual void ExecutePipeline(UInterchangeBaseNodeContainer* InBaseNodeContainer, const TArray<UInterchangeSourceData*>& InSourceDatas) override;
+	virtual void ExecutePostImportPipeline(const UInterchangeBaseNodeContainer* BaseNodeContainer, const FString& NodeKey, UObject* CreatedAsset, bool bIsAReimport) override;
 
 	virtual bool CanExecuteOnAnyThread(EInterchangePipelineTask PipelineTask) override
 	{
 		return true;
 	}
+	/** END UInterchangePipelineBase overrides */
 
 	/**
 	 * PreImport step called for each translated SceneNode.
@@ -53,6 +59,16 @@ protected:
 	
 protected:
 	UInterchangeBaseNodeContainer* BaseNodeContainer = nullptr;
+#if WITH_EDITORONLY_DATA
+	/*
+	 * Factory node created by the pipeline to later create the SceneImportAsset
+	 * This factory node must be unique and depends on all the other factory nodes
+	 * and its factory must be called after all factories
+	 * Note that this factory node is not created at runtime. Therefore, the reimport
+	 * of a level will not work at runtime
+	 */
+	UInterchangeSceneImportAssetFactoryNode* SceneImportFactoryNode = nullptr;
+#endif
 };
 
 
diff --git a/Engine/Source/Runtime/Interchange/Engine/Private/InterchangeImportCommon.cpp b/Engine/Source/Runtime/Interchange/Engine/Private/InterchangeImportCommon.cpp
--- a/Engine/Source/Runtime/Interchange/Engine/Private/InterchangeImportCommon.cpp
+++ b/Engine/Source/Runtime/Interchange/Engine/Private/InterchangeImportCommon.cpp
@@ -6,6 +6,7 @@
 #include "InterchangeAssetImportData.h"
 #include "InterchangePipelineBase.h"
 #include "InterchangePythonPipelineBase.h"
+#include "InterchangeSceneImportAsset.h"
 #include "InterchangeSourceData.h"
 #include "Nodes/InterchangeBaseNodeContainer.h"
 #include "Nodes/InterchangeFactoryBaseNode.h"
@@ -13,304 +14,325 @@
 #include "EditorFramework/AssetImportData.h"
 #include "HAL/FileManager.h"
 
-namespace UE
+namespace UE::Interchange
 {
-	namespace Interchange
+	namespace Private::ImportCommon
 	{
-		namespace Private::ImportCommon
+		UInterchangeAssetImportData* BeginSetupAssetData(FFactoryCommon::FUpdateImportAssetDataParameters& Parameters)
 		{
-			UInterchangeAssetImportData* BeginSetupAssetData(FFactoryCommon::FUpdateImportAssetDataParameters& Parameters)
+			TRACE_CPUPROFILER_EVENT_SCOPE("UE::Interchange::Private::ImportCommon::BeginSetupAssetData")
+			if (!ensure(IsInGameThread()))
 			{
-				TRACE_CPUPROFILER_EVENT_SCOPE("UE::Interchange::Private::ImportCommon::BeginSetupAssetData")
-				if (!ensure(IsInGameThread()))
-				{
-					return nullptr;
-				}
-				if (!ensure(Parameters.SourceData && Parameters.AssetImportDataOuter))
-				{
-					return nullptr;
-				}
+				return nullptr;
+			}
+			if (!ensure(Parameters.SourceData && Parameters.AssetImportDataOuter))
+			{
+				return nullptr;
+			}
 
-				UInterchangeAssetImportData* AssetImportData = Cast<UInterchangeAssetImportData>(Parameters.AssetImportData);
+			UInterchangeAssetImportData* AssetImportData = Cast<UInterchangeAssetImportData>(Parameters.AssetImportData);
 
-				if (!AssetImportData)
-				{
-					AssetImportData = NewObject<UInterchangeAssetImportData>(Parameters.AssetImportDataOuter, NAME_None);
-				}
+			if (!AssetImportData)
+			{
+				AssetImportData = NewObject<UInterchangeAssetImportData>(Parameters.AssetImportDataOuter, NAME_None);
+			}
 
-				ensure(AssetImportData);
+			ensure(AssetImportData);
 
-				return AssetImportData;
-			}
+			return AssetImportData;
+		}
 
-			void RecursivelyDuplicateFactoryNodeDependencies(UInterchangeBaseNodeContainer* SourceContainer, UInterchangeBaseNodeContainer* DestinationContainer, TArray<FString>& FactoryDependencies)
+		void RecursivelyDuplicateFactoryNodeDependencies(UInterchangeBaseNodeContainer* SourceContainer, UInterchangeBaseNodeContainer* DestinationContainer, TArray<FString>& FactoryDependencies)
+		{
+			for (const FString& DependencyUid : FactoryDependencies)
 			{
-				for (const FString& DependencyUid : FactoryDependencies)
+				if (UInterchangeFactoryBaseNode* FactoryNode = Cast<UInterchangeFactoryBaseNode>(SourceContainer->GetFactoryNode(DependencyUid)))
 				{
-					if (UInterchangeFactoryBaseNode* FactoryNode = Cast<UInterchangeFactoryBaseNode>(SourceContainer->GetFactoryNode(DependencyUid)))
-					{
-						FObjectDuplicationParameters DupParam(FactoryNode, DestinationContainer);
-						DestinationContainer->AddNode(CastChecked<UInterchangeBaseNode>(StaticDuplicateObjectEx(DupParam)));
-
-						TArray<FString> ChildFactoryDependencies;
-						FactoryNode->GetFactoryDependencies(ChildFactoryDependencies);
-						RecursivelyDuplicateFactoryNodeDependencies(SourceContainer, DestinationContainer, ChildFactoryDependencies);
-					}
+					FObjectDuplicationParameters DupParam(FactoryNode, DestinationContainer);
+					DestinationContainer->AddNode(CastChecked<UInterchangeBaseNode>(StaticDuplicateObjectEx(DupParam)));
+
+					TArray<FString> ChildFactoryDependencies;
+					FactoryNode->GetFactoryDependencies(ChildFactoryDependencies);
+					RecursivelyDuplicateFactoryNodeDependencies(SourceContainer, DestinationContainer, ChildFactoryDependencies);
 				}
 			}
+		}
 
-			void EndSetupAssetData(FFactoryCommon::FUpdateImportAssetDataParameters& Parameters, UInterchangeAssetImportData* AssetImportData)
+		void EndSetupAssetData(FFactoryCommon::FUpdateImportAssetDataParameters& Parameters, UInterchangeAssetImportData* AssetImportData)
+		{
+			TRACE_CPUPROFILER_EVENT_SCOPE("UE::Interchange::Private::ImportCommon::EndSetupAssetData")
+			AssetImportData->NodeContainer = NewObject<UInterchangeBaseNodeContainer>(AssetImportData);
+			//We copy only the factory node dependencies, we use this only 
+			if (UInterchangeFactoryBaseNode* FactoryNode = Cast<UInterchangeFactoryBaseNode>(Parameters.NodeContainer->GetFactoryNode(Parameters.NodeUniqueID)))
 			{
-				TRACE_CPUPROFILER_EVENT_SCOPE("UE::Interchange::Private::ImportCommon::EndSetupAssetData")
-				AssetImportData->NodeContainer = NewObject<UInterchangeBaseNodeContainer>(AssetImportData);
-				//We copy only the factory node dependencies, we use this only 
-				if (UInterchangeFactoryBaseNode* FactoryNode = Cast<UInterchangeFactoryBaseNode>(Parameters.NodeContainer->GetFactoryNode(Parameters.NodeUniqueID)))
-				{
-					//Set the interchange node graph data
-					AssetImportData->NodeUniqueID = Parameters.NodeUniqueID;
-					TArray<FString> FactoryDependencies;
-					FactoryDependencies.Add(AssetImportData->NodeUniqueID);
-					RecursivelyDuplicateFactoryNodeDependencies(Parameters.NodeContainer, AssetImportData->NodeContainer, FactoryDependencies);
-				}
-				AssetImportData->Pipelines.Reset();
-				for (const UObject* Pipeline : Parameters.Pipelines)
+				//Set the interchange node graph data
+				AssetImportData->NodeUniqueID = Parameters.NodeUniqueID;
+				TArray<FString> FactoryDependencies;
+				FactoryDependencies.Add(AssetImportData->NodeUniqueID);
+				RecursivelyDuplicateFactoryNodeDependencies(Parameters.NodeContainer, AssetImportData->NodeContainer, FactoryDependencies);
+			}
+			AssetImportData->Pipelines.Reset();
+			for (const UObject* Pipeline : Parameters.Pipelines)
+			{
+				UObject* DupPipeline = Cast<UObject>(StaticDuplicateObject(Pipeline, AssetImportData));
+				if (DupPipeline)
 				{
-					UObject* DupPipeline = Cast<UObject>(StaticDuplicateObject(Pipeline, AssetImportData));
-					if (DupPipeline)
-					{
-						AssetImportData->Pipelines.Add(DupPipeline);
-					}
+					AssetImportData->Pipelines.Add(DupPipeline);
 				}
 			}
 		}
+	}
 
-		FFactoryCommon::FUpdateImportAssetDataParameters::FUpdateImportAssetDataParameters(UObject* InAssetImportDataOuter
-																							, UAssetImportData* InAssetImportData
-																							, const UInterchangeSourceData* InSourceData
-																							, FString InNodeUniqueID
-																							, UInterchangeBaseNodeContainer* InNodeContainer
-																						   , const TArray<UObject*>& InPipelines)
-			: AssetImportDataOuter(InAssetImportDataOuter)
-			, AssetImportData(InAssetImportData)
-			, SourceData(InSourceData)
-			, NodeUniqueID(InNodeUniqueID)
-			, NodeContainer(InNodeContainer)
-			, Pipelines(InPipelines)
-		{
-			ensure(AssetImportDataOuter);
-			ensure(SourceData);
-			ensure(!NodeUniqueID.IsEmpty());
-			ensure(NodeContainer);
-		}
+	FFactoryCommon::FUpdateImportAssetDataParameters::FUpdateImportAssetDataParameters(UObject* InAssetImportDataOuter
+																						, UAssetImportData* InAssetImportData
+																						, const UInterchangeSourceData* InSourceData
+																						, FString InNodeUniqueID
+																						, UInterchangeBaseNodeContainer* InNodeContainer
+																						, const TArray<UObject*>& InPipelines)
+		: AssetImportDataOuter(InAssetImportDataOuter)
+		, AssetImportData(InAssetImportData)
+		, SourceData(InSourceData)
+		, NodeUniqueID(InNodeUniqueID)
+		, NodeContainer(InNodeContainer)
+		, Pipelines(InPipelines)
+	{
+		ensure(AssetImportDataOuter);
+		ensure(SourceData);
+		ensure(!NodeUniqueID.IsEmpty());
+		ensure(NodeContainer);
+	}
 
-		UAssetImportData* FFactoryCommon::UpdateImportAssetData(FUpdateImportAssetDataParameters& Parameters)
-		{
-			TRACE_CPUPROFILER_EVENT_SCOPE("UE::Interchange::FFactoryCommon::UpdateImportAssetData")
-			return UpdateImportAssetData(Parameters, [&Parameters](UInterchangeAssetImportData* AssetImportData)
-				{
+	UAssetImportData* FFactoryCommon::UpdateImportAssetData(FUpdateImportAssetDataParameters& Parameters)
+	{
+		TRACE_CPUPROFILER_EVENT_SCOPE("UE::Interchange::FFactoryCommon::UpdateImportAssetData")
+		return UpdateImportAssetData(Parameters, [&Parameters](UInterchangeAssetImportData* AssetImportData)
+			{
 #if WITH_EDITORONLY_DATA
-					//Set the asset import data file source to allow reimport. TODO: manage MD5 Hash properly
-					TOptional<FMD5Hash> FileContentHash = Parameters.SourceData->GetFileContentHash();
+				//Set the asset import data file source to allow reimport. TODO: manage MD5 Hash properly
+				TOptional<FMD5Hash> FileContentHash = Parameters.SourceData->GetFileContentHash();
 
-					//Update the first filename, TODO for asset using multiple source file we have to update the correct index
-					AssetImportData->Update(Parameters.SourceData->GetFilename(), FileContentHash.IsSet() ? &FileContentHash.GetValue() : nullptr);
+				//Update the first filename, TODO for asset using multiple source file we have to update the correct index
+				AssetImportData->Update(Parameters.SourceData->GetFilename(), FileContentHash.IsSet() ? &FileContentHash.GetValue() : nullptr);
 #endif //WITH_EDITORONLY_DATA
-				});
+			});
 
-		}
+	}
 
-		UAssetImportData* FFactoryCommon::UpdateImportAssetData(FUpdateImportAssetDataParameters& Parameters, TFunctionRef<void(UInterchangeAssetImportData*)> CustomFileSourceUpdate)
-		{
-			TRACE_CPUPROFILER_EVENT_SCOPE("UE::Interchange::FFactoryCommon::UpdateImportAssetData2")
+	UAssetImportData* FFactoryCommon::UpdateImportAssetData(FUpdateImportAssetDataParameters& Parameters, TFunctionRef<void(UInterchangeAssetImportData*)> CustomFileSourceUpdate)
+	{
+		TRACE_CPUPROFILER_EVENT_SCOPE("UE::Interchange::FFactoryCommon::UpdateImportAssetData2")
 #if WITH_EDITORONLY_DATA
-			if (!ensure(IsInGameThread()))
-			{
-				return nullptr;
-			}
-			if (!ensure(Parameters.SourceData && Parameters.AssetImportDataOuter))
-			{
-				return nullptr;
-			}
+		if (!ensure(IsInGameThread()))
+		{
+			return nullptr;
+		}
+		if (!ensure(Parameters.SourceData && Parameters.AssetImportDataOuter))
+		{
+			return nullptr;
+		}
 
-			UInterchangeAssetImportData* AssetImportData = Private::ImportCommon::BeginSetupAssetData(Parameters);
+		UInterchangeAssetImportData* AssetImportData = Private::ImportCommon::BeginSetupAssetData(Parameters);
 
-			if (Parameters.AssetImportData) //-V1051
+		if (Parameters.AssetImportData) //-V1051
+		{
+			if (!Parameters.AssetImportData->IsA<UInterchangeAssetImportData>())
 			{
-				if (!Parameters.AssetImportData->IsA<UInterchangeAssetImportData>())
-				{
-					// Migrate the old source data
-					TArray<FAssetImportInfo::FSourceFile> OldSourceFiles = Parameters.AssetImportData->SourceData.SourceFiles;
-					AssetImportData->SetSourceFiles(MoveTemp(OldSourceFiles));
-				}
+				// Migrate the old source data
+				TArray<FAssetImportInfo::FSourceFile> OldSourceFiles = Parameters.AssetImportData->SourceData.SourceFiles;
+				AssetImportData->SetSourceFiles(MoveTemp(OldSourceFiles));
 			}
+		}
 
-			CustomFileSourceUpdate(AssetImportData);
+		CustomFileSourceUpdate(AssetImportData);
 
 
-			Private::ImportCommon::EndSetupAssetData(Parameters, AssetImportData);
+		Private::ImportCommon::EndSetupAssetData(Parameters, AssetImportData);
 
-			// Return the asset import data so it can be set on the imported asset.
-			return AssetImportData;
+		// Return the asset import data so it can be set on the imported asset.
+		return AssetImportData;
 #endif //#if WITH_EDITORONLY_DATA
-			return nullptr;
-		}
+		return nullptr;
+	}
 
 #if WITH_EDITORONLY_DATA
-		FFactoryCommon::FSetImportAssetDataParameters::FSetImportAssetDataParameters(UObject* InAssetImportDataOuter
-																					, UAssetImportData* InAssetImportData
-																					, const UInterchangeSourceData* InSourceData
-																					, FString InNodeUniqueID
-																					, UInterchangeBaseNodeContainer* InNodeContainer
-																					, const TArray<UObject*>& InPipelines)
-			: FUpdateImportAssetDataParameters(InAssetImportDataOuter
-				, InAssetImportData
-				, InSourceData
-				, InNodeUniqueID
-				, InNodeContainer
-				, InPipelines
-				)
-			, SourceFiles()
-		{
-		}
+	FFactoryCommon::FSetImportAssetDataParameters::FSetImportAssetDataParameters(UObject* InAssetImportDataOuter
+																				, UAssetImportData* InAssetImportData
+																				, const UInterchangeSourceData* InSourceData
+																				, FString InNodeUniqueID
+																				, UInterchangeBaseNodeContainer* InNodeContainer
+																				, const TArray<UObject*>& InPipelines)
+		: FUpdateImportAssetDataParameters(InAssetImportDataOuter
+			, InAssetImportData
+			, InSourceData
+			, InNodeUniqueID
+			, InNodeContainer
+			, InPipelines
+			)
+		, SourceFiles()
+	{
+	}
 
-		UAssetImportData* FFactoryCommon::SetImportAssetData(FSetImportAssetDataParameters& Parameters)
-		{
-			TRACE_CPUPROFILER_EVENT_SCOPE("UE::Interchange::FFactoryCommon::SetImportAssetData")
-			UInterchangeAssetImportData* AssetImportData = Private::ImportCommon::BeginSetupAssetData(Parameters);
+	UAssetImportData* FFactoryCommon::SetImportAssetData(FSetImportAssetDataParameters& Parameters)
+	{
+		TRACE_CPUPROFILER_EVENT_SCOPE("UE::Interchange::FFactoryCommon::SetImportAssetData")
+		UInterchangeAssetImportData* AssetImportData = Private::ImportCommon::BeginSetupAssetData(Parameters);
 
-			// Update the source files
+		// Update the source files
+		{
+			if (Parameters.SourceFiles.IsEmpty())
 			{
-				if (Parameters.SourceFiles.IsEmpty())
-				{
-					TOptional<FMD5Hash> FileContentHash = Parameters.SourceData->GetFileContentHash();
-
-					// Todo add display label?
-					Parameters.SourceFiles.Emplace(
-						AssetImportData->SanitizeImportFilename(Parameters.SourceData->GetFilename()),
-						IFileManager::Get().GetTimeStamp(*Parameters.SourceData->GetFilename()),
-						FileContentHash.IsSet() ? FileContentHash.GetValue() : FMD5Hash()
-						);
-				}
-				else
+				TOptional<FMD5Hash> FileContentHash = Parameters.SourceData->GetFileContentHash();
+
+				// Todo add display label?
+				Parameters.SourceFiles.Emplace(
+					AssetImportData->SanitizeImportFilename(Parameters.SourceData->GetFilename()),
+					IFileManager::Get().GetTimeStamp(*Parameters.SourceData->GetFilename()),
+					FileContentHash.IsSet() ? FileContentHash.GetValue() : FMD5Hash()
+					);
+			}
+			else
+			{
+				for (FAssetImportInfo::FSourceFile& Source : Parameters.SourceFiles)
 				{
-					for (FAssetImportInfo::FSourceFile& Source : Parameters.SourceFiles)
-					{
-						// This is done here since this is not thread safe
-						Source.RelativeFilename = AssetImportData->SanitizeImportFilename(Source.RelativeFilename);
-					}
+					// This is done here since this is not thread safe
+					Source.RelativeFilename = AssetImportData->SanitizeImportFilename(Source.RelativeFilename);
 				}
-
-				AssetImportData->SetSourceFiles(MoveTemp(Parameters.SourceFiles));
 			}
 
-			Private::ImportCommon::EndSetupAssetData(Parameters, AssetImportData);
-
-			// Return the asset import data so it can be set on the imported asset.
-			return AssetImportData;
+			AssetImportData->SetSourceFiles(MoveTemp(Parameters.SourceFiles));
 		}
 
-		bool FFactoryCommon::GetSourceFilenames(const UAssetImportData* AssetImportData, TArray<FString>& OutSourceFilenames)
+		Private::ImportCommon::EndSetupAssetData(Parameters, AssetImportData);
+
+		// Return the asset import data so it can be set on the imported asset.
+		return AssetImportData;
+	}
+
+	bool FFactoryCommon::GetSourceFilenames(const UAssetImportData* AssetImportData, TArray<FString>& OutSourceFilenames)
+	{
+		TRACE_CPUPROFILER_EVENT_SCOPE("UE::Interchange::FFactoryCommon::GetSourceFilenames")
+		if (Cast<UInterchangeAssetImportData>(AssetImportData) != nullptr)
 		{
-			TRACE_CPUPROFILER_EVENT_SCOPE("UE::Interchange::FFactoryCommon::GetSourceFilenames")
-			if (Cast<UInterchangeAssetImportData>(AssetImportData) != nullptr)
-			{
-				AssetImportData->ExtractFilenames(OutSourceFilenames);
-				return true;
-			}
-			return false;
+			AssetImportData->ExtractFilenames(OutSourceFilenames);
+			return true;
 		}
+		return false;
+	}
 
-		bool FFactoryCommon::SetSourceFilename(UAssetImportData* AssetImportData, const FString& SourceFilename, int32 SourceIndex, const FString& SourceLabel)
+	bool FFactoryCommon::SetSourceFilename(UAssetImportData* AssetImportData, const FString& SourceFilename, int32 SourceIndex, const FString& SourceLabel)
+	{
+		TRACE_CPUPROFILER_EVENT_SCOPE("UE::Interchange::FFactoryCommon::SetSourceFilename")
+		if (AssetImportData)
 		{
-			TRACE_CPUPROFILER_EVENT_SCOPE("UE::Interchange::FFactoryCommon::SetSourceFilename")
-			if (AssetImportData)
+			const int32 SafeSourceIndex = SourceIndex == INDEX_NONE ? 0 : SourceIndex;				
+			if (SafeSourceIndex < AssetImportData->GetSourceFileCount())
 			{
-				const int32 SafeSourceIndex = SourceIndex == INDEX_NONE ? 0 : SourceIndex;				
-				if (SafeSourceIndex < AssetImportData->GetSourceFileCount())
-				{
-					AssetImportData->UpdateFilenameOnly(SourceFilename, SafeSourceIndex);
-				}
-				else
-				{
-					//Create a source file entry, this case happen when user import a specific content for the first time
-					AssetImportData->AddFileName(SourceFilename, SafeSourceIndex, SourceLabel);
-				}
-				return true;
+				AssetImportData->UpdateFilenameOnly(SourceFilename, SafeSourceIndex);
 			}
-
-			return false;
+			else
+			{
+				//Create a source file entry, this case happen when user import a specific content for the first time
+				AssetImportData->AddFileName(SourceFilename, SafeSourceIndex, SourceLabel);
+			}
+			return true;
 		}
+
+		return false;
+	}
 		
-		bool FFactoryCommon::SetReimportSourceIndex(const UObject* Object, UAssetImportData* AssetImportData, int32 SourceIndex)
+	bool FFactoryCommon::SetReimportSourceIndex(const UObject* Object, UAssetImportData* AssetImportData, int32 SourceIndex)
+	{
+		TRACE_CPUPROFILER_EVENT_SCOPE("UE::Interchange::FFactoryCommon::SetReimportSourceIndex")
+		UInterchangeAssetImportData* InterchangeAssetImportData = Cast<UInterchangeAssetImportData>(AssetImportData);
+		if (!InterchangeAssetImportData)
 		{
-			TRACE_CPUPROFILER_EVENT_SCOPE("UE::Interchange::FFactoryCommon::SetReimportSourceIndex")
-			UInterchangeAssetImportData* InterchangeAssetImportData = Cast<UInterchangeAssetImportData>(AssetImportData);
-			if (!InterchangeAssetImportData)
-			{
-				return false;
-			}
+			return false;
+		}
 
-			for (UObject* PipelineObject : InterchangeAssetImportData->Pipelines)
+		for (UObject* PipelineObject : InterchangeAssetImportData->Pipelines)
+		{
+			if (UInterchangePythonPipelineAsset* PythonPipelineAsset = Cast<UInterchangePythonPipelineAsset>(PipelineObject))
 			{
-				if (UInterchangePythonPipelineAsset* PythonPipelineAsset = Cast<UInterchangePythonPipelineAsset>(PipelineObject))
+				if (PythonPipelineAsset->GeneratedPipeline)
 				{
-					if (PythonPipelineAsset->GeneratedPipeline)
-					{
-						PythonPipelineAsset->GeneratedPipeline->ScriptedSetReimportSourceIndex(Object->GetClass(), SourceIndex);
-						PythonPipelineAsset->SetupFromPipeline(PythonPipelineAsset->GeneratedPipeline);
-					}
-				}
-				else if (UInterchangePipelineBase* PipelineBase = Cast<UInterchangePipelineBase>(PipelineObject))
-				{
-					PipelineBase->ScriptedSetReimportSourceIndex(Object->GetClass(), SourceIndex);
+					PythonPipelineAsset->GeneratedPipeline->ScriptedSetReimportSourceIndex(Object->GetClass(), SourceIndex);
+					PythonPipelineAsset->SetupFromPipeline(PythonPipelineAsset->GeneratedPipeline);
 				}
 			}
-			return true;
+			else if (UInterchangePipelineBase* PipelineBase = Cast<UInterchangePipelineBase>(PipelineObject))
+			{
+				PipelineBase->ScriptedSetReimportSourceIndex(Object->GetClass(), SourceIndex);
+			}
 		}
+		return true;
+	}
 
 #endif //#if WITH_EDITORONLY_DATA
 
 
-		void FFactoryCommon::ApplyReimportStrategyToAsset(UObject* Asset
-										  , const UInterchangeFactoryBaseNode* PreviousAssetNode
-										  , const UInterchangeFactoryBaseNode* CurrentAssetNode
-										  , UInterchangeFactoryBaseNode* PipelineAssetNode)
+	void FFactoryCommon::ApplyReimportStrategyToAsset(UObject* Asset
+										, const UInterchangeFactoryBaseNode* PreviousAssetNode
+										, const UInterchangeFactoryBaseNode* CurrentAssetNode
+										, UInterchangeFactoryBaseNode* PipelineAssetNode)
+	{
+		TRACE_CPUPROFILER_EVENT_SCOPE("UE::Interchange::FFactoryCommon::ApplyReimportStrategyToAsset")
+		if (!ensure(PreviousAssetNode) || !ensure(PipelineAssetNode) || !ensure(CurrentAssetNode))
+		{
+			return;
+		}
+		EReimportStrategyFlags ReimportStrategyFlags = PipelineAssetNode->GetReimportStrategyFlags();
+		switch (ReimportStrategyFlags)
 		{
-			TRACE_CPUPROFILER_EVENT_SCOPE("UE::Interchange::FFactoryCommon::ApplyReimportStrategyToAsset")
-			if (!ensure(PreviousAssetNode) || !ensure(PipelineAssetNode) || !ensure(CurrentAssetNode))
+			case EReimportStrategyFlags::ApplyNoProperties:
 			{
-				return;
+				//We want to have no effect
+				break;
 			}
-			EReimportStrategyFlags ReimportStrategyFlags = PipelineAssetNode->GetReimportStrategyFlags();
-			switch (ReimportStrategyFlags)
-			{
-				case EReimportStrategyFlags::ApplyNoProperties:
-				{
-					//We want to have no effect
-					break;
-				}
 					
-				case EReimportStrategyFlags::ApplyPipelineProperties:
-				{
-					//Directly apply pipeline node attribute to the asset
-					PipelineAssetNode->ApplyAllCustomAttributeToObject(Asset);
-					break;
-				}
+			case EReimportStrategyFlags::ApplyPipelineProperties:
+			{
+				//Directly apply pipeline node attribute to the asset
+				PipelineAssetNode->ApplyAllCustomAttributeToObject(Asset);
+				break;
+			}
 				
-				case EReimportStrategyFlags::ApplyEditorChangedProperties:
-				{
-					TArray<FAttributeKey> RemovedAttributes;
-					TArray<FAttributeKey> AddedAttributes;
-					TArray<FAttributeKey> ModifiedAttributes;
-					UInterchangeBaseNode::CompareNodeStorage(PreviousAssetNode, CurrentAssetNode, RemovedAttributes, AddedAttributes, ModifiedAttributes);
-
-					//set all ModifedAttributes from the CurrentAssetNode to the pipeline node. This will put back all user changes
-					UInterchangeBaseNode::CopyStorageAttributes(CurrentAssetNode, PipelineAssetNode, ModifiedAttributes);
-					//Now apply the pipeline node attribute to the asset
-					PipelineAssetNode->ApplyAllCustomAttributeToObject(Asset);
-					break;
-				}
+			case EReimportStrategyFlags::ApplyEditorChangedProperties:
+			{
+				TArray<FAttributeKey> RemovedAttributes;
+				TArray<FAttributeKey> AddedAttributes;
+				TArray<FAttributeKey> ModifiedAttributes;
+				UInterchangeBaseNode::CompareNodeStorage(PreviousAssetNode, CurrentAssetNode, RemovedAttributes, AddedAttributes, ModifiedAttributes);
+
+				//set all ModifedAttributes from the CurrentAssetNode to the pipeline node. This will put back all user changes
+				UInterchangeBaseNode::CopyStorageAttributes(CurrentAssetNode, PipelineAssetNode, ModifiedAttributes);
+				//Now apply the pipeline node attribute to the asset
+				PipelineAssetNode->ApplyAllCustomAttributeToObject(Asset);
+				break;
 			}
 		}
 	}
+
+	UObject* FFactoryCommon::GetObjectToReimport(UObject* ReimportObject, const FString& PackageName, const FString& AssetName, const FString& SubPathString)
+	{
+#if WITH_EDITORONLY_DATA
+		if (UInterchangeSceneImportAsset* SceneImportAsset = Cast<UInterchangeSceneImportAsset>(ReimportObject))
+		{
+			return SceneImportAsset->GetSceneObject(PackageName, AssetName, SubPathString);
+		}
+#endif
+
+		return ReimportObject;
+	}
+
+	const UInterchangeFactoryBaseNode* FFactoryCommon::GetFactoryNode(UObject* ObjectToReimport, const FString& PackageName, const FString& AssetName, const FString& SubPathString)
+	{
+#if WITH_EDITORONLY_DATA
+		if (UInterchangeSceneImportAsset* SceneImportAsset = Cast<UInterchangeSceneImportAsset>(ObjectToReimport))
+		{
+			return SceneImportAsset->GetFactoryNode(PackageName, AssetName, SubPathString);
+		}
+#endif
+
+		return nullptr;
+	}
 }
diff --git a/Engine/Source/Runtime/Interchange/Engine/Private/InterchangeSceneImportAsset.cpp b/Engine/Source/Runtime/Interchange/Engine/Private/InterchangeSceneImportAsset.cpp
new file mode 100644
--- /dev/null
+++ b/Engine/Source/Runtime/Interchange/Engine/Private/InterchangeSceneImportAsset.cpp
@@ -0,0 +1,218 @@
+// Copyright Epic Games, Inc. All Rights Reserved.
+
+#include "InterchangeSceneImportAsset.h"
+
+#include "InterchangeAssetImportData.h"
+#include "InterchangeFactoryBase.h"
+
+#include "Engine/AssetUserData.h"
+#include "Engine/Blueprint.h"
+#include "Engine/World.h"
+#include "EngineUtils.h"
+#include "UObject/Package.h"
+
+UInterchangeSceneImportAsset::~UInterchangeSceneImportAsset()
+{
+#if WITH_EDITOR
+	if (bWorldRenameCallbacksRegistered)
+	{
+		FWorldDelegates::OnPreWorldRename.RemoveAll(this);
+		FWorldDelegates::OnPostWorldRename.RemoveAll(this);
+	}
+#endif
+}
+
+void UInterchangeSceneImportAsset::GetAssetRegistryTags(TArray<FAssetRegistryTag>& OutTags) const
+{
+#if WITH_EDITORONLY_DATA
+	if (AssetImportData)
+	{
+		OutTags.Add(FAssetRegistryTag(SourceFileTagName(), AssetImportData->GetSourceData().ToJson(), FAssetRegistryTag::TT_Hidden));
+	}
+#endif
+
+	Super::GetAssetRegistryTags(OutTags);
+}
+
+void UInterchangeSceneImportAsset::PostLoad()
+{
+	Super::PostLoad();
+
+#if WITH_EDITORONLY_DATA
+	UpdateSceneObjects();
+#endif
+
+#if WITH_EDITOR
+	RegisterWorldRenameCallbacks();
+#endif
+}
+
+#if WITH_EDITOR
+void UInterchangeSceneImportAsset::RegisterWorldRenameCallbacks()
+{
+	if (!bWorldRenameCallbacksRegistered)
+	{
+		bWorldRenameCallbacksRegistered = true;
+		FWorldDelegates::OnPreWorldRename.AddUObject(this, &UInterchangeSceneImportAsset::OnPreWorldRename);
+		FWorldDelegates::OnPostWorldRename.AddUObject(this, &UInterchangeSceneImportAsset::OnPostWorldRename);
+	}
+}
+
+void UInterchangeSceneImportAsset::OnPreWorldRename(UWorld* World, const TCHAR* InName, UObject* NewOuter, ERenameFlags Flags, bool& bShouldFailRename)
+{
+	PreviousWorldPath = World->GetOutermost()->GetPathName();
+	PreviousWorldName = World->GetName();
+	PreviousLevelName = World->GetCurrentLevel()->GetName();
+}
+
+void UInterchangeSceneImportAsset::OnPostWorldRename(UWorld* World)
+{
+	PreEditChange(nullptr);
+
+	const FString PreviousRoot = FPaths::Combine(PreviousWorldPath, PreviousWorldName, PreviousLevelName);
+
+	TArray< FSoftObjectPath> EntriesToUpdate;
+	EntriesToUpdate.Reserve(SceneObjects.Num());
+
+	for (TPair<FSoftObjectPath, FString>& SceneObject : SceneObjects)
+	{
+		if (SceneObject.Key.GetAssetPathString().StartsWith(PreviousRoot))
+		{
+			EntriesToUpdate.Add(SceneObject.Key);
+		}
+	}
+
+	FString NewWorldPath = World->GetOutermost()->GetPathName();
+	FString NewWorldName = World->GetName();
+	FString NewPrefix = World->GetCurrentLevel()->GetName() + TEXT(".");
+
+	for (FSoftObjectPath& EntryToRemove : EntriesToUpdate)
+	{
+		FString UniqueID;
+		SceneObjects.RemoveAndCopyValue(EntryToRemove, UniqueID);
+
+		UInterchangeFactoryBaseNode* FactoryNode = AssetImportData->NodeContainer->GetFactoryNode(UniqueID);
+		if(ensure(FactoryNode))
+		{
+			const FString DisplayName = FactoryNode->GetDisplayLabel();
+			const FSoftObjectPath ObjectPath{ FName(NewWorldPath), FName(NewWorldName), NewPrefix + DisplayName };
+
+			FactoryNode->SetCustomReferenceObject(ObjectPath);
+			SceneObjects.Add(ObjectPath, UniqueID);
+		}
+	}
+
+	PostEditChange();
+}
+#endif
+
+void UInterchangeSceneImportAsset::AddAssetUserData( UAssetUserData* InUserData )
+{
+#if WITH_EDITORONLY_DATA
+	if ( InUserData != nullptr )
+	{
+		UAssetUserData* ExistingData = GetAssetUserDataOfClass( InUserData->GetClass() );
+		if ( ExistingData != nullptr )
+		{
+			AssetUserData.Remove( ExistingData );
+		}
+		AssetUserData.Add( InUserData );
+	}
+#endif // #if WITH_EDITORONLY_DATA
+}
+
+UAssetUserData* UInterchangeSceneImportAsset::GetAssetUserDataOfClass( TSubclassOf<UAssetUserData> InUserDataClass )
+{
+#if WITH_EDITORONLY_DATA
+	for ( int32 DataIdx = 0; DataIdx < AssetUserData.Num(); DataIdx++ )
+	{
+		UAssetUserData* Datum = AssetUserData[DataIdx];
+		if ( Datum != nullptr && Datum->IsA( InUserDataClass ) )
+		{
+			return Datum;
+		}
+	}
+#endif // #if WITH_EDITORONLY_DATA
+	return nullptr;
+
+}
+
+void UInterchangeSceneImportAsset::RemoveUserDataOfClass( TSubclassOf<UAssetUserData> InUserDataClass )
+{
+#if WITH_EDITORONLY_DATA
+	for ( int32 DataIdx = 0; DataIdx < AssetUserData.Num(); DataIdx++ )
+	{
+		UAssetUserData* Datum = AssetUserData[DataIdx];
+		if ( Datum != nullptr && Datum->IsA(InUserDataClass ) )
+		{
+			AssetUserData.RemoveAt( DataIdx );
+			return;
+		}
+	}
+#endif // #if WITH_EDITORONLY_DATA
+}
+
+const TArray<UAssetUserData*>* UInterchangeSceneImportAsset::GetAssetUserDataArray() const
+{
+#if WITH_EDITORONLY_DATA
+	return &ToRawPtrTArrayUnsafe(AssetUserData);
+#else
+	return nullptr;
+#endif // #if WITH_EDITORONLY_DATA
+}
+
+void UInterchangeSceneImportAsset::UpdateSceneObjects()
+{
+#if WITH_EDITORONLY_DATA
+	ensure(AssetImportData);
+
+	SceneObjects.Reset();
+
+	AssetImportData->NodeContainer->IterateNodesOfType<UInterchangeFactoryBaseNode>(
+		[this](const FString& NodeUid, UInterchangeFactoryBaseNode* FactoryNode)
+		{
+			if (FactoryNode)
+			{
+				FSoftObjectPath ObjectPath;
+				if (FactoryNode->GetCustomReferenceObject(ObjectPath))
+				{
+					this->SceneObjects.Add(ObjectPath, NodeUid);
+				}
+			}
+		}
+	);
+#endif
+}
+
+UObject* UInterchangeSceneImportAsset::GetSceneObject(const FString& PackageName, const FString& AssetName, const FString& SubPathString) const
+{
+#if WITH_EDITORONLY_DATA
+	const FSoftObjectPath ObjectPath(FName(PackageName), FName(AssetName), SubPathString);
+
+	if (SceneObjects.Find(ObjectPath))
+	{
+		return ObjectPath.TryLoad();
+	}
+#endif
+
+	return nullptr;
+}
+
+const UInterchangeFactoryBaseNode* UInterchangeSceneImportAsset::GetFactoryNode(const FString& PackageName, const FString& AssetName, const FString& SubPathString) const
+{
+#if WITH_EDITORONLY_DATA
+	if (!ensure(AssetImportData))
+	{
+		return nullptr;
+	}
+
+	const FSoftObjectPath ObjectPath(FName(PackageName), FName(AssetName), SubPathString);
+
+	if (const FString* UniqueIDPtr = SceneObjects.Find(ObjectPath))
+	{
+		return Cast<UInterchangeFactoryBaseNode>(AssetImportData->NodeContainer->GetNode(*UniqueIDPtr));
+	}
+#endif
+
+	return nullptr;
+}
diff --git a/Engine/Source/Runtime/Interchange/Engine/Private/Tasks/InterchangeTaskCreateSceneObjects.cpp b/Engine/Source/Runtime/Interchange/Engine/Private/Tasks/InterchangeTaskCreateSceneObjects.cpp
--- a/Engine/Source/Runtime/Interchange/Engine/Private/Tasks/InterchangeTaskCreateSceneObjects.cpp
+++ b/Engine/Source/Runtime/Interchange/Engine/Private/Tasks/InterchangeTaskCreateSceneObjects.cpp
@@ -4,9 +4,11 @@
 #include "CoreMinimal.h"
 
 #include "Async/TaskGraphInterfaces.h"
+#include "Engine/Level.h"
 #include "Engine/World.h"
 #include "InterchangeEngineLogPrivate.h"
 #include "InterchangeFactoryBase.h"
+#include "InterchangeImportCommon.h"
 #include "InterchangeManager.h"
 #include "InterchangeSourceData.h"
 #include "InterchangeTranslatorBase.h"
@@ -37,6 +39,7 @@ void UE::Interchange::FTaskCreateSceneObjects::DoTask(ENamedThreads::Type Curren
 #if INTERCHANGE_TRACE_ASYNCHRONOUS_TASK_ENABLED
 	INTERCHANGE_TRACE_ASYNCHRONOUS_TASK(SpawnActor)
 #endif
+	using namespace UE::Interchange;
 
 	TOptional<FGCScopeGuard> GCScopeGuard;
 	if (!IsInGameThread())
@@ -44,7 +47,7 @@ void UE::Interchange::FTaskCreateSceneObjects::DoTask(ENamedThreads::Type Curren
 		GCScopeGuard.Emplace();
 	}
 
-	TSharedPtr<UE::Interchange::FImportAsyncHelper> AsyncHelper = WeakAsyncHelper.Pin();
+	TSharedPtr<FImportAsyncHelper> AsyncHelper = WeakAsyncHelper.Pin();
 	check(WeakAsyncHelper.IsValid());
 
 	//Verify if the task was canceled
@@ -53,6 +56,12 @@ void UE::Interchange::FTaskCreateSceneObjects::DoTask(ENamedThreads::Type Curren
 		return;
 	}
 
+	UObject* ReimportObject = AsyncHelper->TaskData.ReimportObject;
+	ULevel* CurrentLevel = GWorld->GetCurrentLevel();
+	const FString WorldPath = GWorld->GetOutermost()->GetPathName();
+	const FString WorldName = GWorld->GetName();
+	const FString NodePrefix = CurrentLevel->GetName() + TEXT(".");
+
 	for (UInterchangeFactoryBaseNode* FactoryNode : FactoryNodes)
 	{
 		UInterchangeFactoryBase* Factory = NewObject<UInterchangeFactoryBase>(GetTransientPackage(), FactoryClass);
@@ -68,7 +77,9 @@ void UE::Interchange::FTaskCreateSceneObjects::DoTask(ENamedThreads::Type Curren
 		UInterchangeFactoryBase::FImportSceneObjectsParams CreateSceneObjectsParams;
 		CreateSceneObjectsParams.ObjectName = NodeDisplayName;
 		CreateSceneObjectsParams.FactoryNode = FactoryNode;
-		CreateSceneObjectsParams.Level = GWorld->GetCurrentLevel();
+		CreateSceneObjectsParams.Level = CurrentLevel;
+		CreateSceneObjectsParams.ReimportObject = FFactoryCommon::GetObjectToReimport(ReimportObject, WorldPath, WorldName, NodePrefix + NodeDisplayName);
+		CreateSceneObjectsParams.ReimportFactoryNode = FFactoryCommon::GetFactoryNode(ReimportObject, WorldPath, WorldName, NodePrefix + NodeDisplayName);
 
 		if (AsyncHelper->BaseNodeContainers.IsValidIndex(SourceIndex))
 		{
@@ -96,4 +107,4 @@ void UE::Interchange::FTaskCreateSceneObjects::DoTask(ENamedThreads::Type Curren
 			FactoryNode->SetCustomReferenceObject(FSoftObjectPath(SceneObject));
 		}
 	}
-}
\ No newline at end of file
+}
diff --git a/Engine/Source/Runtime/Interchange/Engine/Private/Tasks/InterchangeTaskImportObject.cpp b/Engine/Source/Runtime/Interchange/Engine/Private/Tasks/InterchangeTaskImportObject.cpp
--- a/Engine/Source/Runtime/Interchange/Engine/Private/Tasks/InterchangeTaskImportObject.cpp
+++ b/Engine/Source/Runtime/Interchange/Engine/Private/Tasks/InterchangeTaskImportObject.cpp
@@ -9,6 +9,7 @@
 #include "InterchangeAssetImportData.h"
 #include "InterchangeEngineLogPrivate.h"
 #include "InterchangeFactoryBase.h"
+#include "InterchangeImportCommon.h"
 #include "InterchangeManager.h"
 #include "InterchangeResult.h"
 #include "InterchangeSourceData.h"
@@ -118,6 +119,9 @@ namespace UE::Interchange::Private
 		return false;
 	}
 
+
+
+
 	bool CanImportClass(UE::Interchange::FImportAsyncHelper& AsyncHelper, UInterchangeFactoryBaseNode& FactoryNode, int32 SourceIndex)
 	{
 	#if WITH_EDITOR
@@ -185,7 +189,7 @@ namespace UE::Interchange::Private
 		FString AssetName;
 		Private::InternalGetPackageName(*AsyncHelper, SourceIndex, PackageBasePath, FactoryNode, PackageName, AssetName);
 		bool bSkipAsset = false;
-		UObject* ReimportObject = AsyncHelper->TaskData.ReimportObject;
+		UObject* ReimportObject = UE::Interchange::FFactoryCommon::GetObjectToReimport(AsyncHelper->TaskData.ReimportObject, PackageName, AssetName);
 		if (ReimportObject)
 		{
 			UInterchangeBaseNodeContainer* NodeContainer = nullptr;
@@ -288,7 +292,9 @@ void UE::Interchange::FTaskImportObject_GameThread::DoTask(ENamedThreads::Type C
 	UPackage* Pkg = nullptr;
 	FString PackageName;
 	FString AssetName;
-	UObject* ReimportObject = AsyncHelper->TaskData.ReimportObject;
+	Private::InternalGetPackageName(*AsyncHelper, SourceIndex, PackageBasePath, FactoryNode, PackageName, AssetName);
+
+	UObject* ReimportObject = FFactoryCommon::GetObjectToReimport(AsyncHelper->TaskData.ReimportObject, PackageName, AssetName);
 	//If we do a reimport no need to create a package
 	if (ReimportObject)
 	{
@@ -479,7 +485,7 @@ void UE::Interchange::FTaskImportObjectFinalize_GameThread::DoTask(ENamedThreads
 
 	if (ImportAssetResult.ImportedObject)
 	{
-		//If the factory skip the asset, we simply set the node custom refernece object
+		//If the factory skip the asset, we simply set the node custom reference object
 		if (!ImportAssetResult.bIsFactorySkipAsset)
 		{
 			FScopeLock Lock(&AsyncHelper->ImportedAssetsPerSourceIndexLock);
@@ -495,7 +501,9 @@ void UE::Interchange::FTaskImportObjectFinalize_GameThread::DoTask(ENamedThreads
 				AssetInfo.ImportedObject = ImportAssetResult.ImportedObject;
 				AssetInfo.Factory = Factory;
 				AssetInfo.FactoryNode = FactoryNode;
-				AssetInfo.bIsReimport = bool(AsyncHelper->TaskData.ReimportObject != nullptr);
+
+				const UObject* ReimportObject = FFactoryCommon::GetObjectToReimport(AsyncHelper->TaskData.ReimportObject, ImportAssetResult.ImportedObject->GetPathName(), ImportAssetResult.ImportedObject->GetName());
+				AssetInfo.bIsReimport = bool(ReimportObject != nullptr);
 			}
 
 			// Fill in destination asset and type in any results which have been added previously by a translator or pipeline, now that we have a corresponding factory.
diff --git a/Engine/Source/Runtime/Interchange/Engine/Public/InterchangeImportCommon.h b/Engine/Source/Runtime/Interchange/Engine/Public/InterchangeImportCommon.h
--- a/Engine/Source/Runtime/Interchange/Engine/Public/InterchangeImportCommon.h
+++ b/Engine/Source/Runtime/Interchange/Engine/Public/InterchangeImportCommon.h
@@ -45,18 +45,18 @@ namespace UE::Interchange
 		};
 
 		/**
-		* Update the AssetImportData source file of the specified asset in the parameters. Also update the node container and the node unique id.
-		* If the AssetImportData is null it will create one. If the AssetImportData is not an UInterchangeAssetImportData it will create a new one.
-		* @return The source data that should be stored on the asset or nullptr if a parameter is invalid
-		*/
+		 * Update the AssetImportData source file of the specified asset in the parameters. Also update the node container and the node unique id.
+		 * If the AssetImportData is null it will create one. If the AssetImportData is not an UInterchangeAssetImportData it will create a new one.
+		 * @return The source data that should be stored on the asset or nullptr if a parameter is invalid
+		 */
 		static UAssetImportData* UpdateImportAssetData(FUpdateImportAssetDataParameters& Parameters);
 
 		/**
-		* Update the AssetImportData of the specified asset in the parameters. Also update the node container and the node unique id.
-		* If the AssetImportData is null it will create one. If the AssetImportData is not an UInterchangeAssetImportData it will create a new one.
-		* The file source update is done by calling the function parameter CustomFileSourceUpdate, so its the client responsability to properly update the file source.
-		* @return The source data that should be stored on the asset or nullptr if a parameter is invalid
-		*/
+		 * Update the AssetImportData of the specified asset in the parameters. Also update the node container and the node unique id.
+		 * If the AssetImportData is null it will create one. If the AssetImportData is not an UInterchangeAssetImportData it will create a new one.
+		 * The file source update is done by calling the function parameter CustomFileSourceUpdate, so its the client responsability to properly update the file source.
+		 * @return The source data that should be stored on the asset or nullptr if a parameter is invalid
+		 */
 		static UAssetImportData* UpdateImportAssetData(FUpdateImportAssetDataParameters& Parameters, TFunctionRef<void(UInterchangeAssetImportData*)> CustomFileSourceUpdate);
 
 #if WITH_EDITORONLY_DATA
@@ -74,36 +74,58 @@ namespace UE::Interchange
 		};
 
 		/**
-		* Set the AssetImportData source file of the specified asset in the parameters. Also update the node container and the node unique id.
-		* If the AssetImportData is null it will create one. If the AssetImportData is not an UInterchangeAssetImportData it will create a new one.
-		* @return The source data that should be stored on the asset or nullptr if a parameter is invalid
-		*/
+		 * Set the AssetImportData source file of the specified asset in the parameters. Also update the node container and the node unique id.
+		 * If the AssetImportData is null it will create one. If the AssetImportData is not an UInterchangeAssetImportData it will create a new one.
+		 * @return The source data that should be stored on the asset or nullptr if a parameter is invalid
+		 */
 		static UAssetImportData* SetImportAssetData(FSetImportAssetDataParameters& Parameters);
 
 		/**
-		* Fills the OutSourceFilenames array with the list of source files contained in the asset source data.
-		* Returns true if the operation was successful.
-		*/
+		 * Fills the OutSourceFilenames array with the list of source files contained in the asset source data.
+		 * Returns true if the operation was successful.
+		 */
 		static bool GetSourceFilenames(const UAssetImportData* AssetImportData, TArray<FString>& OutSourceFilenames);
 
 		/**
-		* Sets the SourceFileName value at the specified index.
-		*/
+		 * Sets the SourceFileName value at the specified index.
+		 */
 		static bool SetSourceFilename(UAssetImportData* AssetImportData, const FString& SourceFilename, int32 SourceIndex, const FString& SourceLabel = FString());
 
 		/**
-		* Set the object's reimport source at the specified index value.
-		*/
+		 * Set the object's reimport source at the specified index value.
+		 */
 		static bool SetReimportSourceIndex(const UObject* Object, UAssetImportData* AssetImportData, int32 SourceIndex);
 
 #endif // WITH_EDITORONLY_DATA
 
 		/**
-		* Apply the current strategy to the PipelineAssetNode
-		*/
+		 * Apply the current strategy to the PipelineAssetNode
+		 */
 		static void ApplyReimportStrategyToAsset(UObject* Asset
 											, const UInterchangeFactoryBaseNode* PreviousAssetNode
 											, const UInterchangeFactoryBaseNode* CurrentAssetNode
 											, UInterchangeFactoryBaseNode* PipelineAssetNode);
+
+		/**
+		 * If the ReimportObject is a UInterchangeSceneImportAsset, returns the UObject which asset path
+		 * is //PackageName.AssetName[:SubPathString].
+		 * Returns ReimportObject otherwise.
+		 * @oaram ReimportObject: Original object to reimport
+		 * @oaram PackageName: Package path of the actual object to reimport
+		 * @oaram AssetName: Asset name of the actual object to reimport
+		 * @oaram SubPathString: Optional subobject name
+		 */
+		static 	UObject* GetObjectToReimport(UObject* ReimportObject, const FString& PackageName, const FString& AssetName, const FString& SubPathString = FString());
+
+		/**
+		 * If the ReimportObject is a UInterchangeSceneImportAsset, returns the factory node
+		 * associated with the asset which path is //PackageName.AssetName[:SubPathString].
+		 * Returns nullptr otherwise.
+		 * @oaram ReimportObject: Original object to reimport
+		 * @oaram PackageName: Package path of the actual object to reimport
+		 * @oaram AssetName: Asset name of the actual object to reimport
+		 * @oaram SubPathString: Optional subobject name
+		 */
+		static 	const UInterchangeFactoryBaseNode* GetFactoryNode(UObject* ReimportObject, const FString& PackageName, const FString& AssetName, const FString& SubPathString = FString());
 	};
 } //ns UE::Interchange
diff --git a/Engine/Source/Runtime/Interchange/Engine/Public/InterchangeSceneImportAsset.h b/Engine/Source/Runtime/Interchange/Engine/Public/InterchangeSceneImportAsset.h
new file mode 100644
--- /dev/null
+++ b/Engine/Source/Runtime/Interchange/Engine/Public/InterchangeSceneImportAsset.h
@@ -0,0 +1,98 @@
+// Copyright Epic Games, Inc. All Rights Reserved.
+
+#pragma once
+
+#include "Misc/Guid.h"
+#include "Interfaces/Interface_AssetUserData.h"
+#include "UObject/SoftObjectPath.h"
+
+#include "InterchangeSceneImportAsset.generated.h"
+
+class UInterchangeAssetImportData;
+class UInterchangeFactoryBaseNode;
+
+/*
+ * Class to hold all the data required to properly re-import a level
+ */
+UCLASS()
+class INTERCHANGEENGINE_API UInterchangeSceneImportAsset : public UObject, public IInterface_AssetUserData
+{
+	GENERATED_BODY()
+
+	virtual ~UInterchangeSceneImportAsset();
+
+public:
+#if WITH_EDITORONLY_DATA
+	/** Importing data and options used for this Datasmith scene */
+	UPROPERTY(EditAnywhere, Instanced, Category=ImportSettings)
+	TObjectPtr<UInterchangeAssetImportData> AssetImportData;
+
+	/** Array of user data stored with the asset */
+	UPROPERTY()
+	TArray< TObjectPtr<UAssetUserData> > AssetUserData;
+#endif // #if WITH_EDITORONLY_DATA
+
+	//~ Begin UObject Interface
+	virtual void GetAssetRegistryTags(TArray<FAssetRegistryTag>& OutTags) const override;
+	virtual void PostLoad() override;
+	//~ End UObject Interface
+
+	//~ Begin IInterface_AssetUserData Interface
+	virtual void AddAssetUserData( UAssetUserData* InUserData ) override;
+	virtual void RemoveUserDataOfClass( TSubclassOf<UAssetUserData> InUserDataClass ) override;
+	virtual UAssetUserData* GetAssetUserDataOfClass( TSubclassOf<UAssetUserData> InUserDataClass ) override;
+	virtual const TArray<UAssetUserData*>* GetAssetUserDataArray() const override;
+	//~ End IInterface_AssetUserData Interface
+
+#if WITH_EDITOR
+	void RegisterWorldRenameCallbacks();
+#endif
+
+	/** Updates the SceneObjects cache based on the node container stored in AssetImportData */
+	void UpdateSceneObjects();
+
+	/**
+	 * Returns the UObject which asset path is '//PackageName.AssetName[:SubPathString]'.
+	 * Returns nullptr if the asset which path is '//PackageName.AssetName[:SubPathString]' was not part of
+	 * the level import cached in this UInterchangeSceneImportAsset.
+	 * @oaram PackageName: Package path of the actual object to reimport
+	 * @oaram AssetName: Asset name of the actual object to reimport
+	 * @oaram SubPathString: Optional subobject name
+	 */
+	UObject* GetSceneObject(const FString& PackageName, const FString& AssetName, const FString& SubPathString = FString()) const;
+
+	/**
+	 * Returns the factory node associated with the asset which path is '//PackageName.AssetName[:SubPathString]'.
+	 * Returns nullptr if the asset which path is '//PackageName.AssetName[:SubPathString]' was not part of
+	 * the level import cached in this UInterchangeSceneImportAsset.
+	 * @oaram PackageName: Package path of the actual object to reimport
+	 * @oaram AssetName: Asset name of the actual object to reimport
+	 * @oaram SubPathString: Optional subobject name
+	 */
+	const UInterchangeFactoryBaseNode* GetFactoryNode(const FString& PackageName, const FString& AssetName, const FString& SubPathString = FString()) const;
+
+private:
+
+#if WITH_EDITOR
+	/** Called before a world is renamed */
+	void OnPreWorldRename(UWorld* World, const TCHAR* InName, UObject* NewOuter, ERenameFlags Flags, bool& bShouldFailRename);
+	
+	/** Invoked when a world is successfully renamed. Used to track when a temporary 'Untitled' unsaved map is saved with a new name. */
+	void OnPostWorldRename(UWorld* World);
+
+	/** Members used to cache the path and names related to the world to be renamed*/
+	bool bWorldRenameCallbacksRegistered = false;
+	FString PreviousWorldPath;
+	FString PreviousWorldName;
+	FString PreviousLevelName;
+#endif
+
+#if WITH_EDITORONLY_DATA
+	/**
+	 * Cache to easily retrieve a factory node from an asset's/actor's path
+	 * FSoftObjectPath stores the path of an imported object
+	 * FString stores the unique id of the factory node associated with the imported object
+	 */
+	TMap< FSoftObjectPath, FString > SceneObjects;
+#endif // #if WITH_EDITORONLY_DATA
+};
