Fixed logic issue with world renaming

diff --git a/Engine/Source/Runtime/Interchange/Engine/Private/InterchangeSceneImportAsset.cpp b/Engine/Source/Runtime/Interchange/Engine/Private/InterchangeSceneImportAsset.cpp
--- a/Engine/Source/Runtime/Interchange/Engine/Private/InterchangeSceneImportAsset.cpp
+++ b/Engine/Source/Runtime/Interchange/Engine/Private/InterchangeSceneImportAsset.cpp
@@ -60,23 +60,26 @@ void UInterchangeSceneImportAsset::RegisterWorldRenameCallbacks()
 
 void UInterchangeSceneImportAsset::OnPreWorldRename(UWorld* World, const TCHAR* InName, UObject* NewOuter, ERenameFlags Flags, bool& bShouldFailRename)
 {
-	PreviousWorldPath = World->GetOutermost()->GetPathName();
-	PreviousWorldName = World->GetName();
-	PreviousLevelName = World->GetCurrentLevel()->GetName();
+	// This method is called twice, first before the name change on the outermost then before the name change on the asset
+	// Only cache path and names on the first call to this method.
+	if (PreviousWorldPath.IsEmpty())
+	{
+		PreviousWorldPath = World->GetOutermost()->GetPathName();
+		PreviousWorldName = World->GetName();
+		PreviousLevelName = World->GetCurrentLevel()->GetName();
+	}
 }
 
 void UInterchangeSceneImportAsset::OnPostWorldRename(UWorld* World)
 {
 	PreEditChange(nullptr);
 
-	const FString PreviousRoot = FPaths::Combine(PreviousWorldPath, PreviousWorldName, PreviousLevelName);
-
 	TArray< FSoftObjectPath> EntriesToUpdate;
 	EntriesToUpdate.Reserve(SceneObjects.Num());
 
 	for (TPair<FSoftObjectPath, FString>& SceneObject : SceneObjects)
 	{
-		if (SceneObject.Key.GetAssetPathString().StartsWith(PreviousRoot))
+		if (SceneObject.Key.GetAssetPathString().StartsWith(PreviousWorldPath))
 		{
 			EntriesToUpdate.Add(SceneObject.Key);
 		}
@@ -102,6 +105,11 @@ void UInterchangeSceneImportAsset::OnPostWorldRename(UWorld* World)
 		}
 	}
 
+	// Reset previously cached path and names for subsequent call to OnPreWorldRename
+	PreviousWorldPath.Empty();
+	PreviousWorldName.Empty();
+	PreviousLevelName.Empty();
+
 	PostEditChange();
 }
 #endif
