Fix static analysis build

diff --git a/Engine/Source/Runtime/Interchange/Engine/Private/InterchangeImportCommon.cpp b/Engine/Source/Runtime/Interchange/Engine/Private/InterchangeImportCommon.cpp
--- a/Engine/Source/Runtime/Interchange/Engine/Private/InterchangeImportCommon.cpp
+++ b/Engine/Source/Runtime/Interchange/Engine/Private/InterchangeImportCommon.cpp
@@ -345,9 +345,12 @@ namespace UE::Interchange
 		}
 
 		const UClass* AssetClass = FactoryNode->GetObjectClass();
-		if (!ensure(AssetClass && AssetClass->IsChildOf(FactoryClass)))
+		if (!AssetClass || !AssetClass->IsChildOf(FactoryClass))
 		{
-			UE_LOG(LogInterchangeEngine, Error, TEXT("The asset class %s is not compatible with the factory class %s. Cannot import asset %s ."), *AssetClass->GetName(), *FactoryClass->GetName(), *AssetName);
+			UE_LOG(LogInterchangeEngine, Error, TEXT("The asset class %s is not compatible with the factory class %s. Cannot import asset %s .")
+				, AssetClass ? *AssetClass->GetName() : TEXT("unknown")
+				, FactoryClass ? *FactoryClass->GetName() : TEXT("unknown")
+				, *AssetName);
 			return nullptr;
 		}
 
