#UE Add CVar to allow following redirectors when importing assets

diff --git a/Engine/Source/Developer/AssetTools/Private/AssetTools.cpp b/Engine/Source/Developer/AssetTools/Private/AssetTools.cpp
--- a/Engine/Source/Developer/AssetTools/Private/AssetTools.cpp
+++ b/Engine/Source/Developer/AssetTools/Private/AssetTools.cpp
@@ -451,6 +451,13 @@ namespace UE::AssetTools::Private
 		return EAdvancedCopyConsolidationMethod::PerAsset;
 	}
 
+	bool bFollowRedirectorsWhenImporting = false;
+	static FAutoConsoleVariableRef CVarFollowRedirectorsWhenImporting(
+		TEXT("AssetTools.FollowRedirectorsWhenImporting"),
+		bFollowRedirectorsWhenImporting,
+		TEXT("When set, if you import an asset at a location with a redirector, you'll instead import to the redirector's destination location")
+	);
+
 	// use a struct as a namespace to allow easier friend declarations
 	struct FPackageMigrationImpl
 	{
@@ -3316,6 +3323,7 @@ TArray<UObject*> UAssetToolsImpl::ImportAssetsInternal(const TArray<FString>& Fi
 			{
 				FImportAssetParameters ImportAssetParameters;
 				ImportAssetParameters.bIsAutomated = bAutomatedImport;
+				ImportAssetParameters.bFollowRedirectors = UE::AssetTools::Private::bFollowRedirectorsWhenImporting;
 				ImportAssetParameters.ReimportAsset = nullptr;
 
 
@@ -3580,6 +3588,24 @@ TArray<UObject*> UAssetToolsImpl::ImportAssetsInternal(const TArray<FString>& Fi
 
 			// Check for an existing object
 			UObject* ExistingObject = StaticFindObject(UObject::StaticClass(), Pkg, *Name);
+
+			if (UE::AssetTools::Private::bFollowRedirectorsWhenImporting)
+			{
+				if (UObjectRedirector* Redirector = Cast<UObjectRedirector>(ExistingObject))
+				{
+					if (Redirector->DestinationObject)
+					{
+						ExistingObject = Redirector->DestinationObject;
+						Pkg = Redirector->DestinationObject->GetPackage();
+						if (FPackageName::GetLongPackageAssetName(PackageName) == Name)
+						{
+							Name = FPackageName::GetLongPackageAssetName(Pkg->GetName());
+						}
+						PackageName = Pkg->GetName();
+					}
+				}
+			}
+
 			if(ExistingObject != nullptr)
 			{
 				// If the existing object is one of the imports we've just created we can't replace or overwrite it
diff --git a/Engine/Source/Runtime/Interchange/Engine/Private/InterchangeManager.cpp b/Engine/Source/Runtime/Interchange/Engine/Private/InterchangeManager.cpp
--- a/Engine/Source/Runtime/Interchange/Engine/Private/InterchangeManager.cpp
+++ b/Engine/Source/Runtime/Interchange/Engine/Private/InterchangeManager.cpp
@@ -1122,6 +1122,7 @@ UInterchangeManager::ImportInternal(const FString& ContentPath, const UInterchan
 	//Create a task for every source data
 	UE::Interchange::FImportAsyncHelperData TaskData;
 	TaskData.bIsAutomated = ImportAssetParameters.bIsAutomated;
+	TaskData.bFollowRedirectors = ImportAssetParameters.bFollowRedirectors;
 	TaskData.ImportType = ImportType;
 	TaskData.ReimportObject = ImportAssetParameters.ReimportAsset;
 	TSharedRef<UE::Interchange::FImportAsyncHelper, ESPMode::ThreadSafe> AsyncHelper = CreateAsyncHelper(TaskData, ImportAssetParameters);
diff --git a/Engine/Source/Runtime/Interchange/Engine/Private/Tasks/InterchangeTaskImportObject.cpp b/Engine/Source/Runtime/Interchange/Engine/Private/Tasks/InterchangeTaskImportObject.cpp
--- a/Engine/Source/Runtime/Interchange/Engine/Private/Tasks/InterchangeTaskImportObject.cpp
+++ b/Engine/Source/Runtime/Interchange/Engine/Private/Tasks/InterchangeTaskImportObject.cpp
@@ -240,6 +240,15 @@ namespace UE::Interchange::Private
 
 		if (!bSkipAsset)
 		{
+			if (AsyncHelper->TaskData.bFollowRedirectors)
+			{
+				// Check to see if we were redirected. If so, and if the asset name matched the package name, change the name of the asset to match the new package as well
+				if (Pkg->GetName() != PackageName && FPackageName::GetLongPackageAssetName(PackageName) == AssetName)
+				{
+					AssetName = FPackageName::GetLongPackageAssetName(Pkg->GetName());
+				}
+			}
+
 			UInterchangeTranslatorBase* Translator = AsyncHelper->Translators[SourceIndex];
 			//Import Asset describe by the node
 			UInterchangeFactoryBase::FImportAssetObjectParams CreateAssetParams;
@@ -349,6 +358,7 @@ void UE::Interchange::FTaskImportObject_GameThread::DoTask(ENamedThreads::Type C
 		//If the package already exist we must load it so factory can find any existing asset and decide or not to override it
 		{
 			//Try to find the package in memory
+			bool bPackageWasCreated = false;
 			Pkg = FindPackage(nullptr, *PackageName);
 			if (!Pkg)
 			{
@@ -359,6 +369,7 @@ void UE::Interchange::FTaskImportObject_GameThread::DoTask(ENamedThreads::Type C
 					//Create the package
 					Pkg = CreatePackage(*PackageName);
 					Pkg->SetPackageFlags(PKG_NewlyCreated);
+					bPackageWasCreated = true;
 				}
 			}
 
@@ -373,6 +384,21 @@ void UE::Interchange::FTaskImportObject_GameThread::DoTask(ENamedThreads::Type C
 				//Skip this asset
 				return;
 			}
+
+			if (!bPackageWasCreated && AsyncHelper->TaskData.bFollowRedirectors)
+			{
+				if (UObjectRedirector* Redirector = FindObject<UObjectRedirector>(Pkg, *AssetName))
+				{
+					if (Redirector->DestinationObject)
+					{
+						Pkg = Redirector->DestinationObject->GetPackage();
+						if (FPackageName::GetLongPackageAssetName(PackageName) == AssetName)
+						{
+							AssetName = FPackageName::GetLongPackageAssetName(Pkg->GetName());
+						}
+					}
+				}
+			}
 		}
 
 		//Import Asset describe by the node
diff --git a/Engine/Source/Runtime/Interchange/Engine/Public/InterchangeManager.h b/Engine/Source/Runtime/Interchange/Engine/Public/InterchangeManager.h
--- a/Engine/Source/Runtime/Interchange/Engine/Public/InterchangeManager.h
+++ b/Engine/Source/Runtime/Interchange/Engine/Public/InterchangeManager.h
@@ -82,6 +82,9 @@ namespace UE
 			//True if the import process is unattended. We cannot show UI  if the import is automated
 			bool bIsAutomated = false;
 
+			// True if redirectors will be followed when determining what location to import an asset
+			bool bFollowRedirectors = false;
+
 			//We can import assets or full scene
 			EImportType ImportType = EImportType::ImportType_None;
 
@@ -300,6 +303,10 @@ struct INTERCHANGEENGINE_API FImportAssetParameters
 	UPROPERTY(EditAnywhere, BlueprintReadWrite, Category = "Interchange|ImportAsset")
 	bool bIsAutomated = false;
 
+	// Tell interchange to follow redirectors when determining the location an asset will be imported
+	UPROPERTY(EditAnywhere, BlueprintReadWrite, Category = "Interchange|ImportAsset")
+	bool bFollowRedirectors = false;
+
 	// Adding some override will tell interchange to use the specific custom set pipelines instead of letting the user or the system chose
 	UPROPERTY(EditAnywhere, BlueprintReadWrite, Category = "Interchange|ImportAsset")
 	TArray<TObjectPtr<UInterchangePipelineBase>> OverridePipelines;
