#UE IWYU fix

diff --git a/Engine/Source/Runtime/Interchange/Engine/Private/Tasks/InterchangeTaskImportObject.cpp b/Engine/Source/Runtime/Interchange/Engine/Private/Tasks/InterchangeTaskImportObject.cpp
--- a/Engine/Source/Runtime/Interchange/Engine/Private/Tasks/InterchangeTaskImportObject.cpp
+++ b/Engine/Source/Runtime/Interchange/Engine/Private/Tasks/InterchangeTaskImportObject.cpp
@@ -22,6 +22,7 @@
 #include "UObject/Object.h"
 #include "UObject/UObjectGlobals.h"
 #include "UObject/ObjectMacros.h"
+#include "UObject/ObjectRedirector.h"
 #include "UObject/WeakObjectPtrTemplates.h"
 
 namespace UE::Interchange::Private
