Allow multiple concurrent import in the same folder

#rnx

diff --git a/Engine/Plugins/Experimental/Enterprise/DatasmithInterchange/Source/DatasmithInterchange/Private/InterchangeDatasmithSceneFactory.cpp b/Engine/Plugins/Experimental/Enterprise/DatasmithInterchange/Source/DatasmithInterchange/Private/InterchangeDatasmithSceneFactory.cpp
--- a/Engine/Plugins/Experimental/Enterprise/DatasmithInterchange/Source/DatasmithInterchange/Private/InterchangeDatasmithSceneFactory.cpp
+++ b/Engine/Plugins/Experimental/Enterprise/DatasmithInterchange/Source/DatasmithInterchange/Private/InterchangeDatasmithSceneFactory.cpp
@@ -34,15 +34,22 @@ UInterchangeFactoryBase::FImportAssetResult UInterchangeDatasmithSceneFactory::B
 		return ImportAssetResult;
 	}
 
-	// create an asset if it doesn't exist
-	UObject* ExistingAsset = StaticFindObject(nullptr, Arguments.Parent, *Arguments.AssetName);
+	UObject* ExistingAsset = Arguments.ReimportObject;
+	if (!ExistingAsset)
+	{
+		FSoftObjectPath ReferenceObject;
+		if (Arguments.AssetNode->GetCustomReferenceObject(ReferenceObject))
+		{
+			ExistingAsset = ReferenceObject.TryLoad();
+		}
+	}
 
 	// create a new texture or overwrite existing asset, if possible
 	if (!ExistingAsset)
 	{
 		DatasmithScene = NewObject<UDatasmithScene>(Arguments.Parent, DatasmithSceneClass, *Arguments.AssetName, RF_Public | RF_Standalone);
 	}
-	else if (ExistingAsset->GetClass()->IsChildOf(DatasmithSceneClass))
+	else
 	{
 		//This is a reimport, we are just re-updating the source data
 		DatasmithScene = static_cast<UDatasmithScene*>(ExistingAsset);
diff --git a/Engine/Plugins/Interchange/Runtime/Source/Import/Private/Animation/InterchangeAnimSequenceFactory.cpp b/Engine/Plugins/Interchange/Runtime/Source/Import/Private/Animation/InterchangeAnimSequenceFactory.cpp
--- a/Engine/Plugins/Interchange/Runtime/Source/Import/Private/Animation/InterchangeAnimSequenceFactory.cpp
+++ b/Engine/Plugins/Interchange/Runtime/Source/Import/Private/Animation/InterchangeAnimSequenceFactory.cpp
@@ -156,7 +156,7 @@ namespace UE::Interchange::Private
 		, const bool bMaterialCurve
 		, const bool bShouldTransact)
 	{
-		if (!TargetSequence || Curves.Num() != 1 || CurveName.IsEmpty())
+		if (!TargetSequence || Curves.Num() != 1 || CurveName.IsEmpty() || Curves[0].IsEmpty())
 		{
 			return false;
 		}
@@ -806,15 +806,22 @@ UInterchangeFactoryBase::FImportAssetResult UInterchangeAnimSequenceFactory::Beg
 		return ImportAssetResult;
 	}
 
-	// create an asset if it doesn't exist
-	UObject* ExistingAsset = StaticFindObject(nullptr, Arguments.Parent, *Arguments.AssetName);
+	UObject* ExistingAsset = Arguments.ReimportObject;
+	if (!ExistingAsset)
+	{
+		FSoftObjectPath ReferenceObject;
+		if (AnimSequenceFactoryNode->GetCustomReferenceObject(ReferenceObject))
+		{
+			ExistingAsset = ReferenceObject.TryLoad();
+		}
+	}
 
 	// create a new material or overwrite existing asset, if possible
 	if (!ExistingAsset)
 	{
 		AnimSequence = NewObject<UAnimSequence>(Arguments.Parent, *Arguments.AssetName, RF_Public | RF_Standalone);
 	}
-	else if (ExistingAsset->GetClass()->IsChildOf(UAnimSequence::StaticClass()))
+	else
 	{
 		//This is a reimport, we are just re-updating the source data
 		AnimSequence = Cast<UAnimSequence>(ExistingAsset);
diff --git a/Engine/Plugins/Interchange/Runtime/Source/Import/Private/Animation/InterchangeAnimationTrackSetFactory.cpp b/Engine/Plugins/Interchange/Runtime/Source/Import/Private/Animation/InterchangeAnimationTrackSetFactory.cpp
--- a/Engine/Plugins/Interchange/Runtime/Source/Import/Private/Animation/InterchangeAnimationTrackSetFactory.cpp
+++ b/Engine/Plugins/Interchange/Runtime/Source/Import/Private/Animation/InterchangeAnimationTrackSetFactory.cpp
@@ -661,14 +661,20 @@ UInterchangeFactoryBase::FImportAssetResult UInterchangeAnimationTrackSetFactory
 	UE_LOG(LogInterchangeImport, Error, TEXT("Cannot import levelsequence asset in runtime, this is an editor only feature."));
 	return ImportAssetResult;
 #else
-	if (Arguments.ReimportObject)
+
+	auto CannotReimportMessage = [&Arguments, this]()
 	{
 		UInterchangeResultError_Generic* Message = AddMessage<UInterchangeResultError_Generic>();
 		Message->SourceAssetName = Arguments.SourceData->GetFilename();
 		Message->DestinationAssetName = Arguments.AssetName;
 		Message->AssetType = ULevelSequence::StaticClass();
 		Message->Text = LOCTEXT("CreateEmptyAssetUnsupportedReimport", "Re-import of ULevelSequence not supported yet.");
+		Arguments.AssetNode->SetSkipNodeImport();
+	};
 
+	if (Arguments.ReimportObject)
+	{
+		CannotReimportMessage();
 		return ImportAssetResult;
 	}
 
@@ -690,18 +696,29 @@ UInterchangeFactoryBase::FImportAssetResult UInterchangeAnimationTrackSetFactory
 		return ImportAssetResult;
 	}
 
-	// create an asset if it doesn't exist
-	UObject* ExistingAsset = StaticFindObject(nullptr, Arguments.Parent, *Arguments.AssetName);
+	UObject* ExistingAsset = Arguments.ReimportObject;
+	if (!ExistingAsset)
+	{
+		FSoftObjectPath ReferenceObject;
+		if (FactoryNode->GetCustomReferenceObject(ReferenceObject))
+		{
+			ExistingAsset = ReferenceObject.TryLoad();
+		}
+	}
 
 	// create a new material or overwrite existing asset, if possible
 	if (!ExistingAsset)
 	{
 		LevelSequence = NewObject<ULevelSequence>(Arguments.Parent, *Arguments.AssetName, RF_Public | RF_Standalone);
 	}
-	else if (ExistingAsset->GetClass()->IsChildOf(ULevelSequence::StaticClass()))
+	else
 	{
-		// This is a reimport, we are just re-updating the source data
-		LevelSequence = Cast<ULevelSequence>(ExistingAsset);
+		// This is a reimport or an override, we are just re-updating the source data
+
+		//TODO: put back the Cast when the LevelSequence will support re-import
+		//LevelSequence = Cast<ULevelSequence>(ExistingAsset);
+		CannotReimportMessage();
+		return ImportAssetResult;
 	}
 
 	if (!LevelSequence)
@@ -806,4 +823,4 @@ void UInterchangeAnimationTrackSetFactory::SetupObject_GameThread(const FSetupOb
 	// TODO: Talk with sequence team about adding AssetImportData to ULevelSequence for re-import 
 }
 
-#undef LOCTEXT_NAMESPACE
+#undef LOCTEXT_NAMESPACE
\ No newline at end of file
diff --git a/Engine/Plugins/Interchange/Runtime/Source/Import/Private/Material/InterchangeMaterialFactory.cpp b/Engine/Plugins/Interchange/Runtime/Source/Import/Private/Material/InterchangeMaterialFactory.cpp
--- a/Engine/Plugins/Interchange/Runtime/Source/Import/Private/Material/InterchangeMaterialFactory.cpp
+++ b/Engine/Plugins/Interchange/Runtime/Source/Import/Private/Material/InterchangeMaterialFactory.cpp
@@ -505,8 +505,15 @@ UInterchangeFactoryBase::FImportAssetResult UInterchangeMaterialFactory::BeginIm
 
 	const bool bIsReimport = Arguments.ReimportObject != nullptr;
 
-	// Find existing asset
-	UObject* ExistingAsset = bIsReimport ? Arguments.ReimportObject : StaticFindObject(nullptr, Arguments.Parent, *Arguments.AssetName);
+	UObject* ExistingAsset = Arguments.ReimportObject;
+	if (!ExistingAsset)
+	{
+		FSoftObjectPath ReferenceObject;
+		if (MaterialFactoryNode->GetCustomReferenceObject(ReferenceObject))
+		{
+			ExistingAsset = ReferenceObject.TryLoad();
+		}
+	}
 
 	// create a new material or overwrite existing asset, if possible
 	if (!ExistingAsset)
@@ -1398,11 +1405,17 @@ UInterchangeFactoryBase::FImportAssetResult UInterchangeMaterialFunctionFactory:
 		return ImportAssetResult;
 	}
 
-	// create an asset if it doesn't exist
 	const bool bIsReimport = Arguments.ReimportObject != nullptr;
 
-	// Find existing asset
-	UObject* ExistingAsset = bIsReimport ? Arguments.ReimportObject : StaticFindObject(nullptr, Arguments.Parent, *Arguments.AssetName);
+	UObject* ExistingAsset = Arguments.ReimportObject;
+	if (!ExistingAsset)
+	{
+		FSoftObjectPath ReferenceObject;
+		if (MaterialFactoryNode->GetCustomReferenceObject(ReferenceObject))
+		{
+			ExistingAsset = ReferenceObject.TryLoad();
+		}
+	}
 
 	// create a new material or overwrite existing asset, if possible
 	if (!ExistingAsset)
diff --git a/Engine/Plugins/Interchange/Runtime/Source/Import/Private/Mesh/InterchangePhysicsAssetFactory.cpp b/Engine/Plugins/Interchange/Runtime/Source/Import/Private/Mesh/InterchangePhysicsAssetFactory.cpp
--- a/Engine/Plugins/Interchange/Runtime/Source/Import/Private/Mesh/InterchangePhysicsAssetFactory.cpp
+++ b/Engine/Plugins/Interchange/Runtime/Source/Import/Private/Mesh/InterchangePhysicsAssetFactory.cpp
@@ -33,8 +33,15 @@ UInterchangeFactoryBase::FImportAssetResult UInterchangePhysicsAssetFactory::Beg
 		return ImportAssetResult;
 	}
 
-	// create an asset if it doesn't exist
-	UObject* ExistingAsset = StaticFindObject(nullptr, Arguments.Parent, *Arguments.AssetName);
+	UObject* ExistingAsset = Arguments.ReimportObject;
+	if (!ExistingAsset)
+	{
+		FSoftObjectPath ReferenceObject;
+		if (PhysicsAssetNode->GetCustomReferenceObject(ReferenceObject))
+		{
+			ExistingAsset = ReferenceObject.TryLoad();
+		}
+	}
 
 	// create a new PhysicsAsset or overwrite existing asset, if possible
 	if (!ExistingAsset)
diff --git a/Engine/Plugins/Interchange/Runtime/Source/Import/Private/Mesh/InterchangeSkeletalMeshFactory.cpp b/Engine/Plugins/Interchange/Runtime/Source/Import/Private/Mesh/InterchangeSkeletalMeshFactory.cpp
--- a/Engine/Plugins/Interchange/Runtime/Source/Import/Private/Mesh/InterchangeSkeletalMeshFactory.cpp
+++ b/Engine/Plugins/Interchange/Runtime/Source/Import/Private/Mesh/InterchangeSkeletalMeshFactory.cpp
@@ -39,6 +39,7 @@
 #include "Rendering/SkeletalMeshModel.h"
 #include "SkeletalMeshAttributes.h"
 #include "SkeletalMeshOperations.h"
+#include "SkinnedAssetCompiler.h"
 
 #if WITH_EDITOR
 #include "LODUtilities.h"
@@ -909,25 +910,33 @@ UInterchangeFactoryBase::FImportAssetResult UInterchangeSkeletalMeshFactory::Beg
 		return ImportAssetResult;
 	}
 
-	// create an asset if it doesn't exist
-	UObject* ExistingAsset = StaticFindObject(nullptr, Arguments.Parent, *Arguments.AssetName);
+	UObject* ExistingAsset = Arguments.ReimportObject;
+	if (!ExistingAsset)
+	{
+		FSoftObjectPath ReferenceObject;
+		if (SkeletalMeshFactoryNode->GetCustomReferenceObject(ReferenceObject))
+		{
+			ExistingAsset = ReferenceObject.TryLoad();
+		}
+	}
 
 	// create a new material or overwrite existing asset, if possible
 	if (!ExistingAsset)
 	{
 		SkeletalMesh = NewObject<USkeletalMesh>(Arguments.Parent, *Arguments.AssetName, RF_Public | RF_Standalone);
 	}
-	else if (ExistingAsset->GetClass()->IsChildOf(USkeletalMesh::StaticClass()))
+	else
 	{
-		//This is a reimport, we are just re-updating the source data
+		//This is a reimport or an override, simply cast it to USkeletalMesh, the class has been verified by the caller (UE::Interchange::FTaskImportObject_GameThread::DoTask)
 		SkeletalMesh = Cast<USkeletalMesh>(ExistingAsset);
 	}
 	
-	if (!SkeletalMesh)
+	//This should not happen
+	if (!ensure(SkeletalMesh))
 	{
 		if (!Arguments.ReimportObject)
 		{
-			UE_LOG(LogInterchangeImport, Warning, TEXT("Could not create SkeletalMesh asset %s"), *Arguments.AssetName);
+			UE_LOG(LogInterchangeImport, Warning, TEXT("Could not create or find a SkeletalMesh asset name %s"), *Arguments.AssetName);
 		}
 		return ImportAssetResult;
 	}
diff --git a/Engine/Plugins/Interchange/Runtime/Source/Import/Private/Mesh/InterchangeSkeletonFactory.cpp b/Engine/Plugins/Interchange/Runtime/Source/Import/Private/Mesh/InterchangeSkeletonFactory.cpp
--- a/Engine/Plugins/Interchange/Runtime/Source/Import/Private/Mesh/InterchangeSkeletonFactory.cpp
+++ b/Engine/Plugins/Interchange/Runtime/Source/Import/Private/Mesh/InterchangeSkeletonFactory.cpp
@@ -37,15 +37,22 @@ UInterchangeFactoryBase::FImportAssetResult UInterchangeSkeletonFactory::BeginIm
 
 	bool bIsReImport = Arguments.ReimportObject != nullptr;
 
-	// create an asset if it doesn't exist
-	UObject* ExistingAsset = StaticFindObject(nullptr, Arguments.Parent, *Arguments.AssetName);
+	UObject* ExistingAsset = Arguments.ReimportObject;
+	if (!ExistingAsset)
+	{
+		FSoftObjectPath ReferenceObject;
+		if (SkeletonNode->GetCustomReferenceObject(ReferenceObject))
+		{
+			ExistingAsset = ReferenceObject.TryLoad();
+		}
+	}
 
 	// create a new skeleton or overwrite existing asset, if possible
 	if (!ExistingAsset)
 	{
 		Skeleton = NewObject<USkeleton>(Arguments.Parent, USkeleton::StaticClass(), *Arguments.AssetName, RF_Public | RF_Standalone);
 	}
-	else if (ExistingAsset->GetClass()->IsChildOf(USkeleton::StaticClass()))
+	else
 	{
 		//This is a reimport, we are just re-updating the source data
 		Skeleton = Cast<USkeleton>(ExistingAsset);
diff --git a/Engine/Plugins/Interchange/Runtime/Source/Import/Private/Mesh/InterchangeStaticMeshFactory.cpp b/Engine/Plugins/Interchange/Runtime/Source/Import/Private/Mesh/InterchangeStaticMeshFactory.cpp
--- a/Engine/Plugins/Interchange/Runtime/Source/Import/Private/Mesh/InterchangeStaticMeshFactory.cpp
+++ b/Engine/Plugins/Interchange/Runtime/Source/Import/Private/Mesh/InterchangeStaticMeshFactory.cpp
@@ -36,6 +36,7 @@
 #include "PhysicsEngine/SphereElem.h"
 #include "PhysicsEngine/SphylElem.h"
 #include "StaticMeshAttributes.h"
+#include "StaticMeshCompiler.h"
 #include "StaticMeshOperations.h"
 
 #include UE_INLINE_GENERATED_CPP_BY_NAME(InterchangeStaticMeshFactory)
@@ -159,15 +160,22 @@ UInterchangeFactoryBase::FImportAssetResult UInterchangeStaticMeshFactory::Begin
 		return ImportAssetResult;
 	}
 
-	// create an asset if it doesn't exist
-	UObject* ExistingAsset = (Arguments.ReimportObject != nullptr) ? Arguments.ReimportObject : StaticFindObject(nullptr, Arguments.Parent, *Arguments.AssetName);
+	UObject* ExistingAsset = Arguments.ReimportObject;
+	if (!ExistingAsset)
+	{
+		FSoftObjectPath ReferenceObject;
+		if (StaticMeshFactoryNode->GetCustomReferenceObject(ReferenceObject))
+		{
+			ExistingAsset = ReferenceObject.TryLoad();
+		}
+	}
 
 	// create a new static mesh or overwrite existing asset, if possible
 	if (!ExistingAsset)
 	{
 		StaticMesh = NewObject<UStaticMesh>(Arguments.Parent, *Arguments.AssetName, RF_Public | RF_Standalone);
 	}
-	else if (ExistingAsset->GetClass()->IsChildOf(UStaticMesh::StaticClass()))
+	else
 	{
 		//This is a reimport, we are just re-updating the source data
 		StaticMesh = Cast<UStaticMesh>(ExistingAsset);
diff --git a/Engine/Plugins/Interchange/Runtime/Source/Import/Private/Scene/InterchangeSceneVariantSetsFactory.cpp b/Engine/Plugins/Interchange/Runtime/Source/Import/Private/Scene/InterchangeSceneVariantSetsFactory.cpp
--- a/Engine/Plugins/Interchange/Runtime/Source/Import/Private/Scene/InterchangeSceneVariantSetsFactory.cpp
+++ b/Engine/Plugins/Interchange/Runtime/Source/Import/Private/Scene/InterchangeSceneVariantSetsFactory.cpp
@@ -302,14 +302,19 @@ UInterchangeFactoryBase::FImportAssetResult UInterchangeSceneVariantSetsFactory:
 	UE_LOG(LogInterchangeImport, Error, TEXT("Cannot import levelsequence asset in runtime, this is an editor only feature."));
 	return ImportAssetResult;
 #else
-	if (Arguments.ReimportObject)
+	auto CannotReimportMessage = [&Arguments, this]()
 	{
 		UInterchangeResultError_Generic* Message = AddMessage<UInterchangeResultError_Generic>();
 		Message->SourceAssetName = Arguments.SourceData->GetFilename();
 		Message->DestinationAssetName = Arguments.AssetName;
 		Message->AssetType = ULevelVariantSets::StaticClass();
 		Message->Text = LOCTEXT("CreateEmptyAssetUnsupportedReimport", "Re-import of ULevelVariantSets not supported yet.");
+		Arguments.AssetNode->SetSkipNodeImport();
+	};
 
+	if (Arguments.ReimportObject)
+	{
+		CannotReimportMessage();
 		return ImportAssetResult;
 	}
 
@@ -331,18 +336,28 @@ UInterchangeFactoryBase::FImportAssetResult UInterchangeSceneVariantSetsFactory:
 		return ImportAssetResult;
 	}
 
-	// create an asset if it doesn't exist
-	UObject* ExistingAsset = StaticFindObject(nullptr, Arguments.Parent, *Arguments.AssetName);
+	UObject* ExistingAsset = Arguments.ReimportObject;
+	if (!ExistingAsset)
+	{
+		FSoftObjectPath ReferenceObject;
+		if (FactoryNode->GetCustomReferenceObject(ReferenceObject))
+		{
+			ExistingAsset = ReferenceObject.TryLoad();
+		}
+	}
 
 	// create a new material or overwrite existing asset, if possible
 	if (!ExistingAsset)
 	{
 		LevelVariantSets = NewObject<ULevelVariantSets>(Arguments.Parent, *Arguments.AssetName, RF_Public | RF_Standalone);
 	}
-	else if (ExistingAsset->GetClass()->IsChildOf(ULevelVariantSets::StaticClass()))
+	else
 	{
+		CannotReimportMessage();
+		return ImportAssetResult;
+		// TODO put back the cast when the re-import will be supported
 		// This is a reimport, we are just re-updating the source data
-		LevelVariantSets = Cast<ULevelVariantSets>(ExistingAsset);
+		//LevelVariantSets = Cast<ULevelVariantSets>(ExistingAsset);
 	}
 
 	if (!LevelVariantSets)
@@ -351,12 +366,12 @@ UInterchangeFactoryBase::FImportAssetResult UInterchangeSceneVariantSetsFactory:
 		return ImportAssetResult;
 	}
 
-	ImportAssetResult.ImportedObject = ImportObjectSourceData(Arguments);
+	ImportAssetResult.ImportedObject = ImportObjectSourceData(Arguments, LevelVariantSets);
 	return ImportAssetResult;
 #endif //else !WITH_EDITOR || !WITH_EDITORONLY_DATA
 }
 
-UObject* UInterchangeSceneVariantSetsFactory::ImportObjectSourceData(const FImportAssetObjectParams& Arguments)
+UObject* UInterchangeSceneVariantSetsFactory::ImportObjectSourceData(const FImportAssetObjectParams& Arguments, ULevelVariantSets* LevelVariantSets)
 {
 #if !WITH_EDITOR || !WITH_EDITORONLY_DATA
 	// TODO: Can we import ULevelVariantSets at runtime
@@ -366,20 +381,8 @@ UObject* UInterchangeSceneVariantSetsFactory::ImportObjectSourceData(const FImpo
 #else
 	using namespace UE::Interchange;
 
-	// Re-import is not supported yet
-	// Need to add an AssetImportData property to ULevelVariantSets
-	if (Arguments.ReimportObject)
-	{
-		UInterchangeResultError_Generic* Message = AddMessage<UInterchangeResultError_Generic>();
-		Message->SourceAssetName = Arguments.SourceData->GetFilename();
-		Message->DestinationAssetName = Arguments.AssetName;
-		Message->AssetType = ULevelVariantSets::StaticClass();
-		Message->Text = LOCTEXT("CreateAssetUnsupportedReimport", "Re-import of ULevelVariantSets not supported yet.");
 
-		return nullptr;
-	}
-
-	if (!Arguments.NodeContainer || !Arguments.AssetNode || !Arguments.AssetNode->GetObjectClass()->IsChildOf(GetFactoryClass()))
+	if (!LevelVariantSets || !Arguments.NodeContainer || !Arguments.AssetNode || !Arguments.AssetNode->GetObjectClass()->IsChildOf(GetFactoryClass()))
 	{
 		return nullptr;
 	}
@@ -401,34 +404,6 @@ UObject* UInterchangeSceneVariantSetsFactory::ImportObjectSourceData(const FImpo
 	const UClass* LevelVariantSetsClass = FactoryNode->GetObjectClass();
 	check(LevelVariantSetsClass && LevelVariantSetsClass->IsChildOf(GetFactoryClass()));
 
-	// create an asset if it doesn't exist
-	UObject* ExistingAsset = StaticFindObject(nullptr, Arguments.Parent, *Arguments.AssetName);
-
-	ULevelVariantSets* LevelVariantSets = nullptr;
-	// create a new level sequence or overwrite existing asset, if possible
-	if (!ExistingAsset)
-	{
-		//NewObject is not thread safe, the asset registry directory watcher tick on the main thread can trig before we finish initializing the UObject and will crash
-		//The UObject should have been create by calling CreateEmptyAsset on the main thread.
-		check(IsInGameThread());
-		LevelVariantSets = NewObject<ULevelVariantSets>(Arguments.Parent, LevelVariantSetsClass, *Arguments.AssetName, RF_Public | RF_Standalone);
-	}
-	else if (ExistingAsset->GetClass()->IsChildOf(LevelVariantSetsClass))
-	{
-		//This is a reimport, we are just re-updating the source data
-		LevelVariantSets = Cast<ULevelVariantSets>(ExistingAsset);
-	}
-
-	if (!ensure(LevelVariantSets))
-	{
-		UInterchangeResultError_Generic* Message = AddMessage<UInterchangeResultError_Generic>();
-		Message->SourceAssetName = Arguments.SourceData->GetFilename();
-		Message->DestinationAssetName = Arguments.AssetName;
-		Message->AssetType = ULevelVariantSets::StaticClass();
-		Message->Text = FText::Format(LOCTEXT("CreateAssetFailed", "Could not create nor find LevelVariantSets asset {0}."), FText::FromString(Arguments.AssetName));
-		return nullptr;
-	}
-
 	Private::FLevelVariantSetsHelper Helper(*LevelVariantSets, *FactoryNode, *Arguments.NodeContainer, *LevelVariantSetsPayloadInterface);
 	Helper.PopulateLevelVariantSets();
 
diff --git a/Engine/Plugins/Interchange/Runtime/Source/Import/Private/Texture/InterchangeTextureFactory.cpp b/Engine/Plugins/Interchange/Runtime/Source/Import/Private/Texture/InterchangeTextureFactory.cpp
--- a/Engine/Plugins/Interchange/Runtime/Source/Import/Private/Texture/InterchangeTextureFactory.cpp
+++ b/Engine/Plugins/Interchange/Runtime/Source/Import/Private/Texture/InterchangeTextureFactory.cpp
@@ -1531,7 +1531,7 @@ UInterchangeFactoryBase::FImportAssetResult UInterchangeTextureFactory::BeginImp
 {
 	using namespace  UE::Interchange::Private::InterchangeTextureFactory;
 	FImportAssetResult ImportAssetResult;
-	UObject* Texture = nullptr;
+	UTexture* Texture = nullptr;
 
 	auto CouldNotCreateTextureLog = [this, &Arguments, &ImportAssetResult](const FText& Info)
 	{
@@ -1582,16 +1582,22 @@ UInterchangeFactoryBase::FImportAssetResult UInterchangeTextureFactory::BeginImp
 
 	const bool bIsReimport = Arguments.ReimportObject != nullptr;
 
-	// Find existing asset
-	UObject* ExistingAsset = bIsReimport ? Arguments.ReimportObject : StaticFindObject(nullptr, Arguments.Parent, *Arguments.AssetName);
+	UObject* ExistingAsset = Arguments.ReimportObject;
+	if (!ExistingAsset)
+	{
+		FSoftObjectPath ReferenceObject;
+		if (Arguments.AssetNode->GetCustomReferenceObject(ReferenceObject))
+		{
+			ExistingAsset = ReferenceObject.TryLoad();
+		}
+	}
 
 	// create a new texture or overwrite existing asset, if possible
 	if (!ExistingAsset)
 	{
-		UTexture* NewTexture = NewObject<UTexture>(Arguments.Parent, TextureClass, *Arguments.AssetName, RF_Public | RF_Standalone);
-		Texture = NewTexture;
+		Texture = NewObject<UTexture>(Arguments.Parent, TextureClass, *Arguments.AssetName, RF_Public | RF_Standalone);
 
-		if (UTextureLightProfile* LightProfile = Cast<UTextureLightProfile>(NewTexture))
+		if (UTextureLightProfile* LightProfile = Cast<UTextureLightProfile>(Texture))
 		{
 			LightProfile->AddressX = TA_Clamp;
 			LightProfile->AddressY = TA_Clamp;
@@ -1601,10 +1607,9 @@ UInterchangeFactoryBase::FImportAssetResult UInterchangeTextureFactory::BeginImp
 			LightProfile->LODGroup = TEXTUREGROUP_IESLightProfile;
 		}
 	}
-	else if (ExistingAsset->GetClass()->IsChildOf(TextureClass))
+	else
 	{
-		//This is a reimport, we are just re-updating the source data
-		Texture = ExistingAsset;
+		Texture = Cast<UTexture>(ExistingAsset);
 		//We allow override of existing Textures only if the translator is a pure texture translator or the user directly ask to re-import this object
 		if (!bIsReimport && Arguments.Translator->GetSupportedAssetTypes() != EInterchangeTranslatorAssetType::Textures)
 		{
@@ -1613,21 +1618,6 @@ UInterchangeFactoryBase::FImportAssetResult UInterchangeTextureFactory::BeginImp
 			bSkipImport = true;
 		}
 	}
-	else
-	{
-		UInterchangeResultError_Generic* Message = AddMessage<UInterchangeResultError_Generic>();
-		Message->SourceAssetName = Arguments.SourceData->GetFilename();
-		Message->DestinationAssetName = Arguments.AssetName;
-		Message->AssetType = TextureClass;
-
-		const FText TargetClassName = FText::FromString(TextureClass->GetName());
-		const FText ExistingClassName = FText::FromString(ExistingAsset->GetClass()->GetName());
-		const FText AssetName = FText::FromString(Arguments.AssetName);
-		const FText FolderName = FText::FromString(FPaths::GetPath(Arguments.Parent->GetPathName()));
-		Message->Text = FText::Format(NSLOCTEXT("InterchangeTextureFactory", "ClassMismatch", "You cannot create a '{0}' asset named '{1}' in '{2}', as there is already a '{3}' asset with the same name in this folder."), TargetClassName, AssetName, FolderName, ExistingClassName);
-		
-		return ImportAssetResult;
-	}
 
 	if (!Texture)
 	{
@@ -1724,9 +1714,9 @@ UInterchangeFactoryBase::FImportAssetResult UInterchangeTextureFactory::ImportAs
 		ImportTextureErrorLog(LOCTEXT("TextureFactory_Async_CannotCreateAsync", "UInterchangeTextureFactory: Could not create Texture asset outside of the game thread."));
 		return ImportAssetResult;
 	}
-	else if(ExistingAsset->GetClass()->IsChildOf(TextureClass))
+	else
 	{
-		Texture = static_cast<UTexture*>(ExistingAsset);
+		Texture = Cast<UTexture>(ExistingAsset);
 	}
 
 	if (!Texture)
diff --git a/Engine/Plugins/Interchange/Runtime/Source/Import/Public/Scene/InterchangeSceneVariantSetsFactory.h b/Engine/Plugins/Interchange/Runtime/Source/Import/Public/Scene/InterchangeSceneVariantSetsFactory.h
--- a/Engine/Plugins/Interchange/Runtime/Source/Import/Public/Scene/InterchangeSceneVariantSetsFactory.h
+++ b/Engine/Plugins/Interchange/Runtime/Source/Import/Public/Scene/InterchangeSceneVariantSetsFactory.h
@@ -11,6 +11,7 @@
 
 #include "InterchangeSceneVariantSetsFactory.generated.h"
 
+class ULevelVariantSets;
 class UVariantObjectBinding;
 
 UCLASS(BlueprintType, Experimental)
@@ -31,7 +32,7 @@ public:
 	//////////////////////////////////////////////////////////////////////////
 
 private:
-	UObject* ImportObjectSourceData(const FImportAssetObjectParams& Arguments);
+	UObject* ImportObjectSourceData(const FImportAssetObjectParams& Arguments, ULevelVariantSets* LevelVariantSets);
 	const UInterchangeTranslatorBase* Translator = nullptr;
 };
 
diff --git a/Engine/Plugins/Interchange/Runtime/Source/Pipelines/Private/InterchangePipelineMeshesUtilities.cpp b/Engine/Plugins/Interchange/Runtime/Source/Pipelines/Private/InterchangePipelineMeshesUtilities.cpp
--- a/Engine/Plugins/Interchange/Runtime/Source/Pipelines/Private/InterchangePipelineMeshesUtilities.cpp
+++ b/Engine/Plugins/Interchange/Runtime/Source/Pipelines/Private/InterchangePipelineMeshesUtilities.cpp
@@ -266,15 +266,18 @@ UInterchangePipelineMeshesUtilities* UInterchangePipelineMeshesUtilities::Create
 					if (IsSceneNodeASocket(SceneNode))
 					{
 						FString MeshUid;
-						const UInterchangeSceneNode* ParentMeshSceneNode = Cast<UInterchangeSceneNode>(BaseNodeContainer->GetNode(SceneNode->GetParentUid()));
-						while (ParentMeshSceneNode)
+						if (!SceneNode->GetCustomAssetInstanceUid(MeshUid))
 						{
-							if (ParentMeshSceneNode->GetCustomAssetInstanceUid(MeshUid))
+							const UInterchangeSceneNode* ParentMeshSceneNode = Cast<UInterchangeSceneNode>(BaseNodeContainer->GetNode(SceneNode->GetParentUid()));
+							while (ParentMeshSceneNode)
 							{
-								break;
-							}
+								if (ParentMeshSceneNode->GetCustomAssetInstanceUid(MeshUid))
+								{
+									break;
+								}
 
-							ParentMeshSceneNode = Cast<UInterchangeSceneNode>(BaseNodeContainer->GetNode(SceneNode->GetParentUid()));
+								ParentMeshSceneNode = Cast<UInterchangeSceneNode>(BaseNodeContainer->GetNode(ParentMeshSceneNode->GetParentUid()));
+							}
 						}
 
 						if (!MeshUid.IsEmpty())
diff --git a/Engine/Plugins/Interchange/Runtime/Source/Pipelines/Public/InterchangeGenericAssetsPipeline.h b/Engine/Plugins/Interchange/Runtime/Source/Pipelines/Public/InterchangeGenericAssetsPipeline.h
--- a/Engine/Plugins/Interchange/Runtime/Source/Pipelines/Public/InterchangeGenericAssetsPipeline.h
+++ b/Engine/Plugins/Interchange/Runtime/Source/Pipelines/Public/InterchangeGenericAssetsPipeline.h
@@ -83,6 +83,11 @@ public:
 
 	virtual void AdjustSettingsForContext(EInterchangePipelineContext ImportType, TObjectPtr<UObject> ReimportAsset) override;
 
+	virtual bool IsScripted() override
+	{
+		return false;
+	}
+
 #if WITH_EDITOR
 	virtual bool GetPropertyPossibleValues(const FName PropertyPath, TArray<FString>& PossibleValues) override;
 #endif
@@ -95,8 +100,10 @@ protected:
 
 	virtual bool CanExecuteOnAnyThread(EInterchangePipelineTask PipelineTask) override
 	{
-		//If a blueprint or python derived from this class, it will be execute on the game thread since we cannot currently execute script outside of the game thread, even if this function return true.
-		return true;
+		//We cannot run asynchronously because of the two following issues
+		// Post Translator Task: material pipeline is loading assets (Parent Material)
+		// Post Import Task: physics asset need to create a scene preview to be created
+		return false;
 	}
 
 	virtual void SetReimportSourceIndex(UClass* ReimportObjectClass, const int32 SourceFileIndex) override;
diff --git a/Engine/Source/Runtime/Interchange/Core/Private/Nodes/InterchangeFactoryBaseNode.cpp b/Engine/Source/Runtime/Interchange/Core/Private/Nodes/InterchangeFactoryBaseNode.cpp
--- a/Engine/Source/Runtime/Interchange/Core/Private/Nodes/InterchangeFactoryBaseNode.cpp
+++ b/Engine/Source/Runtime/Interchange/Core/Private/Nodes/InterchangeFactoryBaseNode.cpp
@@ -21,11 +21,17 @@ namespace UE
 			return BaseNodeFactoryDependencies_BaseKey;
 		}
 
-		const FAttributeKey& FFactoryBaseNodeStaticData::ReimportStrategyFlagsKey()
-		{
-			static FAttributeKey AttributeKey(TEXT("__Reimport_Strategy_Flags_Key__"));
-			return AttributeKey;
-		}
+const FAttributeKey& FFactoryBaseNodeStaticData::ReimportStrategyFlagsKey()
+{
+	static FAttributeKey AttributeKey(TEXT("__Reimport_Strategy_Flags_Key__"));
+	return AttributeKey;
+}
+
+const FAttributeKey& FFactoryBaseNodeStaticData::SkipNodeImportKey()
+{
+	static FAttributeKey AttributeKey(TEXT("__Internal_Task_Skip_Node__"));
+	return AttributeKey;
+}
 
 	} //ns Interchange
 } //ns UE
@@ -111,6 +117,43 @@ bool UInterchangeFactoryBaseNode::SetReimportStrategyFlags(const EReimportStrate
 	return false;
 }
 
+bool UInterchangeFactoryBaseNode::ShouldSkipNodeImport() const
+{
+	if (!Attributes->ContainAttribute(UE::Interchange::FFactoryBaseNodeStaticData::SkipNodeImportKey()))
+	{
+		return false;
+	}
+	bool bShouldSkipNodeImport = false;
+	Attributes->GetAttributeHandle<bool>(UE::Interchange::FFactoryBaseNodeStaticData::SkipNodeImportKey()).Get(bShouldSkipNodeImport);
+	return bShouldSkipNodeImport;
+}
+
+bool UInterchangeFactoryBaseNode::SetSkipNodeImport()
+{
+	UE::Interchange::EAttributeStorageResult Result = Attributes->RegisterAttribute(UE::Interchange::FFactoryBaseNodeStaticData::SkipNodeImportKey(), true);
+	if (IsAttributeStorageResultSuccess(Result))
+	{
+		UE::Interchange::FAttributeStorage::TAttributeHandle<bool> Handle = Attributes->GetAttributeHandle<bool>(UE::Interchange::FFactoryBaseNodeStaticData::SkipNodeImportKey());
+		return Handle.IsValid();
+	}
+	return false;
+}
+
+bool UInterchangeFactoryBaseNode::UnsetSkipNodeImport()
+{
+	if (!Attributes->ContainAttribute(UE::Interchange::FFactoryBaseNodeStaticData::SkipNodeImportKey()))
+	{
+		return true;
+	}
+	UE::Interchange::EAttributeStorageResult Result = Attributes->UnregisterAttribute(UE::Interchange::FFactoryBaseNodeStaticData::SkipNodeImportKey());
+	if (IsAttributeStorageResultSuccess(Result))
+	{
+		UE::Interchange::FAttributeStorage::TAttributeHandle<bool> Handle = Attributes->GetAttributeHandle<bool>(UE::Interchange::FFactoryBaseNodeStaticData::SkipNodeImportKey());
+		return !Handle.IsValid();
+	}
+	return false;
+}
+
 bool UInterchangeFactoryBaseNode::GetCustomSubPath(FString& AttributeValue) const
 {
 	IMPLEMENT_NODE_ATTRIBUTE_GETTER(SubPath, FString)
diff --git a/Engine/Source/Runtime/Interchange/Core/Public/Nodes/InterchangeFactoryBaseNode.h b/Engine/Source/Runtime/Interchange/Core/Public/Nodes/InterchangeFactoryBaseNode.h
--- a/Engine/Source/Runtime/Interchange/Core/Public/Nodes/InterchangeFactoryBaseNode.h
+++ b/Engine/Source/Runtime/Interchange/Core/Public/Nodes/InterchangeFactoryBaseNode.h
@@ -134,6 +134,7 @@ namespace UE::Interchange
 	{
 		static const FString& FactoryDependenciesBaseKey();
 		static const FAttributeKey& ReimportStrategyFlagsKey();
+		static const FAttributeKey& SkipNodeImportKey();
 	};
 
 } // namespace UE::Interchange
@@ -179,6 +180,28 @@ public:
 	UFUNCTION(BlueprintCallable, Category = "Interchange | Node")
 	bool SetReimportStrategyFlags(const EReimportStrategyFlags& ReimportStrategyFlags);
 
+	/**
+	 * Return true if this node should skip the factory import process, false otherwise.
+	 * Node can be in a situation where we have to skip the import process because we cannot import the associated asset for multiple reason.
+	 * - An asset can already exist and is representing a different type (UClass)
+	 * - An asset can already exist and is being compiled
+	 * - An asset can already exist and is being imported by another concurrent import task (user importing multiple file in same time in the same content folder)
+	 */
+	UFUNCTION(BlueprintCallable, Category = "Interchange | Node")
+	bool ShouldSkipNodeImport() const;
+
+	/**
+	 * Add the skip node attribute. Use this function to cancel the creation of the unreal asset, see ShouldSkipNodeImport for more documentation.
+	 */
+	UFUNCTION(BlueprintCallable, Category = "Interchange | Node")
+	bool SetSkipNodeImport();
+
+	/**
+	 * Remove the skip node attribute. See ShouldSkipNodeImport for more documentation.
+	 */
+	UFUNCTION(BlueprintCallable, Category = "Interchange | Node")
+	bool UnsetSkipNodeImport();
+
 	/**
 	 * Return the UClass of the object we represent so we can find factory/writer
 	 */
diff --git a/Engine/Source/Runtime/Interchange/Engine/Private/InterchangeManager.cpp b/Engine/Source/Runtime/Interchange/Engine/Private/InterchangeManager.cpp
--- a/Engine/Source/Runtime/Interchange/Engine/Private/InterchangeManager.cpp
+++ b/Engine/Source/Runtime/Interchange/Engine/Private/InterchangeManager.cpp
@@ -212,6 +212,174 @@ void UE::Interchange::FImportAsyncHelper::AddReferencedObjects(FReferenceCollect
 	Collector.AddReferencedObjects(CreatedFactories);
 }
 
+/*
+* // Created factories map, Key is factory node UID
+			mutable FCriticalSection CreatedFactoriesLock;
+			TMap<FString, UInterchangeFactoryBase*> CreatedFactories;
+*/
+
+//Create package map, Key is package name. We cannot create package asynchronously so we have to create a game thread task to do this
+UPackage* UE::Interchange::FImportAsyncHelper::GetCreatedPackage(const FString& PackageName) const
+{
+	FScopeLock Lock(&CreatedPackagesLock);
+	if (UPackage* const* PkgPtr = CreatedPackages.Find(PackageName))
+	{
+		return *PkgPtr;
+	}
+	return nullptr;
+}
+
+void UE::Interchange::FImportAsyncHelper::AddCreatedPackage(const FString& PackageName, UPackage* Package)
+{
+	FScopeLock Lock(&CreatedPackagesLock);
+	if (ensure(!CreatedPackages.Contains(PackageName)))
+	{
+		CreatedPackages.Add(PackageName, Package);
+	}
+}
+
+UInterchangeFactoryBase* UE::Interchange::FImportAsyncHelper::GetCreatedFactory(const FString& FactoryNodeUniqueId) const
+{
+	FScopeLock Lock(&CreatedFactoriesLock);
+	if (TObjectPtr<UInterchangeFactoryBase>const* FactoryPtr = CreatedFactories.Find(FactoryNodeUniqueId))
+	{
+		return *FactoryPtr;
+	}
+	return nullptr;
+}
+
+void UE::Interchange::FImportAsyncHelper::AddCreatedFactory(const FString& FactoryNodeUniqueId, UInterchangeFactoryBase* Factory)
+{
+	FScopeLock Lock(&CreatedFactoriesLock);
+	if (ensure(!CreatedFactories.Contains(FactoryNodeUniqueId)))
+	{
+		CreatedFactories.Add(FactoryNodeUniqueId, Factory);
+	}
+}
+
+UE::Interchange::FImportAsyncHelper::FImportedObjectInfo& UE::Interchange::FImportAsyncHelper::AddDefaultImportedAssetGetRef(int32 SourceIndex)
+{
+	FScopeLock Lock(&ImportedAssetsPerSourceIndexLock);
+	return ImportedAssetsPerSourceIndex.FindOrAdd(SourceIndex).AddDefaulted_GetRef();
+}
+
+const UE::Interchange::FImportAsyncHelper::FImportedObjectInfo* UE::Interchange::FImportAsyncHelper::FindImportedAssets(int32 SourceIndex, TFunction< bool(const UE::Interchange::FImportAsyncHelper::FImportedObjectInfo& ImportedObject) > Predicate) const
+{
+	FScopeLock Lock(&ImportedAssetsPerSourceIndexLock);
+	if (!ImportedAssetsPerSourceIndex.Contains(SourceIndex))
+	{
+		return nullptr;
+	}
+	const TArray<FImportedObjectInfo>& ImportedObjectInfos = ImportedAssetsPerSourceIndex.FindChecked(SourceIndex);
+	for (const FImportedObjectInfo& ImportObjectInfo : ImportedObjectInfos)
+	{
+		if (Predicate(ImportObjectInfo))
+		{
+			return &ImportObjectInfo;
+		}
+	}
+	return nullptr;
+}
+
+void UE::Interchange::FImportAsyncHelper::IterateImportedAssets(int32 SourceIndex, TFunction< void(const TArray<UE::Interchange::FImportAsyncHelper::FImportedObjectInfo>& ImportedObjects) > Callback) const
+{
+	FScopeLock Lock(&ImportedAssetsPerSourceIndexLock);
+	
+	if (!ImportedAssetsPerSourceIndex.Contains(SourceIndex))
+	{
+		return;
+	}
+
+	Callback(ImportedAssetsPerSourceIndex.FindChecked(SourceIndex));
+}
+
+void UE::Interchange::FImportAsyncHelper::IterateImportedAssetsPerSourceIndex(TFunction< void(int32 SourceIndex, const TArray<UE::Interchange::FImportAsyncHelper::FImportedObjectInfo>& ImportedObjects) > Callback) const
+{
+	FScopeLock Lock(&ImportedAssetsPerSourceIndexLock);
+	for (const TPair<int32, TArray<FImportedObjectInfo>>& SourceIndexAndImportedAssets : ImportedAssetsPerSourceIndex)
+	{
+		Callback(SourceIndexAndImportedAssets.Key, SourceIndexAndImportedAssets.Value);
+	}
+}
+
+UE::Interchange::FImportAsyncHelper::FImportedObjectInfo& UE::Interchange::FImportAsyncHelper::AddDefaultImportedSceneObjectGetRef(int32 SourceIndex)
+{
+	FScopeLock Lock(&ImportedSceneObjectsPerSourceIndexLock);
+	return ImportedSceneObjectsPerSourceIndex.FindOrAdd(SourceIndex).AddDefaulted_GetRef();
+}
+
+const UE::Interchange::FImportAsyncHelper::FImportedObjectInfo* UE::Interchange::FImportAsyncHelper::FindImportedSceneObjects(int32 SourceIndex, TFunction< bool(const UE::Interchange::FImportAsyncHelper::FImportedObjectInfo& ImportedObject) > Predicate) const
+{
+	FScopeLock Lock(&ImportedSceneObjectsPerSourceIndexLock);
+	if (!ImportedSceneObjectsPerSourceIndex.Contains(SourceIndex))
+	{
+		return nullptr;
+	}
+	const TArray<FImportedObjectInfo>& ImportedObjectInfos = ImportedSceneObjectsPerSourceIndex.FindChecked(SourceIndex);
+	for (const FImportedObjectInfo& ImportObjectInfo : ImportedObjectInfos)
+	{
+		if (Predicate(ImportObjectInfo))
+		{
+			return &ImportObjectInfo;
+		}
+	}
+	return nullptr;
+}
+
+void UE::Interchange::FImportAsyncHelper::IterateImportedSceneObjects(int32 SourceIndex, TFunction< void(const TArray<UE::Interchange::FImportAsyncHelper::FImportedObjectInfo>& ImportedObjects) > Callback) const
+{
+	FScopeLock Lock(&ImportedSceneObjectsPerSourceIndexLock);
+
+	if (!ImportedSceneObjectsPerSourceIndex.Contains(SourceIndex))
+	{
+		return;
+	}
+
+	Callback(ImportedSceneObjectsPerSourceIndex.FindChecked(SourceIndex));
+}
+
+void UE::Interchange::FImportAsyncHelper::IterateImportedSceneObjectsPerSourceIndex(TFunction< void(int32 SourceIndex, const TArray<UE::Interchange::FImportAsyncHelper::FImportedObjectInfo>& ImportedObjects) > Callback) const
+{
+	FScopeLock Lock(&ImportedSceneObjectsPerSourceIndexLock);
+	for (const TPair<int32, TArray<FImportedObjectInfo>>& SourceIndexAndImportedSceneObjects : ImportedSceneObjectsPerSourceIndex)
+	{
+		Callback(SourceIndexAndImportedSceneObjects.Key, SourceIndexAndImportedSceneObjects.Value);
+	}
+}
+
+bool UE::Interchange::FImportAsyncHelper::IsImportingObject(UObject* Object) const
+{
+	if (!Object)
+	{
+		return false;
+	}
+
+	bool bFoundAsset = false;
+	auto IsImportingAsset = [&bFoundAsset, Object](int32 SourceIndex, const TArray<FImportedObjectInfo>& ImportedObjects)
+	{
+		if (bFoundAsset)
+		{
+			return;
+		}
+		for (const FImportedObjectInfo& ImportedObjectInfo : ImportedObjects)
+		{
+			if (ImportedObjectInfo.ImportedObject == Object)
+			{
+				bFoundAsset = true;
+				break;
+			}
+		}
+	};
+	//Asset import
+	IterateImportedAssetsPerSourceIndex(IsImportingAsset);
+	if (!bFoundAsset)
+	{
+		IterateImportedSceneObjectsPerSourceIndex(IsImportingAsset);
+	}
+
+	return bFoundAsset;
+}
+
 void UE::Interchange::FImportAsyncHelper::SendAnalyticImportEndData()
 {
 	if (!FEngineAnalytics::IsAvailable())
@@ -639,7 +807,12 @@ UInterchangeManager& UInterchangeManager::GetInterchangeManager()
 		check(IsInGameThread() && GetTransientPackage());
 
 		InterchangeManager = TStrongObjectPtr<UInterchangeManager>(NewObject<UInterchangeManager>(GetTransientPackage(), NAME_None, EObjectFlags::RF_NoFlags));
-		
+
+		InterchangeManager->GCEndDelegate = FCoreUObjectDelegates::GetPostGarbageCollect().AddLambda([]()
+			{
+				InterchangeManager->StartQueuedTasks(InterchangeManager->bGCEndDelegateCancellAllTask);
+			});
+
 		//We cancel any running task when we pre exit the engine
 		FCoreDelegates::OnEnginePreExit.AddLambda([]()
 		{
@@ -659,6 +832,12 @@ UInterchangeManager& UInterchangeManager::GetInterchangeManager()
 				InterchangeManager->WaitUntilAllTasksDone(bCancel);
 			}
 
+			//Remove any delegate
+			if (InterchangeManager->GCEndDelegate.IsValid())
+			{
+				FCoreUObjectDelegates::GetPostGarbageCollect().Remove(InterchangeManager->GCEndDelegate);
+				InterchangeManager->GCEndDelegate.Reset();
+			}
 			//Task should have been cancel in the Engine pre exit callback
 			ensure(InterchangeManager->ImportTasks.Num() == 0);
 			InterchangeManager->OnPreDestroyInterchangeManager.Broadcast();
@@ -877,6 +1056,20 @@ void UInterchangeManager::StartQueuedTasks(bool bCancelAllTasks /*= false*/)
 		return;
 	}
 
+	//Garbage collect can stall and tick the task graph (if accessing a compiling asset locked UProperty )
+	//We must avoid starting an import task in this case, import cannot be done when GC runs.
+	//A delegate is implemented and call StartQueuedTasks when the GC is finish
+	if(IsGarbageCollecting())
+	{
+		if (bCancelAllTasks)
+		{
+			bGCEndDelegateCancellAllTask = bCancelAllTasks;
+		}
+		return;
+	}
+
+	bGCEndDelegateCancellAllTask = false;
+
 	auto UpdateNotification = [this]()
 	{
 		if (Notification.IsValid())
@@ -922,8 +1115,10 @@ void UInterchangeManager::StartQueuedTasks(bool bCancelAllTasks /*= false*/)
 	};
 
 
-	//Hack, import the jobs one by one
-	const int32 MaxNumWorker = 1;
+	//We need to leave some free task in the pool to avoid deadlock.
+	//Each import can use 2 tasks in same time if the build of the asset ddc use the same task pool (i.e. staticmesh, skeletalmesh, texture...)
+	const int32 PoolWorkerThreadCount = FTaskGraphInterface::Get().GetNumWorkerThreads() / 2;
+	const int32 MaxNumWorker = FMath::Max(PoolWorkerThreadCount, 1);
 	while (!QueuedTasks.IsEmpty() && (ImportTasks.Num() < MaxNumWorker || bCancelAllTasks))
 	{
 		FQueuedTaskData QueuedTaskData;
@@ -1328,6 +1523,14 @@ UInterchangeManager::ImportInternal(const FString& ContentPath, const UInterchan
 				{
 					//When we do not show the UI we use the original stack
 					OutPipelines = StackInfoPtr->Pipelines;
+					if (bImportAllWithDefault)
+					{
+						//When the user want to use the same settings for all source data we just need to load the settings save by the initial dialog
+						for (UInterchangePipelineBase* Pipeline : OutPipelines)
+						{
+							Pipeline->LoadSettings(DefaultStackName);
+						}
+					}
 				}
 				else if (PipelineStacks.Num() > 0)
 				{
@@ -1416,6 +1619,24 @@ UInterchangeManager::ImportInternal(const FString& ContentPath, const UInterchan
 	return TTuple<UE::Interchange::FAssetImportResultRef, UE::Interchange::FSceneImportResultRef>{ AsyncHelper->AssetImportResult, AsyncHelper->SceneImportResult };
 }
 
+bool UInterchangeManager::IsObjectBeingImported(UObject* Object) const
+{
+	if (!ensure(IsInGameThread()))
+	{
+		return false;
+	}
+
+	for (const TSharedPtr < UE::Interchange::FImportAsyncHelper, ESPMode::ThreadSafe>& AsyncHelper : ImportTasks)
+	{
+		if (AsyncHelper->IsImportingObject(Object))
+		{
+			return true;
+		}
+	}
+
+	return false;
+}
+
 bool UInterchangeManager::IsInterchangeImportEnabled()
 {
 	return CCvarInterchangeImportEnable->GetBool();
diff --git a/Engine/Source/Runtime/Interchange/Engine/Private/Tasks/InterchangeTaskCompletion.cpp b/Engine/Source/Runtime/Interchange/Engine/Private/Tasks/InterchangeTaskCompletion.cpp
--- a/Engine/Source/Runtime/Interchange/Engine/Private/Tasks/InterchangeTaskCompletion.cpp
+++ b/Engine/Source/Runtime/Interchange/Engine/Private/Tasks/InterchangeTaskCompletion.cpp
@@ -46,155 +46,157 @@ void UE::Interchange::FTaskPreCompletion::DoTask(ENamedThreads::Type CurrentThre
 
 	UInterchangeResultsContainer* Results = AsyncHelper->AssetImportResult->GetResults();
 
-	auto ForEachImportedObjectForEachSource = [&AsyncHelper, &Results](TMap<int32, TArray<FImportAsyncHelper::FImportedObjectInfo>>& ImportedObjectsPerSourceIndex, bool bIsAsset)
+	bool bIsAsset = true;
+
+	auto IterationCallback = [&AsyncHelper, &Results, &bIsAsset](int32 SourceIndex, const TArray<FImportAsyncHelper::FImportedObjectInfo>& ImportedObjects)
 	{
-		for (TPair<int32, TArray<FImportAsyncHelper::FImportedObjectInfo>>& ObjectInfosPerSourceIndexPair : ImportedObjectsPerSourceIndex)
+		//Verify if the task was cancel
+		if (AsyncHelper->bCancel)
 		{
-			//Verify if the task was cancel
-			if (AsyncHelper->bCancel)
+			for (const FImportAsyncHelper::FImportedObjectInfo& ObjectInfo : ImportedObjects)
 			{
-				for (const FImportAsyncHelper::FImportedObjectInfo& ObjectInfo : ObjectInfosPerSourceIndexPair.Value)
+				//Cancel factories so they can do proper cleanup
+				if (ObjectInfo.Factory)
 				{
-					//Cancel factories so they can do proper cleanup
-					if (ObjectInfo.Factory)
-					{
-						ObjectInfo.Factory->Cancel();
-					}
+					ObjectInfo.Factory->Cancel();
 				}
-				//Skip if cancel
-				continue;
 			}
+			//Skip if cancel
+			return;
+		}
 
-			const int32 SourceIndex = ObjectInfosPerSourceIndexPair.Key;
-			const bool bCallPostImportGameThreadCallback = ensure(AsyncHelper->SourceDatas.IsValidIndex(SourceIndex));
+		const bool bCallPostImportGameThreadCallback = ensure(AsyncHelper->SourceDatas.IsValidIndex(SourceIndex));
 
-			//First iteration to call SetupObject_GameThread and pipeline ExecutePostFactoryPipeline
-			for (const FImportAsyncHelper::FImportedObjectInfo& ObjectInfo : ObjectInfosPerSourceIndexPair.Value)
+		//First iteration to call SetupObject_GameThread and pipeline ExecutePostFactoryPipeline
+		for (const FImportAsyncHelper::FImportedObjectInfo& ObjectInfo : ImportedObjects)
+		{
+			UObject* ImportedObject = ObjectInfo.ImportedObject;
+			//In case Some factory code cannot run outside of the main thread we offer this callback to finish the work before calling post edit change (building the asset)
+			if (bCallPostImportGameThreadCallback && ObjectInfo.Factory)
 			{
-				UObject* ImportedObject = ObjectInfo.ImportedObject;
-				//In case Some factory code cannot run outside of the main thread we offer this callback to finish the work before calling post edit change (building the asset)
-				if (bCallPostImportGameThreadCallback && ObjectInfo.Factory)
-				{
-					UInterchangeFactoryBase::FSetupObjectParams Arguments;
-					Arguments.ImportedObject = ImportedObject;
-					Arguments.SourceData = AsyncHelper->SourceDatas[SourceIndex];
-					Arguments.FactoryNode = ObjectInfo.FactoryNode;
-					// Should we assert if there is no factory node?
-					Arguments.NodeUniqueID = ObjectInfo.FactoryNode ? ObjectInfo.FactoryNode->GetUniqueID() : FString();
-					Arguments.NodeContainer = AsyncHelper->BaseNodeContainers[SourceIndex].Get();
-					Arguments.Pipelines = ObjectPtrDecay(AsyncHelper->Pipelines);
-					Arguments.OriginalPipelines = AsyncHelper->OriginalPipelines;
-					Arguments.bIsReimport = ObjectInfo.bIsReimport;
-					ObjectInfo.Factory->SetupObject_GameThread(Arguments);
-				}
+				UInterchangeFactoryBase::FSetupObjectParams Arguments;
+				Arguments.ImportedObject = ImportedObject;
+				Arguments.SourceData = AsyncHelper->SourceDatas[SourceIndex];
+				Arguments.FactoryNode = ObjectInfo.FactoryNode;
+				// Should we assert if there is no factory node?
+				Arguments.NodeUniqueID = ObjectInfo.FactoryNode ? ObjectInfo.FactoryNode->GetUniqueID() : FString();
+				Arguments.NodeContainer = AsyncHelper->BaseNodeContainers[SourceIndex].Get();
+				Arguments.Pipelines = ObjectPtrDecay(AsyncHelper->Pipelines);
+				Arguments.Pipelines = AsyncHelper->Pipelines;
+				Arguments.OriginalPipelines = AsyncHelper->OriginalPipelines;
+				Arguments.bIsReimport = ObjectInfo.bIsReimport;
+				ObjectInfo.Factory->SetupObject_GameThread(Arguments);
+			}
 
-				if (ImportedObject == nullptr || !IsValid(ImportedObject))
-				{
-					continue;
-				}
+			if (ImportedObject == nullptr || !IsValid(ImportedObject))
+			{
+				continue;
+			}
 
-				UInterchangeResultSuccess* Message = Results->Add<UInterchangeResultSuccess>();
-				Message->SourceAssetName = AsyncHelper->SourceDatas[SourceIndex]->GetFilename();
-				Message->DestinationAssetName = ImportedObject->GetPathName();
-				Message->AssetType = ImportedObject->GetClass();
+			UInterchangeResultSuccess* Message = Results->Add<UInterchangeResultSuccess>();
+			Message->SourceAssetName = AsyncHelper->SourceDatas[SourceIndex]->GetFilename();
+			Message->DestinationAssetName = ImportedObject->GetPathName();
+			Message->AssetType = ImportedObject->GetClass();
 
-				//Clear any async flag from the created asset and all its subobjects
-				const EInternalObjectFlags AsyncFlags = EInternalObjectFlags::Async | EInternalObjectFlags::AsyncLoading;
-				ImportedObject->ClearInternalFlags(AsyncFlags);
+			//Clear any async flag from the created asset and all its subobjects
+			const EInternalObjectFlags AsyncFlags = EInternalObjectFlags::Async | EInternalObjectFlags::AsyncLoading;
+			ImportedObject->ClearInternalFlags(AsyncFlags);
 
-				TArray<UObject*> ImportedSubobjects;
-				const bool bIncludeNestedObjects = true;
-				GetObjectsWithOuter(ImportedObject, ImportedSubobjects, bIncludeNestedObjects);
-				for (UObject* ImportedSubobject : ImportedSubobjects)
-				{
-					ImportedSubobject->ClearInternalFlags(AsyncFlags);
-				}
+			TArray<UObject*> ImportedSubobjects;
+			const bool bIncludeNestedObjects = true;
+			GetObjectsWithOuter(ImportedObject, ImportedSubobjects, bIncludeNestedObjects);
+			for (UObject* ImportedSubobject : ImportedSubobjects)
+			{
+				ImportedSubobject->ClearInternalFlags(AsyncFlags);
+			}
 
-				//Make sure the package is dirty
-				ImportedObject->MarkPackageDirty();
+			//Make sure the package is dirty
+			ImportedObject->MarkPackageDirty();
 
-				if (!bIsAsset)
+			if (!bIsAsset)
+			{
+				if (AActor* Actor = Cast<AActor>(ImportedObject))
 				{
-					if (AActor* Actor = Cast<AActor>(ImportedObject))
-					{
-						Actor->RegisterAllComponents();
-					}
-					else if (UActorComponent* Component = Cast<UActorComponent>(ImportedObject))
-					{
-						Component->RegisterComponent();
-					}
+					Actor->RegisterAllComponents();
 				}
-
-				for (UInterchangePipelineBase* PipelineBase : AsyncHelper->Pipelines)
+				else if (UActorComponent* Component = Cast<UActorComponent>(ImportedObject))
 				{
-					PipelineBase->ScriptedExecutePostFactoryPipeline(AsyncHelper->BaseNodeContainers[SourceIndex].Get()
-						, ObjectInfo.FactoryNode ? ObjectInfo.FactoryNode->GetUniqueID() : FString()
-						, ImportedObject
-						, ObjectInfo.bIsReimport);
+					Component->RegisterComponent();
 				}
 			}
 
+			for (UInterchangePipelineBase* PipelineBase : AsyncHelper->Pipelines)
+			{
+				PipelineBase->ScriptedExecutePostFactoryPipeline(AsyncHelper->BaseNodeContainers[SourceIndex].Get()
+					, ObjectInfo.FactoryNode ? ObjectInfo.FactoryNode->GetUniqueID() : FString()
+					, ImportedObject
+					, ObjectInfo.bIsReimport);
+			}
+		}
+
 #if WITH_EDITOR
-			//Second iteration to call PostEditChange
-			for (const FImportAsyncHelper::FImportedObjectInfo& ObjectInfo : ObjectInfosPerSourceIndexPair.Value)
+		//Second iteration to call PostEditChange
+		for (const FImportAsyncHelper::FImportedObjectInfo& ObjectInfo : ImportedObjects)
+		{
+			UObject* ImportedObject = ObjectInfo.ImportedObject;
+			if (ImportedObject == nullptr || !IsValid(ImportedObject))
 			{
-				UObject* ImportedObject = ObjectInfo.ImportedObject;
-				if (ImportedObject == nullptr || !IsValid(ImportedObject))
-				{
-					continue;
-				}
-				ImportedObject->PostEditChange();
+				continue;
 			}
+			ImportedObject->PostEditChange();
+		}
 #endif //WITH_EDITOR
 
-			//Third iteration to call FinalizeObject_GameThread
-			for (const FImportAsyncHelper::FImportedObjectInfo& ObjectInfo : ObjectInfosPerSourceIndexPair.Value)
+		//Third iteration to call FinalizeObject_GameThread
+		for (const FImportAsyncHelper::FImportedObjectInfo& ObjectInfo : ImportedObjects)
+		{
+			UObject* ImportedObject = ObjectInfo.ImportedObject;
+			if (ImportedObject == nullptr || !IsValid(ImportedObject))
 			{
-				UObject* ImportedObject = ObjectInfo.ImportedObject;
-				if (ImportedObject == nullptr || !IsValid(ImportedObject))
-				{
-					continue;
-				}
-				//Post import broadcast
-				if (bIsAsset)
-				{
-					AsyncHelper->AssetImportResult->AddImportedObject(ImportedObject);
+				continue;
+			}
+			//Post import broadcast
+			if (bIsAsset)
+			{
+				AsyncHelper->AssetImportResult->AddImportedObject(ImportedObject);
 
-					if (!AsyncHelper->TaskData.ReimportObject)
-					{
-						//Notify the asset registry, only when we have created the asset
-						FAssetRegistryModule::AssetCreated(ImportedObject);
-					}
-				}
-				else
+				if (!AsyncHelper->TaskData.ReimportObject)
 				{
-					AsyncHelper->SceneImportResult->AddImportedObject(ImportedObject);
+					//Notify the asset registry, only when we have created the asset
+					FAssetRegistryModule::AssetCreated(ImportedObject);
 				}
+			}
+			else
+			{
+				AsyncHelper->SceneImportResult->AddImportedObject(ImportedObject);
+			}
 
-				//In case Some factory code cannot run outside of the main thread we offer this callback to finish the work after calling post edit change (building the asset)
-				//Its possible the build of the asset to be asynchronous, the factory must handle is own asset correctly
-				if (bCallPostImportGameThreadCallback && ObjectInfo.Factory)
-				{
-					UInterchangeFactoryBase::FSetupObjectParams Arguments;
-					Arguments.ImportedObject = ImportedObject;
-					Arguments.SourceData = AsyncHelper->SourceDatas[SourceIndex];
-					Arguments.FactoryNode = ObjectInfo.FactoryNode;
-					Arguments.NodeUniqueID = ObjectInfo.FactoryNode ? ObjectInfo.FactoryNode->GetUniqueID() : FString();
-					Arguments.NodeContainer = AsyncHelper->BaseNodeContainers[SourceIndex].Get();
-					Arguments.Pipelines = ObjectPtrDecay(AsyncHelper->Pipelines);
-					Arguments.OriginalPipelines = AsyncHelper->OriginalPipelines;
-					Arguments.bIsReimport = ObjectInfo.bIsReimport;
-					ObjectInfo.Factory->FinalizeObject_GameThread(Arguments);
-				}
+			//In case Some factory code cannot run outside of the main thread we offer this callback to finish the work after calling post edit change (building the asset)
+			//Its possible the build of the asset to be asynchronous, the factory must handle is own asset correctly
+			if (bCallPostImportGameThreadCallback && ObjectInfo.Factory)
+			{
+				UInterchangeFactoryBase::FSetupObjectParams Arguments;
+				Arguments.ImportedObject = ImportedObject;
+				Arguments.SourceData = AsyncHelper->SourceDatas[SourceIndex];
+				Arguments.FactoryNode = ObjectInfo.FactoryNode;
+				Arguments.NodeUniqueID = ObjectInfo.FactoryNode ? ObjectInfo.FactoryNode->GetUniqueID() : FString();
+				Arguments.NodeContainer = AsyncHelper->BaseNodeContainers[SourceIndex].Get();
+				Arguments.Pipelines = ObjectPtrDecay(AsyncHelper->Pipelines);
+				Arguments.Pipelines = AsyncHelper->Pipelines;
+				Arguments.OriginalPipelines = AsyncHelper->OriginalPipelines;
+				Arguments.bIsReimport = ObjectInfo.bIsReimport;
+				ObjectInfo.Factory->FinalizeObject_GameThread(Arguments);
 			}
 		}
 	};
 
-	bool bIsAsset = true;
-	ForEachImportedObjectForEachSource(AsyncHelper->ImportedAssetsPerSourceIndex, bIsAsset);
+	//Asset import
+	bIsAsset = true;
+	AsyncHelper->IterateImportedAssetsPerSourceIndex(IterationCallback);
 
+	//Scene import
 	bIsAsset = false;
-	ForEachImportedObjectForEachSource(AsyncHelper->ImportedSceneObjectsPerSourceIndex, bIsAsset);
+	AsyncHelper->IterateImportedSceneObjectsPerSourceIndex(IterationCallback);
 }
 
 
@@ -211,63 +213,57 @@ void UE::Interchange::FTaskCompletion::DoTask(ENamedThreads::Type CurrentThread,
 	//No need anymore of the translators sources
 	AsyncHelper->ReleaseTranslatorsSource();
 
-	//Broadcast OnAssetPostImport/OnAssetPostReimport for each imported asset
-	for(TPair<int32, TArray<FImportAsyncHelper::FImportedObjectInfo>>& AssetInfosPerSourceIndexPair : AsyncHelper->ImportedAssetsPerSourceIndex)
+	if (!AsyncHelper->bCancel)
 	{
-		//Verify if the task was cancel
-		if (AsyncHelper->bCancel)
-		{
-			break;
-		}
-		const int32 SourceIndex = AssetInfosPerSourceIndexPair.Key;
-		for (const FImportAsyncHelper::FImportedObjectInfo& AssetInfo : AssetInfosPerSourceIndexPair.Value)
-		{
-			UObject* Asset = AssetInfo.ImportedObject;
-			if (AsyncHelper->TaskData.ReimportObject)
+		//Broadcast OnAssetPostImport/OnAssetPostReimport for each imported asset
+		AsyncHelper->IterateImportedAssetsPerSourceIndex([AsyncHelper](int32 SourceIndex, const TArray<FImportAsyncHelper::FImportedObjectInfo>& AssetInfos)
 			{
-				InterchangeManager->OnAssetPostReimport.Broadcast(Asset);
-			}
-			//We broadcast this event for both import and reimport.
-			InterchangeManager->OnAssetPostImport.Broadcast(Asset);
-		}
+				for (const FImportAsyncHelper::FImportedObjectInfo& AssetInfo : AssetInfos)
+				{
+					UObject* Asset = AssetInfo.ImportedObject;
+					if (AsyncHelper->TaskData.ReimportObject && AsyncHelper->TaskData.ReimportObject == Asset)
+					{
+						UInterchangeManager::GetInterchangeManager().OnAssetPostReimport.Broadcast(Asset);
+					}
+					//We broadcast this event for both import and reimport.
+					UInterchangeManager::GetInterchangeManager().OnAssetPostImport.Broadcast(Asset);
+				}
+			});
 	}
-
-	if (AsyncHelper->bCancel)
+	else
 	{
 		//If task is canceled, delete all created assets by this task
-		for (TPair<int32, TArray<FImportAsyncHelper::FImportedObjectInfo>>& AssetInfosPerSourceIndexPair : AsyncHelper->ImportedAssetsPerSourceIndex)
-		{
-			const int32 SourceIndex = AssetInfosPerSourceIndexPair.Key;
-			for (const FImportAsyncHelper::FImportedObjectInfo& AssetInfo : AssetInfosPerSourceIndexPair.Value)
+		AsyncHelper->IterateImportedAssetsPerSourceIndex([](int32 SourceIndex, const TArray<FImportAsyncHelper::FImportedObjectInfo>& AssetInfos)
 			{
-				UObject* Asset = AssetInfo.ImportedObject;
-				if (Asset)
+				for (const FImportAsyncHelper::FImportedObjectInfo& AssetInfo : AssetInfos)
 				{
-					//Make any created asset go away
-					Asset->ClearFlags(RF_Standalone | RF_Public | RF_Transactional);
-					Asset->ClearInternalFlags(EInternalObjectFlags::Async);
-					Asset->SetFlags(RF_Transient);
-					Asset->MarkAsGarbage();
+					UObject* Asset = AssetInfo.ImportedObject;
+					if (Asset)
+					{
+						//Make any created asset go away
+						Asset->ClearFlags(RF_Standalone | RF_Public | RF_Transactional);
+						Asset->ClearInternalFlags(EInternalObjectFlags::Async);
+						Asset->SetFlags(RF_Transient);
+						Asset->MarkAsGarbage();
+					}
 				}
-			}
-		}
+			});
 
 		//If task is canceled, remove all actors from their world
-		for (TPair<int32, TArray<FImportAsyncHelper::FImportedObjectInfo>>& SceneObjectInfosPerSourceIndexPair : AsyncHelper->ImportedSceneObjectsPerSourceIndex)
-		{
-			const int32 SourceIndex = SceneObjectInfosPerSourceIndexPair.Key;
-			for (const FImportAsyncHelper::FImportedObjectInfo& SceneObjectInfo : SceneObjectInfosPerSourceIndexPair.Value)
+		AsyncHelper->IterateImportedSceneObjectsPerSourceIndex([](int32 SourceIndex, const TArray<FImportAsyncHelper::FImportedObjectInfo>& AssetInfos)
 			{
-				if (AActor* Actor = Cast<AActor>(SceneObjectInfo.ImportedObject))
+				for (const FImportAsyncHelper::FImportedObjectInfo& SceneObjectInfo : AssetInfos)
 				{
-					if (UWorld* ActorWorld = Actor->GetWorld())
+					if (AActor* Actor = Cast<AActor>(SceneObjectInfo.ImportedObject))
 					{
-						const bool bModifyLevel = false; //This isn't undoable
-						ActorWorld->RemoveActor(Actor, bModifyLevel);
+						if (UWorld* ActorWorld = Actor->GetWorld())
+						{
+							const bool bModifyLevel = false; //This isn't undoable
+							ActorWorld->RemoveActor(Actor, bModifyLevel);
+						}
 					}
 				}
-			}
-		}
+			});
 	}
 
 	AsyncHelper->AssetImportResult->SetDone();
diff --git a/Engine/Source/Runtime/Interchange/Engine/Private/Tasks/InterchangeTaskCreateSceneObjects.cpp b/Engine/Source/Runtime/Interchange/Engine/Private/Tasks/InterchangeTaskCreateSceneObjects.cpp
--- a/Engine/Source/Runtime/Interchange/Engine/Private/Tasks/InterchangeTaskCreateSceneObjects.cpp
+++ b/Engine/Source/Runtime/Interchange/Engine/Private/Tasks/InterchangeTaskCreateSceneObjects.cpp
@@ -66,10 +66,7 @@ void UE::Interchange::FTaskCreateSceneObjects::DoTask(ENamedThreads::Type Curren
 	{
 		UInterchangeFactoryBase* Factory = NewObject<UInterchangeFactoryBase>(GetTransientPackage(), FactoryClass);
 		Factory->SetResultsContainer(AsyncHelper->AssetImportResult->GetResults());
-		{
-			FScopeLock Lock(&AsyncHelper->CreatedFactoriesLock);
-			AsyncHelper->CreatedFactories.Add(FactoryNode->GetUniqueID(), Factory);
-		}
+		AsyncHelper->AddCreatedFactory(FactoryNode->GetUniqueID(), Factory);
 
 		FString NodeDisplayName = FactoryNode->GetDisplayLabel();
 		SanitizeObjectName(NodeDisplayName);
@@ -89,16 +86,14 @@ void UE::Interchange::FTaskCreateSceneObjects::DoTask(ENamedThreads::Type Curren
 		UObject* SceneObject = Factory->ImportSceneObject_GameThread(CreateSceneObjectsParams);
 		if (SceneObject)
 		{
-			FScopeLock Lock(&AsyncHelper->ImportedSceneObjectsPerSourceIndexLock);
-			TArray<UE::Interchange::FImportAsyncHelper::FImportedObjectInfo>& ImportedInfos = AsyncHelper->ImportedSceneObjectsPerSourceIndex.FindOrAdd(SourceIndex);
-			UE::Interchange::FImportAsyncHelper::FImportedObjectInfo* ImportedInfoPtr = ImportedInfos.FindByPredicate([SceneObject](const UE::Interchange::FImportAsyncHelper::FImportedObjectInfo& CurInfo)
+			const FImportAsyncHelper::FImportedObjectInfo* ImportedInfoPtr = AsyncHelper->FindImportedSceneObjects(SourceIndex, [SceneObject](const FImportAsyncHelper::FImportedObjectInfo& CurInfo)
 				{
 					return CurInfo.ImportedObject == SceneObject;
 				});
 
 			if (!ImportedInfoPtr)
 			{
-				UE::Interchange::FImportAsyncHelper::FImportedObjectInfo& ObjectInfo = ImportedInfos.AddDefaulted_GetRef();
+				UE::Interchange::FImportAsyncHelper::FImportedObjectInfo& ObjectInfo = AsyncHelper->AddDefaultImportedSceneObjectGetRef(SourceIndex);
 				ObjectInfo.ImportedObject = SceneObject;
 				ObjectInfo.Factory = Factory;
 				ObjectInfo.FactoryNode = FactoryNode;
diff --git a/Engine/Source/Runtime/Interchange/Engine/Private/Tasks/InterchangeTaskImportObject.cpp b/Engine/Source/Runtime/Interchange/Engine/Private/Tasks/InterchangeTaskImportObject.cpp
--- a/Engine/Source/Runtime/Interchange/Engine/Private/Tasks/InterchangeTaskImportObject.cpp
+++ b/Engine/Source/Runtime/Interchange/Engine/Private/Tasks/InterchangeTaskImportObject.cpp
@@ -4,6 +4,7 @@
 #if WITH_EDITOR
 #include "AssetToolsModule.h"
 #endif //WITH_EDITOR
+#include "AssetCompilingManager.h"
 #include "Async/TaskGraphInterfaces.h"
 #include "CoreMinimal.h"
 #include "InterchangeAssetImportData.h"
@@ -179,10 +180,16 @@ namespace UE::Interchange::Private
 			return ErrorImportAssetResult;
 		}
 
-		UInterchangeFactoryBase* Factory = nullptr;
+		if (FactoryNode->ShouldSkipNodeImport())
 		{
-			FScopeLock Lock(&AsyncHelper->CreatedFactoriesLock);
-			Factory = AsyncHelper->CreatedFactories.FindChecked(FactoryNode->GetUniqueID());
+			return ErrorImportAssetResult;
+		}
+
+		UInterchangeFactoryBase* Factory = AsyncHelper->GetCreatedFactory(FactoryNode->GetUniqueID());
+		if(!ensure(Factory))
+		{
+			UE_LOG(LogInterchangeEngine, Error, TEXT("The factory to import the factory node was not created. FactoryNode: class '%s' name '%s'."), *FactoryNode->GetObjectClass()->GetName(), *FactoryNode->GetDisplayLabel());
+			return ErrorImportAssetResult;
 		}
 
 		UPackage* Pkg = nullptr;
@@ -212,17 +219,15 @@ namespace UE::Interchange::Private
 		}
 		else
 		{
-			FScopeLock Lock(&AsyncHelper->CreatedPackagesLock);
-			UPackage** PkgPtr = AsyncHelper->CreatedPackages.Find(PackageName);
+			Pkg = AsyncHelper->GetCreatedPackage(PackageName);
 
-			if (!PkgPtr || !(*PkgPtr))
+			if (!Pkg)
 			{
 				UInterchangeResultError_Generic* Message = Factory->AddMessage<UInterchangeResultError_Generic>();
 				Message->SourceAssetName = AsyncHelper->SourceDatas[SourceIndex]->GetFilename();
 				Message->DestinationAssetName = AssetName;
 				Message->AssetType = FactoryNode->GetObjectClass();
-				Message->Text = NSLOCTEXT("Interchange", "BadPackage", "It was not possible to create the asset as its package was not created correctly.");
-
+				Message->Text = NSLOCTEXT("InternalImportObjectStartup", "BadPackage", "It was not possible to create the asset as its package was not created correctly.");
 				return ErrorImportAssetResult;
 			}
 
@@ -231,12 +236,9 @@ namespace UE::Interchange::Private
 				UInterchangeResultError_Generic* Message = Factory->AddMessage<UInterchangeResultError_Generic>();
 				Message->DestinationAssetName = AssetName;
 				Message->AssetType = FactoryNode->GetObjectClass();
-				Message->Text = NSLOCTEXT("Interchange", "SourceDataOrTranslatorInvalid", "It was not possible to create the asset as its translator was not created correctly.");
-
+				Message->Text = NSLOCTEXT("InternalImportObjectStartup", "SourceDataOrTranslatorInvalid", "It was not possible to create the asset as its translator was not created correctly.");
 				return ErrorImportAssetResult;
 			}
-
-			Pkg = *PkgPtr;
 		}
 
 		if (!bSkipAsset)
@@ -281,6 +283,16 @@ void UE::Interchange::FTaskImportObject_GameThread::DoTask(ENamedThreads::Type C
 	TSharedPtr<FImportAsyncHelper, ESPMode::ThreadSafe> AsyncHelper = WeakAsyncHelper.Pin();
 	check(AsyncHelper.IsValid());
 
+	//Set it to true if everything is fine and you don't want to skip the node import
+	bool bCanContinueImport = false;
+	ON_SCOPE_EXIT
+	{
+		if (FactoryNode && !bCanContinueImport)
+		{
+			  FactoryNode->SetSkipNodeImport();
+		}
+	};
+	
 	//Verify if the task was cancel or class to import is denied
 	if (AsyncHelper->bCancel || !FactoryNode || !Private::CanImportClass(*AsyncHelper, *FactoryNode, SourceIndex))
 	{
@@ -294,10 +306,7 @@ void UE::Interchange::FTaskImportObject_GameThread::DoTask(ENamedThreads::Type C
 	UInterchangeFactoryBase* Factory = NewObject<UInterchangeFactoryBase>(GetTransientPackage(), FactoryClass);
 	Factory->SetResultsContainer(AsyncHelper->AssetImportResult->GetResults());
 
-	{
-		FScopeLock Lock(&AsyncHelper->CreatedFactoriesLock);
-		AsyncHelper->CreatedFactories.Add(FactoryNode->GetUniqueID(), Factory);
-	}
+	AsyncHelper->AddCreatedFactory(FactoryNode->GetUniqueID(), Factory);
 
 	UPackage* Pkg = nullptr;
 	FString PackageName;
@@ -305,15 +314,35 @@ void UE::Interchange::FTaskImportObject_GameThread::DoTask(ENamedThreads::Type C
 	Private::InternalGetPackageName(*AsyncHelper, SourceIndex, PackageBasePath, FactoryNode, PackageName, AssetName);
 
 	UObject* ReimportObject = FFactoryCommon::GetObjectToReimport(AsyncHelper->TaskData.ReimportObject, PackageName, AssetName);
+	if (!ensure(!IsGarbageCollecting()))
+	{
+		//Skip this asset
+		UInterchangeResultError_Generic* Message = Factory->AddMessage<UInterchangeResultError_Generic>();
+		Message->SourceAssetName = AsyncHelper->SourceDatas[SourceIndex]->GetFilename();
+		Message->DestinationAssetName = ReimportObject ? ReimportObject->GetName() : FactoryNode->GetAssetName();
+		Message->AssetType = FactoryNode->GetObjectClass();
+		Message->Text = FText::Format(NSLOCTEXT("FTaskImportObject_GameThread", "GarbageCollectIsRunning_Error", "Cannot import asset '{0}'; the garbage collect is running during the import process.")
+			, FText::FromString(AssetName));
+		return;
+	}
+
+	if (!AsyncHelper->BaseNodeContainers.IsValidIndex(SourceIndex))
+	{
+		//Skip this asset
+		UInterchangeResultError_Generic* Message = Factory->AddMessage<UInterchangeResultError_Generic>();
+		Message->SourceAssetName = AsyncHelper->SourceDatas[SourceIndex]->GetFilename();
+		Message->DestinationAssetName = ReimportObject ? ReimportObject->GetName() : FactoryNode->GetAssetName();
+		Message->AssetType = FactoryNode->GetObjectClass();
+		Message->Text = FText::Format(NSLOCTEXT("FTaskImportObject_GameThread", "InvalidBaseNodeContainer", "Cannot import asset '{0}'; the node container is invalid for this source index.")
+			, FText::FromString(AssetName));
+		return;
+	}
+	UInterchangeBaseNodeContainer* NodeContainer = AsyncHelper->BaseNodeContainers[SourceIndex].Get();
+
+	UObject* ExistingAsset = nullptr;
 	//If we do a reimport no need to create a package
 	if (ReimportObject)
 	{
-		UInterchangeBaseNodeContainer* NodeContainer = nullptr;
-		if (AsyncHelper->BaseNodeContainers.IsValidIndex(SourceIndex))
-		{
-			NodeContainer = AsyncHelper->BaseNodeContainers[SourceIndex].Get();
-		}
-		
 		if (Private::ShouldReimportFactoryNode(FactoryNode, NodeContainer, ReimportObject))
 		{
 			FactoryNode->SetDisplayLabel(ReimportObject->GetName());
@@ -321,24 +350,13 @@ void UE::Interchange::FTaskImportObject_GameThread::DoTask(ENamedThreads::Type C
 			Pkg = ReimportObject->GetPackage();
 			PackageName = Pkg->GetPathName();
 			AssetName = ReimportObject->GetName();
-			//Import Asset describe by the node
-			UInterchangeFactoryBase::FImportAssetObjectParams CreateAssetParams;
-			CreateAssetParams.AssetName = AssetName;
-			CreateAssetParams.AssetNode = FactoryNode;
-			CreateAssetParams.Parent = Pkg;
-			CreateAssetParams.SourceData = AsyncHelper->SourceDatas[SourceIndex];
-			CreateAssetParams.Translator = AsyncHelper->Translators[SourceIndex];
-			CreateAssetParams.NodeContainer = NodeContainer;
-			CreateAssetParams.ReimportObject = ReimportObject;
-			FactoryNode->SetCustomReferenceObject(FSoftObjectPath(ReimportObject));
-			//We call CreateEmptyAsset to ensure any resource use by an existing UObject is released on the game thread
-			Factory->BeginImportAsset_GameThread(CreateAssetParams);
 		}
 		else
 		{
 			//Skip this asset, reimport object original factory node does not match this factory node
 			return;
 		}
+		ExistingAsset = ReimportObject;
 	}
 	else
 	{
@@ -346,105 +364,146 @@ void UE::Interchange::FTaskImportObject_GameThread::DoTask(ENamedThreads::Type C
 		// We can not create assets that share the name of a map file in the same location
 		if (FPackageUtils::IsMapPackageAsset(PackageName))
 		{
+			//Skip this asset
 			UInterchangeResultError_Generic* Message = Factory->AddMessage<UInterchangeResultError_Generic>();
 			Message->SourceAssetName = AsyncHelper->SourceDatas[SourceIndex]->GetFilename();
 			Message->DestinationAssetName = AssetName;
 			Message->AssetType = FactoryNode->GetObjectClass();
-			Message->Text = NSLOCTEXT("Interchange", "MapExistsWithSameName", "You cannot create an asset with this name, as there is already a map file with the same name in this folder.");
-
-			//Skip this asset
+			Message->Text = NSLOCTEXT("FTaskImportObject_GameThread", "MapExistsWithSameName", "You cannot create an asset with this name, as there is already a map file with the same name in this folder.");
 			return;
 		}
 
-		//If the package already exist we must load it so factory can find any existing asset and decide or not to override it
+		//Try to find the package in memory
+		bool bPackageWasCreated = false;
+		Pkg = FindPackage(nullptr, *PackageName);
+		if (!Pkg)
 		{
-			//Try to find the package in memory
-			bool bPackageWasCreated = false;
-			Pkg = FindPackage(nullptr, *PackageName);
+			//Try to load the package from disk
+			Pkg = LoadPackage(nullptr, *PackageName, LOAD_NoWarn | LOAD_Quiet);
 			if (!Pkg)
 			{
-				//Try to load the package from disk
-				Pkg = LoadPackage(nullptr, *PackageName, LOAD_NoWarn | LOAD_Quiet);
-				if (!Pkg)
-				{
-					//Create the package
-					Pkg = CreatePackage(*PackageName);
-					Pkg->SetPackageFlags(PKG_NewlyCreated);
-					bPackageWasCreated = true;
-				}
+				//Create the package
+				Pkg = CreatePackage(*PackageName);
+				Pkg->SetPackageFlags(PKG_NewlyCreated);
+				bPackageWasCreated = true;
 			}
+		}
 
-			if (Pkg == nullptr)
-			{
-				UInterchangeResultError_Generic* Message = Factory->AddMessage<UInterchangeResultError_Generic>();
-				Message->SourceAssetName = AsyncHelper->SourceDatas[SourceIndex]->GetFilename();
-				Message->DestinationAssetName = AssetName;
-				Message->AssetType = FactoryNode->GetObjectClass();
-				Message->Text = FText::Format(NSLOCTEXT("Interchange", "CouldntCreatePackage", "It was not possible to create a package named '{0}'; the asset will not be imported."), FText::FromString(PackageName));
-
-				//Skip this asset
-				return;
-			}
+		if (Pkg == nullptr)
+		{
+			//Skip this asset
+			UInterchangeResultError_Generic* Message = Factory->AddMessage<UInterchangeResultError_Generic>();
+			Message->SourceAssetName = AsyncHelper->SourceDatas[SourceIndex]->GetFilename();
+			Message->DestinationAssetName = AssetName;
+			Message->AssetType = FactoryNode->GetObjectClass();
+			Message->Text = FText::Format(NSLOCTEXT("FTaskImportObject_GameThread", "CouldntCreatePackage", "It was not possible to create a package named '{0}'; the asset will not be imported.")
+				, FText::FromString(PackageName));
+			return;
+		}
 
-			if (!bPackageWasCreated && AsyncHelper->TaskData.bFollowRedirectors)
+		if (!bPackageWasCreated && AsyncHelper->TaskData.bFollowRedirectors)
+		{
+			if (UObjectRedirector* Redirector = FindObject<UObjectRedirector>(Pkg, *AssetName))
 			{
-				if (UObjectRedirector* Redirector = FindObject<UObjectRedirector>(Pkg, *AssetName))
+				if (Redirector->DestinationObject)
 				{
-					if (Redirector->DestinationObject)
+					Pkg = Redirector->DestinationObject->GetPackage();
+					if (FPackageName::GetLongPackageAssetName(PackageName) == AssetName)
 					{
-						Pkg = Redirector->DestinationObject->GetPackage();
-						if (FPackageName::GetLongPackageAssetName(PackageName) == AssetName)
-						{
-							AssetName = FPackageName::GetLongPackageAssetName(Pkg->GetName());
-						}
+						AssetName = FPackageName::GetLongPackageAssetName(Pkg->GetName());
 					}
 				}
 			}
 		}
+		ExistingAsset = StaticFindObject(nullptr, Pkg, *AssetName);
+	}
 
-		//Import Asset describe by the node
-		UInterchangeFactoryBase::FImportAssetObjectParams CreateAssetParams;
-		CreateAssetParams.AssetName = AssetName;
-		CreateAssetParams.AssetNode = FactoryNode;
-		CreateAssetParams.Parent = Pkg;
-		CreateAssetParams.SourceData = AsyncHelper->SourceDatas[SourceIndex];
-		CreateAssetParams.Translator = AsyncHelper->Translators[SourceIndex];
-		if (AsyncHelper->BaseNodeContainers.IsValidIndex(SourceIndex))
+	if (ExistingAsset)
+	{
+		//Set the factory node reference to the existing object, so we do not need to find the asset in the factory
+		FactoryNode->SetCustomReferenceObject(FSoftObjectPath(ExistingAsset));
+
+		if (UInterchangeManager::GetInterchangeManager().IsObjectBeingImported(ExistingAsset))
 		{
-			CreateAssetParams.NodeContainer = AsyncHelper->BaseNodeContainers[SourceIndex].Get();
+			//Skip this node, it is currently being imported by another import task
+			UInterchangeResultError_Generic* Message = Factory->AddMessage<UInterchangeResultError_Generic>();
+			Message->SourceAssetName = AsyncHelper->SourceDatas[SourceIndex]->GetFilename();
+			Message->DestinationAssetName = AssetName;
+			Message->AssetType = FactoryNode->GetObjectClass();
+			Message->Text = FText::Format(NSLOCTEXT("FTaskImportObject_GameThread", "AssetCollisionBetweenImportTask", "Multiple import task are importing the same asset at the same location [{0}]. This asset will be skip. See more information in the log")
+				, FText::FromString(AssetName));
+			return;
 		}
-		CreateAssetParams.ReimportObject = ReimportObject;
-		//Make sure the asset UObject is created with the correct type on the main thread
-		UInterchangeFactoryBase::FImportAssetResult ImportAssetResult = Factory->BeginImportAsset_GameThread(CreateAssetParams);
-		if (ImportAssetResult.ImportedObject)
+
+		if (!ExistingAsset->GetClass() || !ExistingAsset->GetClass()->IsChildOf(FactoryNode->GetObjectClass()))
 		{
-			//If the factory skip the asset, we simply set the node custom reference object
-			if (!ImportAssetResult.bIsFactorySkipAsset)
-			{
-				if (!ImportAssetResult.ImportedObject->HasAnyInternalFlags(EInternalObjectFlags::Async))
-				{
-					//Since the async flag is not set we must be in the game thread
-					ensure(IsInGameThread());
-					ImportAssetResult.ImportedObject->SetInternalFlags(EInternalObjectFlags::Async);
-				}
-				FScopeLock Lock(&AsyncHelper->ImportedAssetsPerSourceIndexLock);
-				TArray<FImportAsyncHelper::FImportedObjectInfo>& ImportedInfos = AsyncHelper->ImportedAssetsPerSourceIndex.FindOrAdd(SourceIndex);
-				FImportAsyncHelper::FImportedObjectInfo& AssetInfo = ImportedInfos.AddDefaulted_GetRef();
-				AssetInfo.ImportedObject = ImportAssetResult.ImportedObject;
-				AssetInfo.Factory = Factory;
-				AssetInfo.FactoryNode = FactoryNode;
-				AssetInfo.bIsReimport = bool(ReimportObject != nullptr);
-			}
-			FactoryNode->SetCustomReferenceObject(FSoftObjectPath(ImportAssetResult.ImportedObject));
+			//Skip this asset
+			UInterchangeResultError_Generic* Message = Factory->AddMessage<UInterchangeResultError_Generic>();
+			Message->SourceAssetName = AsyncHelper->SourceDatas[SourceIndex]->GetFilename();
+			Message->DestinationAssetName = AssetName;
+			Message->AssetType = FactoryNode->GetObjectClass();
+
+			const FText TargetClassNameText = FText::FromString(FactoryNode->GetObjectClass()->GetName());
+			const FText ExistingClassNameText = FText::FromString(ExistingAsset->GetClass()->GetName());
+			const FText AssetNameText = FText::FromString(AssetName);
+			const FText FolderNameText = FText::FromString(FPaths::GetPath(Pkg->GetPathName()));
+			Message->Text = FText::Format(NSLOCTEXT("FTaskImportObject_GameThread", "AssetVersusFactoryClassWrong_Error", "You cannot create a '{0}' asset named '{1}' in '{2}', as there is already a '{3}' asset with the same name in this folder.")
+				, TargetClassNameText
+				, AssetNameText
+				, FolderNameText
+				, ExistingClassNameText);
+			return;
+		}
+
+		if (IInterface_AsyncCompilation* AssetCompilationInterface = Cast<IInterface_AsyncCompilation>(ExistingAsset))
+		{
+			//If an asset exist make sure we wait until its compiled before calling the factory
+			FAssetCompilingManager::Get().FinishCompilationForObjects({ ExistingAsset });
 		}
 	}
 
-	// Make sure the destination package is loaded
-	Pkg->FullyLoad();
-	
+	//Import Asset describe by the node
+	UInterchangeFactoryBase::FImportAssetObjectParams CreateAssetParams;
+	CreateAssetParams.AssetName = AssetName;
+	CreateAssetParams.AssetNode = FactoryNode;
+	CreateAssetParams.Parent = Pkg;
+	CreateAssetParams.SourceData = AsyncHelper->SourceDatas[SourceIndex];
+	CreateAssetParams.Translator = AsyncHelper->Translators[SourceIndex];
+	CreateAssetParams.NodeContainer = NodeContainer;
+	CreateAssetParams.ReimportObject = ReimportObject;
+	//Make sure the asset UObject is created with the correct type on the main thread
+	UInterchangeFactoryBase::FImportAssetResult ImportAssetResult = Factory->BeginImportAsset_GameThread(CreateAssetParams);
+	if (ImportAssetResult.ImportedObject)
 	{
-		FScopeLock Lock(&AsyncHelper->CreatedPackagesLock);
-		AsyncHelper->CreatedPackages.Add(PackageName, Pkg);
+		//If the factory skip the asset, we simply set the node custom reference object
+		if (!ImportAssetResult.bIsFactorySkipAsset)
+		{
+			if (!ImportAssetResult.ImportedObject->HasAnyInternalFlags(EInternalObjectFlags::Async))
+			{
+				//Since the async flag is not set we must be in the game thread
+				ensure(IsInGameThread());
+				ImportAssetResult.ImportedObject->SetInternalFlags(EInternalObjectFlags::Async);
+			}
+			FImportAsyncHelper::FImportedObjectInfo& AssetInfo = AsyncHelper->AddDefaultImportedAssetGetRef(SourceIndex);
+
+			AssetInfo.ImportedObject = ImportAssetResult.ImportedObject;
+			AssetInfo.Factory = Factory;
+			AssetInfo.FactoryNode = FactoryNode;
+			AssetInfo.bIsReimport = bool(ReimportObject != nullptr);
+		}
+		// Make sure the destination package is loaded
+		Pkg->FullyLoad();
+		if (!ReimportObject)
+		{
+			AsyncHelper->AddCreatedPackage(PackageName, Pkg);
+		}
+
+		//this was verified before calling the factory
+		ensure(ImportAssetResult.ImportedObject->GetClass()->IsChildOf(FactoryNode->GetObjectClass()));
+
+		FactoryNode->SetCustomReferenceObject(FSoftObjectPath(ImportAssetResult.ImportedObject));
+		//Make sure the node is not skipped when the scope will be terminate
+		bCanContinueImport = true;
 	}
 }
 
@@ -464,10 +523,10 @@ void UE::Interchange::FTaskImportObject_Async::DoTask(ENamedThreads::Type Curren
 		return;
 	}
 
-	UInterchangeFactoryBase* Factory = nullptr;
+	UInterchangeFactoryBase* Factory = AsyncHelper->GetCreatedFactory(FactoryNode->GetUniqueID());
+	if(!Factory)
 	{
-		FScopeLock Lock(&AsyncHelper->CreatedFactoriesLock);
-		Factory = AsyncHelper->CreatedFactories.FindChecked(FactoryNode->GetUniqueID());
+		return;
 	}
 
 	Private::InternalImportObjectStartup(AsyncHelper
@@ -496,11 +555,12 @@ void UE::Interchange::FTaskImportObjectFinalize_GameThread::DoTask(ENamedThreads
 		return;
 	}
 
-	UInterchangeFactoryBase* Factory = nullptr;
+	UInterchangeFactoryBase* Factory = AsyncHelper->GetCreatedFactory(FactoryNode->GetUniqueID());
+	if(!Factory)
 	{
-		FScopeLock Lock(&AsyncHelper->CreatedFactoriesLock);
-		Factory = AsyncHelper->CreatedFactories.FindChecked(FactoryNode->GetUniqueID());
+		return;
 	}
+
 	UInterchangeFactoryBase::FImportAssetResult ImportAssetResult = Private::InternalImportObjectStartup(AsyncHelper
 		, FactoryNode
 		, SourceIndex
@@ -515,16 +575,14 @@ void UE::Interchange::FTaskImportObjectFinalize_GameThread::DoTask(ENamedThreads
 		//If the factory skip the asset, we simply set the node custom reference object
 		if (!ImportAssetResult.bIsFactorySkipAsset)
 		{
-			FScopeLock Lock(&AsyncHelper->ImportedAssetsPerSourceIndexLock);
-			TArray<FImportAsyncHelper::FImportedObjectInfo>& ImportedInfos = AsyncHelper->ImportedAssetsPerSourceIndex.FindOrAdd(SourceIndex);
-			FImportAsyncHelper::FImportedObjectInfo* AssetInfoPtr = ImportedInfos.FindByPredicate([&ImportAssetResult](const FImportAsyncHelper::FImportedObjectInfo& CurInfo)
+			const FImportAsyncHelper::FImportedObjectInfo* AssetInfoPtr = AsyncHelper->FindImportedAssets(SourceIndex, [&ImportAssetResult](const FImportAsyncHelper::FImportedObjectInfo& CurInfo)
 				{
 					return CurInfo.ImportedObject == ImportAssetResult.ImportedObject;
 				});
 
 			if (!AssetInfoPtr)
 			{
-				FImportAsyncHelper::FImportedObjectInfo& AssetInfo = ImportedInfos.AddDefaulted_GetRef();
+				FImportAsyncHelper::FImportedObjectInfo& AssetInfo = AsyncHelper->AddDefaultImportedAssetGetRef(SourceIndex);
 				AssetInfo.ImportedObject = ImportAssetResult.ImportedObject;
 				AssetInfo.Factory = Factory;
 				AssetInfo.FactoryNode = FactoryNode;
@@ -549,7 +607,6 @@ void UE::Interchange::FTaskImportObjectFinalize_GameThread::DoTask(ENamedThreads
 				}
 			}
 		}
-
 		FactoryNode->SetCustomReferenceObject(FSoftObjectPath(ImportAssetResult.ImportedObject));
 	}
 }
\ No newline at end of file
diff --git a/Engine/Source/Runtime/Interchange/Engine/Private/Tasks/InterchangeTaskPipeline.cpp b/Engine/Source/Runtime/Interchange/Engine/Private/Tasks/InterchangeTaskPipeline.cpp
--- a/Engine/Source/Runtime/Interchange/Engine/Private/Tasks/InterchangeTaskPipeline.cpp
+++ b/Engine/Source/Runtime/Interchange/Engine/Private/Tasks/InterchangeTaskPipeline.cpp
@@ -80,26 +80,21 @@ void UE::Interchange::FTaskPipelinePostImport::DoTask(ENamedThreads::Type Curren
 	TArray<bool> IsAssetsReimported;
 
 	auto FillImportedObjectsFromSource =
-		[&NodeUniqueIDs, &ImportedObjects, &IsAssetsReimported, this](FCriticalSection& CriticalSection, const TMap<int32, TArray<UE::Interchange::FImportAsyncHelper::FImportedObjectInfo>>& ImportedInfosPerSource)
+		[&NodeUniqueIDs, &ImportedObjects, &IsAssetsReimported, this](const TArray<UE::Interchange::FImportAsyncHelper::FImportedObjectInfo>& ImportedInfos)
 		{
-			FScopeLock Lock(&CriticalSection);
-			if (ImportedInfosPerSource.Contains(SourceIndex))
+			NodeUniqueIDs.Reserve(NodeUniqueIDs.Num() + ImportedInfos.Num());
+			ImportedObjects.Reserve(ImportedObjects.Num() + ImportedInfos.Num());
+			IsAssetsReimported.Reserve(IsAssetsReimported.Num() + ImportedInfos.Num());
+			for (const UE::Interchange::FImportAsyncHelper::FImportedObjectInfo& ImportedInfo : ImportedInfos)
 			{
-				const TArray<UE::Interchange::FImportAsyncHelper::FImportedObjectInfo>& ImportedInfos = ImportedInfosPerSource.FindChecked(SourceIndex);
-				NodeUniqueIDs.Reserve(ImportedInfos.Num());
-				ImportedObjects.Reserve(ImportedInfos.Num());
-				IsAssetsReimported.Reserve(ImportedInfos.Num());
-				for (const UE::Interchange::FImportAsyncHelper::FImportedObjectInfo& ImportedInfo : ImportedInfos)
-				{
-					NodeUniqueIDs.Add(ImportedInfo.FactoryNode->GetUniqueID());
-					ImportedObjects.Add(ImportedInfo.ImportedObject);
-					IsAssetsReimported.Add(ImportedInfo.bIsReimport);
-				}
+				NodeUniqueIDs.Add(ImportedInfo.FactoryNode->GetUniqueID());
+				ImportedObjects.Add(ImportedInfo.ImportedObject);
+				IsAssetsReimported.Add(ImportedInfo.bIsReimport);
 			}
 		};
 
-	FillImportedObjectsFromSource(AsyncHelper->ImportedAssetsPerSourceIndexLock, AsyncHelper->ImportedAssetsPerSourceIndex);
-	FillImportedObjectsFromSource(AsyncHelper->ImportedSceneObjectsPerSourceIndexLock, AsyncHelper->ImportedSceneObjectsPerSourceIndex);
+	AsyncHelper->IterateImportedAssets(SourceIndex, FillImportedObjectsFromSource);
+	AsyncHelper->IterateImportedSceneObjects(SourceIndex, FillImportedObjectsFromSource);
 
 	if (!ensure(NodeUniqueIDs.Num() == ImportedObjects.Num()))
 	{
diff --git a/Engine/Source/Runtime/Interchange/Engine/Public/InterchangeManager.h b/Engine/Source/Runtime/Interchange/Engine/Public/InterchangeManager.h
--- a/Engine/Source/Runtime/Interchange/Engine/Public/InterchangeManager.h
+++ b/Engine/Source/Runtime/Interchange/Engine/Public/InterchangeManager.h
@@ -9,6 +9,7 @@
 #include "Async/TaskGraphInterfaces.h"
 #include "Containers/Queue.h"
 #include "Delegates/DelegateCombinations.h"
+#include "HAL/CriticalSection.h"
 #include "HAL/Thread.h"
 #include "HAL/ThreadSafeBool.h"
 #include "InterchangePipelineConfigurationBase.h"
@@ -214,13 +215,11 @@ namespace UE
 			FGraphEventRef PreCompletionTask;
 			FGraphEventRef CompletionTask;
 
-			//Create package map, Key is package name. We cannot create package asynchronously so we have to create a game thread task to do this
-			FCriticalSection CreatedPackagesLock;
-			TMap<FString, UPackage*> CreatedPackages;
+			UPackage* GetCreatedPackage(const FString& PackageName) const;
+			void AddCreatedPackage(const FString& PackageName, UPackage* Package);
 
-			// Created factories map, Key is factory node UID
-			FCriticalSection CreatedFactoriesLock;
-			TMap<FString, TObjectPtr<UInterchangeFactoryBase>> CreatedFactories;
+			UInterchangeFactoryBase* GetCreatedFactory(const FString& FactoryNodeUniqueId) const;
+			void AddCreatedFactory(const FString& FactoryNodeUniqueId, UInterchangeFactoryBase* Factory);
 
 			// Set of classes which creation has been denied
 			TSet<UClass*> DeniedClasses;
@@ -236,11 +235,20 @@ namespace UE
 				bool bIsReimport;
 			};
 
-			FCriticalSection ImportedAssetsPerSourceIndexLock;
-			TMap<int32, TArray<FImportedObjectInfo>> ImportedAssetsPerSourceIndex;
+			FImportedObjectInfo& AddDefaultImportedAssetGetRef(int32 SourceIndex);
+			const FImportedObjectInfo* FindImportedAssets(int32 SourceIndex, TFunction< bool(const FImportedObjectInfo& ImportedObjects) > Predicate) const;
+			void IterateImportedAssets(int32 SourceIndex, TFunction< void(const TArray<FImportedObjectInfo>& ImportedObjects) > Callback) const;
+			void IterateImportedAssetsPerSourceIndex(TFunction< void(int32 SourceIndex, const TArray<FImportedObjectInfo>& ImportedObjects) > Callback) const;
 
-			FCriticalSection ImportedSceneObjectsPerSourceIndexLock;
-			TMap<int32, TArray<FImportedObjectInfo>> ImportedSceneObjectsPerSourceIndex;
+			FImportedObjectInfo& AddDefaultImportedSceneObjectGetRef(int32 SourceIndex);
+			const FImportedObjectInfo* FindImportedSceneObjects(int32 SourceIndex, TFunction< bool(const FImportedObjectInfo& ImportedObjects) > Predicate) const;
+			void IterateImportedSceneObjects(int32 SourceIndex, TFunction< void(const TArray<FImportedObjectInfo>& ImportedObjects) > Callback) const;
+			void IterateImportedSceneObjectsPerSourceIndex(TFunction< void(int32 SourceIndex, const TArray<FImportedObjectInfo>& ImportedObjects) > Callback) const;
+
+			/*
+			 * Return true if the Object is imported by this async import, false otherwise
+			 */
+			bool IsImportingObject(UObject* Object) const;
 
 			FImportAsyncHelperData TaskData;
 
@@ -261,6 +269,21 @@ namespace UE
 			void InitCancel();
 
 			void CleanUp();
+
+		private:
+			//Create package map, Key is package name. We cannot create package asynchronously so we have to create a game thread task to do this
+			mutable FCriticalSection CreatedPackagesLock;
+			TMap<FString, UPackage*> CreatedPackages;
+
+			// Created factories map, Key is factory node UID
+			mutable FCriticalSection CreatedFactoriesLock;
+			TMap<FString, TObjectPtr<UInterchangeFactoryBase>> CreatedFactories;
+
+			mutable FCriticalSection ImportedAssetsPerSourceIndexLock;
+			TMap<int32, TArray<FImportedObjectInfo>> ImportedAssetsPerSourceIndex;
+
+			mutable FCriticalSection ImportedSceneObjectsPerSourceIndexLock;
+			TMap<int32, TArray<FImportedObjectInfo>> ImportedSceneObjectsPerSourceIndex;
 		};
 
 		void SanitizeObjectPath(FString& ObjectPath);
@@ -537,6 +560,12 @@ public:
 	 */
 	UInterchangeTranslatorBase* GetTranslatorSupportingPayloadInterfaceForSourceData(const UInterchangeSourceData* SourceData, const UClass* PayloadInterfaceClass) const;
 
+	/**
+	 * Return true if the object is being imported, return false otherwise. If the user import multiple file in the same folder its possible to
+	 * have the same asset name in two different files.
+	 */
+	bool IsObjectBeingImported(UObject* Object) const;
+
 protected:
 
 	/** Return true if we can show some UI */
@@ -629,5 +658,9 @@ private:
 	//This boolean is reset to false when the ImportTasks array is empty.
 	bool bImportCanceled = false;
 
+	//We want to avoid starting an import task during a GC
+	FDelegateHandle GCEndDelegate;
+	bool bGCEndDelegateCancellAllTask = false;
+
 	friend class UE::Interchange::FScopedTranslator;
 };
