Make sure the interchange notification go away when we have a failure. This will avoid pilling many notifications when doing a batch import.

#rnx

diff --git a/Engine/Source/Runtime/Interchange/Engine/Private/InterchangeManager.cpp b/Engine/Source/Runtime/Interchange/Engine/Private/InterchangeManager.cpp
--- a/Engine/Source/Runtime/Interchange/Engine/Private/InterchangeManager.cpp
+++ b/Engine/Source/Runtime/Interchange/Engine/Private/InterchangeManager.cpp
@@ -1102,11 +1102,11 @@ void UInterchangeManager::StartQueuedTasks(bool bCancelAllTasks /*= false*/)
 				FText TitleText = NSLOCTEXT("Interchange", "Asynchronous_import_start", "Importing");
 				FAsyncTaskNotificationConfig NotificationConfig;
 				NotificationConfig.bIsHeadless = false;
-				NotificationConfig.bKeepOpenOnFailure = true;
 				NotificationConfig.TitleText = TitleText;
 				NotificationConfig.LogCategory = UE::Interchange::Private::GetLogInterchangePtr();
 				NotificationConfig.bCanCancel.Set(true);
-				NotificationConfig.bKeepOpenOnFailure.Set(true);
+				NotificationConfig.bKeepOpenOnSuccess.Set(false);
+				NotificationConfig.bKeepOpenOnFailure.Set(false);
 
 				Notification = MakeShared<FAsyncTaskNotification>(NotificationConfig);
 				Notification->SetNotificationState(FAsyncNotificationStateData(TitleText, FText::GetEmpty(), EAsyncTaskNotificationState::Pending));
