Fix the source data display name
#rnx

diff --git a/Engine/Source/Runtime/Interchange/Core/Public/InterchangeSourceData.h b/Engine/Source/Runtime/Interchange/Core/Public/InterchangeSourceData.h
--- a/Engine/Source/Runtime/Interchange/Core/Public/InterchangeSourceData.h
+++ b/Engine/Source/Runtime/Interchange/Core/Public/InterchangeSourceData.h
@@ -68,13 +68,13 @@ public:
 			FString CleanFilename = FPaths::GetCleanFilename(Filename);
 			FString BasePath = FPaths::GetPath(Filename);
 			//limit the base path to 20 character do ellipsis in the middle
-			if (BasePath.Len() > 33)
+			if (BasePath.Len() > 43)
 			{
-				FString RightPath = BasePath.Right(15);
-				FString LeftPath = BasePath.Left(15);
+				FString RightPath = BasePath.Right(20);
+				FString LeftPath = BasePath.Left(20);
 				BasePath = LeftPath + TEXT("...") + RightPath;
 			}
-			FString DisplayString = BasePath + CleanFilename;
+			FString DisplayString = BasePath + TEXT("/") + CleanFilename;
 			return DisplayString;
 		}
 		return FString();
