Notify the user if there is no valid asset to import in a source data.

#rnx

diff --git a/Engine/Source/Runtime/Interchange/Engine/Private/Tasks/InterchangeTaskParsing.cpp b/Engine/Source/Runtime/Interchange/Engine/Private/Tasks/InterchangeTaskParsing.cpp
--- a/Engine/Source/Runtime/Interchange/Engine/Private/Tasks/InterchangeTaskParsing.cpp
+++ b/Engine/Source/Runtime/Interchange/Engine/Private/Tasks/InterchangeTaskParsing.cpp
@@ -322,6 +322,18 @@ void UE::Interchange::FTaskParsing::DoTask(ENamedThreads::Type CurrentThread, co
 		}
 	};
 
+	//If we find nothing to import, we notify the user
+	if (TaskDatas.Num() == 0)
+	{
+		for(TObjectPtr<UInterchangeTranslatorBase> Translator : AsyncHelper->Translators)
+		{
+			UInterchangeResultError_Generic* ErrorResult = NewObject<UInterchangeResultError_Generic>(GetTransientPackage(), UInterchangeResultError_Generic::StaticClass());
+			ErrorResult->Text = NSLOCTEXT("InterchangeTaskParsing", "TranslatorNoAssetImported", "There was no asset to import in the provided source data (found no mesh and no animation).");
+			ErrorResult->SourceAssetName = Translator->GetSourceData()->GetFilename();
+			Translator->AddMessage(ErrorResult);
+		}
+	}
+
 	FGraphEventArray CompletionPrerequistes;
 	const int32 PoolWorkerThreadCount = FTaskGraphInterface::Get().GetNumWorkerThreads() / 2;
 	const int32 MaxNumWorker = FMath::Max(PoolWorkerThreadCount, 1);
