Interchange: Make sure we can set blueprint and python pipeline in the override pipeline array of the import parameter.

#rnx

diff --git a/Engine/Plugins/Tests/InterchangeTests/Source/InterchangeTests/Private/InterchangeImportTestStepImport.cpp b/Engine/Plugins/Tests/InterchangeTests/Source/InterchangeTests/Private/InterchangeImportTestStepImport.cpp
--- a/Engine/Plugins/Tests/InterchangeTests/Source/InterchangeTests/Private/InterchangeImportTestStepImport.cpp
+++ b/Engine/Plugins/Tests/InterchangeTests/Source/InterchangeTests/Private/InterchangeImportTestStepImport.cpp
@@ -31,7 +31,11 @@ TTuple<UE::Interchange::FAssetImportResultPtr, UE::Interchange::FSceneImportResu
 	UE::Interchange::FScopedSourceData ScopedSourceData(SourceFile.FilePath);
 
 	FImportAssetParameters Params;
-	Params.OverridePipelines = PipelineStack;
+	Params.OverridePipelines.Reserve(PipelineStack.Num());
+	for(TObjectPtr<UInterchangePipelineBase> Pipeline : PipelineStack)
+	{
+		Params.OverridePipelines.Add(Pipeline);
+	}
 	Params.bIsAutomated = true;
 
 	UInterchangeManager& InterchangeManager = UInterchangeManager::GetInterchangeManager();
diff --git a/Engine/Source/Runtime/Interchange/Engine/Private/InterchangeManager.cpp b/Engine/Source/Runtime/Interchange/Engine/Private/InterchangeManager.cpp
--- a/Engine/Source/Runtime/Interchange/Engine/Private/InterchangeManager.cpp
+++ b/Engine/Source/Runtime/Interchange/Engine/Private/InterchangeManager.cpp
@@ -766,8 +766,11 @@ UInterchangePipelineBase* UE::Interchange::GeneratePipelineInstance(const FSoftO
 		UE_LOG(LogInterchangeEngine, Error, TEXT("Cannot generate a pipeline instance because the pipeline asset %s type is unknown."), *PipelineInstance.GetWithoutSubPath().ToString());
 	}
 
-	// Make sure that the instance does not carry over standalone and public flags as they are not actual assets to be persisted
-	GeneratedPipeline->ClearFlags(EObjectFlags::RF_Standalone|EObjectFlags::RF_Public);
+	if(GeneratedPipeline)
+	{
+		// Make sure that the instance does not carry over standalone and public flags as they are not actual assets to be persisted
+		GeneratedPipeline->ClearFlags(EObjectFlags::RF_Standalone|EObjectFlags::RF_Public);
+	}
 
 	return GeneratedPipeline;
 }
@@ -1576,8 +1579,8 @@ UInterchangeManager::ImportInternal(const FString& ContentPath, const UInterchan
 	{
 		for (int32 GraphPipelineIndex = 0; GraphPipelineIndex < ImportAssetParameters.OverridePipelines.Num(); ++GraphPipelineIndex)
 		{
-			UInterchangePipelineBase* OverridePipeline = ImportAssetParameters.OverridePipelines[GraphPipelineIndex];
-			if (!OverridePipeline)
+			UInterchangePipelineBase* GeneratedPipeline = UE::Interchange::GeneratePipelineInstance(ImportAssetParameters.OverridePipelines[GraphPipelineIndex]);
+			if (!GeneratedPipeline)
 			{
 				UE_LOG(LogInterchangeEngine, Error, TEXT("Interchange import: Override pipeline array contains a NULL pipeline. Script or code need to be fix to avoid this. "));
 				continue;
@@ -1585,7 +1588,6 @@ UInterchangeManager::ImportInternal(const FString& ContentPath, const UInterchan
 			else
 			{
 				// Duplicate the override pipelines to protect the scripted users form making race conditions
-				UInterchangePipelineBase* GeneratedPipeline = DuplicateObject<UInterchangePipelineBase>(OverridePipeline, GetTransientPackage());
 				AdjustPipelineSettingForContext(GeneratedPipeline);
 				AsyncHelper->Pipelines.Add(GeneratedPipeline);
 				AsyncHelper->OriginalPipelines.Add(GeneratedPipeline);
diff --git a/Engine/Source/Runtime/Interchange/Engine/Public/InterchangeManager.h b/Engine/Source/Runtime/Interchange/Engine/Public/InterchangeManager.h
--- a/Engine/Source/Runtime/Interchange/Engine/Public/InterchangeManager.h
+++ b/Engine/Source/Runtime/Interchange/Engine/Public/InterchangeManager.h
@@ -331,8 +331,8 @@ struct INTERCHANGEENGINE_API FImportAssetParameters
 	bool bFollowRedirectors = false;
 
 	// Adding some override will tell interchange to use the specific custom set pipelines instead of letting the user or the system chose
-	UPROPERTY(EditAnywhere, BlueprintReadWrite, Category = "Interchange|ImportAsset")
-	TArray<TObjectPtr<UInterchangePipelineBase>> OverridePipelines;
+	UPROPERTY(EditAnywhere, BlueprintReadWrite, Category = "Interchange|ImportAsset", meta = (AllowedClasses = "/Script/InterchangeCore.InterchangePipelineBase, /Script/InterchangeEngine.InterchangeBlueprintPipelineBase, /Script/InterchangeEngine.InterchangePythonPipelineAsset"))
+	TArray<FSoftObjectPath> OverridePipelines;
 
 	/* Delegates used track the imported objects */
 
