Make sure existing asset we want to import over are valid and output a message if they are invalid.

#rnx

diff --git a/Engine/Source/Runtime/Interchange/Engine/Private/Tasks/InterchangeTaskImportObject.cpp b/Engine/Source/Runtime/Interchange/Engine/Private/Tasks/InterchangeTaskImportObject.cpp
--- a/Engine/Source/Runtime/Interchange/Engine/Private/Tasks/InterchangeTaskImportObject.cpp
+++ b/Engine/Source/Runtime/Interchange/Engine/Private/Tasks/InterchangeTaskImportObject.cpp
@@ -420,21 +420,6 @@ void UE::Interchange::FTaskImportObject_GameThread::DoTask(ENamedThreads::Type C
 
 	if (ExistingAsset)
 	{
-		//Set the factory node reference to the existing object, so we do not need to find the asset in the factory
-		FactoryNode->SetCustomReferenceObject(FSoftObjectPath(ExistingAsset));
-
-		if (UInterchangeManager::GetInterchangeManager().IsObjectBeingImported(ExistingAsset))
-		{
-			//Skip this node, it is currently being imported by another import task
-			UInterchangeResultError_Generic* Message = Factory->AddMessage<UInterchangeResultError_Generic>();
-			Message->SourceAssetName = AsyncHelper->SourceDatas[SourceIndex]->GetFilename();
-			Message->DestinationAssetName = AssetName;
-			Message->AssetType = FactoryNode->GetObjectClass();
-			Message->Text = FText::Format(NSLOCTEXT("FTaskImportObject_GameThread", "AssetCollisionBetweenImportTask", "Multiple import task are importing the same asset at the same location [{0}]. This asset will be skip. See more information in the log")
-				, FText::FromString(AssetName));
-			return;
-		}
-
 		if (!ExistingAsset->GetClass() || !ExistingAsset->GetClass()->IsChildOf(FactoryNode->GetObjectClass()))
 		{
 			//Skip this asset
@@ -455,6 +440,37 @@ void UE::Interchange::FTaskImportObject_GameThread::DoTask(ENamedThreads::Type C
 			return;
 		}
 
+		if (!IsValid(ExistingAsset))
+		{
+			ExistingAsset->ClearGarbage();
+			if (!IsValid(ExistingAsset))
+			{
+				//Skip this node, the existing asset is not a valid asset and we cannot import over it
+				UInterchangeResultError_Generic* Message = Factory->AddMessage<UInterchangeResultError_Generic>();
+				Message->SourceAssetName = AsyncHelper->SourceDatas[SourceIndex]->GetFilename();
+				Message->DestinationAssetName = AssetName;
+				Message->AssetType = FactoryNode->GetObjectClass();
+				Message->Text = FText::Format(NSLOCTEXT("FTaskImportObject_GameThread", "InvalidExistingAsset", "An invalid asset exist at the same asset location [{0}]. This asset will be skip.")
+					, FText::FromString(AssetName));
+				return;
+			}
+		}
+
+		//Set the factory node reference to the existing object, so we do not need to find the asset in the factory
+		FactoryNode->SetCustomReferenceObject(FSoftObjectPath(ExistingAsset));
+
+		if (UInterchangeManager::GetInterchangeManager().IsObjectBeingImported(ExistingAsset))
+		{
+			//Skip this node, it is currently being imported by another import task
+			UInterchangeResultError_Generic* Message = Factory->AddMessage<UInterchangeResultError_Generic>();
+			Message->SourceAssetName = AsyncHelper->SourceDatas[SourceIndex]->GetFilename();
+			Message->DestinationAssetName = AssetName;
+			Message->AssetType = FactoryNode->GetObjectClass();
+			Message->Text = FText::Format(NSLOCTEXT("FTaskImportObject_GameThread", "AssetCollisionBetweenImportTask", "Multiple import task are importing the same asset at the same location [{0}]. This asset will be skip. See more information in the log")
+				, FText::FromString(AssetName));
+			return;
+		}
+
 		if (IInterface_AsyncCompilation* AssetCompilationInterface = Cast<IInterface_AsyncCompilation>(ExistingAsset))
 		{
 			//If an asset exist make sure we wait until its compiled before calling the factory
