Fixed crash on import high res mesh description with Interchange

diff --git a/Engine/Source/Runtime/Interchange/Engine/Private/InterchangeMeshUtilities.cpp b/Engine/Source/Runtime/Interchange/Engine/Private/InterchangeMeshUtilities.cpp
--- a/Engine/Source/Runtime/Interchange/Engine/Private/InterchangeMeshUtilities.cpp
+++ b/Engine/Source/Runtime/Interchange/Engine/Private/InterchangeMeshUtilities.cpp
@@ -97,7 +97,7 @@ TFuture<bool> UInterchangeMeshUtilities::InternalImportCustomLodAsync(TSharedPtr
 	if (SkeletalMesh)
 	{
 		InterchangeAssetImportData = Cast<UInterchangeAssetImportData>(SkeletalMesh->GetAssetImportData());
-		if (SkeletalMesh->GetLODNum() > LodIndex)
+		if (SkeletalMesh->GetLODNum() > LodIndex && InterchangeAssetImportData)
 		{
 			ImportType = EInterchangePipelineContext::AssetCustomLODReimport;
 		}
@@ -105,7 +105,7 @@ TFuture<bool> UInterchangeMeshUtilities::InternalImportCustomLodAsync(TSharedPtr
 	else if (StaticMesh)
 	{
 		InterchangeAssetImportData = Cast<UInterchangeAssetImportData>(StaticMesh->GetAssetImportData());
-		if (StaticMesh->GetNumSourceModels() > LodIndex)
+		if (StaticMesh->GetNumSourceModels() > LodIndex && InterchangeAssetImportData)
 		{
 			ImportType = EInterchangePipelineContext::AssetCustomLODReimport;
 		}
@@ -119,23 +119,27 @@ TFuture<bool> UInterchangeMeshUtilities::InternalImportCustomLodAsync(TSharedPtr
 
 	FImportAssetParameters ImportAssetParameters;
 	ImportAssetParameters.bIsAutomated = true;
-	for (TObjectPtr<UObject> SelectedPipeline : InterchangeAssetImportData->Pipelines)
+	if (InterchangeAssetImportData)
 	{
-		UInterchangePipelineBase* GeneratedPipeline = nullptr;
-		if (UInterchangePythonPipelineAsset* PythonPipelineAsset = Cast<UInterchangePythonPipelineAsset>(SelectedPipeline))
+		for (TObjectPtr<UObject> SelectedPipeline : InterchangeAssetImportData->Pipelines)
 		{
-			GeneratedPipeline = Cast<UInterchangePipelineBase>(StaticDuplicateObject(PythonPipelineAsset->GeneratedPipeline, GetTransientPackage()));
-		}
-		else
-		{
-			GeneratedPipeline = Cast<UInterchangePipelineBase>(StaticDuplicateObject(SelectedPipeline, GetTransientPackage()));
-		}
-		if (ensure(GeneratedPipeline))
-		{
-			GeneratedPipeline->AdjustSettingsForContext(ImportType, nullptr);
-			ImportAssetParameters.OverridePipelines.Add(GeneratedPipeline);
+			UInterchangePipelineBase* GeneratedPipeline = nullptr;
+			if (UInterchangePythonPipelineAsset* PythonPipelineAsset = Cast<UInterchangePythonPipelineAsset>(SelectedPipeline))
+			{
+				GeneratedPipeline = Cast<UInterchangePipelineBase>(StaticDuplicateObject(PythonPipelineAsset->GeneratedPipeline, GetTransientPackage()));
+			}
+			else
+			{
+				GeneratedPipeline = Cast<UInterchangePipelineBase>(StaticDuplicateObject(SelectedPipeline, GetTransientPackage()));
+			}
+			if (ensure(GeneratedPipeline))
+			{
+				GeneratedPipeline->AdjustSettingsForContext(ImportType, nullptr);
+				ImportAssetParameters.OverridePipelines.Add(GeneratedPipeline);
+			}
 		}
 	}
+
 	FString ImportAssetPath = TEXT("/Engine/TempEditor/Interchange/") + FGuid::NewGuid().ToString(EGuidFormats::Base36Encoded);
 	UE::Interchange::FAssetImportResultRef AssetImportResult = InterchangeManager.ImportAssetAsync(ImportAssetPath, SourceData, ImportAssetParameters);
 	FString SourceDataFilename = SourceData->GetFilename();
