[Engine]

diff --git a/Engine/Source/Runtime/InteractiveToolsFramework/Public/InteractiveToolsContext.h b/Engine/Source/Runtime/InteractiveToolsFramework/Public/InteractiveToolsContext.h
--- a/Engine/Source/Runtime/InteractiveToolsFramework/Public/InteractiveToolsContext.h
+++ b/Engine/Source/Runtime/InteractiveToolsFramework/Public/InteractiveToolsContext.h
@@ -34,13 +34,13 @@ class UToolTargetManager;
  * desirable to subclass and extend the ToolsContext. For example, UEdModeInteractiveToolsContext allows
  * a ToolsContext to live within a UEdMode.
  */
-UCLASS(Transient)
-class INTERACTIVETOOLSFRAMEWORK_API UInteractiveToolsContext : public UObject
+UCLASS(Transient, MinimalAPI)
+class UInteractiveToolsContext : public UObject
 {
 	GENERATED_BODY()
 	
 public:
-	UInteractiveToolsContext();
+	INTERACTIVETOOLSFRAMEWORK_API UInteractiveToolsContext();
 
 	// FContextInitInfo is used by Initialize() to pass information to the various creator functions below.
 	// See ::UInteractiveToolsContext() and ::Initialize() for how these fields are used/initialized.
@@ -58,46 +58,46 @@ public:
 	// By doing this, clients of the ITF can provide their own subclass implementations, or do other
 	// custom setup/teardown as necessary. See comments above CreateInputRouterFunc below
 	//
-	virtual void SetCreateInputRouterFunc(TUniqueFunction<UInputRouter* (const FContextInitInfo&)> Func);
-	virtual void SetCreateToolManagerFunc(TUniqueFunction<UInteractiveToolManager* (const FContextInitInfo&)> Func);
-	virtual void SetCreateToolTargetManagerFunc(TUniqueFunction<UToolTargetManager* (const FContextInitInfo&)> Func);
-	virtual void SetCreateGizmoManagerFunc(TUniqueFunction<UInteractiveGizmoManager* (const FContextInitInfo&)> Func);
-	virtual void SetCreateContextStoreFunc(TUniqueFunction<UContextObjectStore* (const FContextInitInfo&)> Func);
-
-	virtual void SetShutdownInputRouterFunc(TUniqueFunction<void(UInputRouter*)> Func);
-	virtual void SetShutdownToolManagerFunc(TUniqueFunction<void(UInteractiveToolManager*)> Func);
-	virtual void SetShutdownToolTargetManagerFunc(TUniqueFunction<void(UToolTargetManager*)> Func);
-	virtual void SetShutdownGizmoManagerFunc(TUniqueFunction<void(UInteractiveGizmoManager*)> Func);
-	virtual void SetShutdownContextStoreFunc(TUniqueFunction<void(UContextObjectStore*)> Func);
+	INTERACTIVETOOLSFRAMEWORK_API virtual void SetCreateInputRouterFunc(TUniqueFunction<UInputRouter* (const FContextInitInfo&)> Func);
+	INTERACTIVETOOLSFRAMEWORK_API virtual void SetCreateToolManagerFunc(TUniqueFunction<UInteractiveToolManager* (const FContextInitInfo&)> Func);
+	INTERACTIVETOOLSFRAMEWORK_API virtual void SetCreateToolTargetManagerFunc(TUniqueFunction<UToolTargetManager* (const FContextInitInfo&)> Func);
+	INTERACTIVETOOLSFRAMEWORK_API virtual void SetCreateGizmoManagerFunc(TUniqueFunction<UInteractiveGizmoManager* (const FContextInitInfo&)> Func);
+	INTERACTIVETOOLSFRAMEWORK_API virtual void SetCreateContextStoreFunc(TUniqueFunction<UContextObjectStore* (const FContextInitInfo&)> Func);
+
+	INTERACTIVETOOLSFRAMEWORK_API virtual void SetShutdownInputRouterFunc(TUniqueFunction<void(UInputRouter*)> Func);
+	INTERACTIVETOOLSFRAMEWORK_API virtual void SetShutdownToolManagerFunc(TUniqueFunction<void(UInteractiveToolManager*)> Func);
+	INTERACTIVETOOLSFRAMEWORK_API virtual void SetShutdownToolTargetManagerFunc(TUniqueFunction<void(UToolTargetManager*)> Func);
+	INTERACTIVETOOLSFRAMEWORK_API virtual void SetShutdownGizmoManagerFunc(TUniqueFunction<void(UInteractiveGizmoManager*)> Func);
+	INTERACTIVETOOLSFRAMEWORK_API virtual void SetShutdownContextStoreFunc(TUniqueFunction<void(UContextObjectStore*)> Func);
 
 	/** 
 	 * Initialize the Context. This creates the InputRouter, ToolManager, GizmoManager, TargetManager, and ContextStore
 	 * @param QueriesAPI client-provided implementation of the API for querying the higher-evel scene state
 	 * @param TransactionsAPI client-provided implementation of the API for publishing events and transactions
 	 */
-	virtual void Initialize(IToolsContextQueriesAPI* QueriesAPI, IToolsContextTransactionsAPI* TransactionsAPI);
+	INTERACTIVETOOLSFRAMEWORK_API virtual void Initialize(IToolsContextQueriesAPI* QueriesAPI, IToolsContextTransactionsAPI* TransactionsAPI);
 
 	/** Shutdown Context by destroying InputRouter and ToolManager */
-	virtual void Shutdown();
-
-	virtual void DeactivateActiveTool(EToolSide WhichSide, EToolShutdownType ShutdownType);
-	virtual void DeactivateAllActiveTools(EToolShutdownType ShutdownType);
-
-	bool CanStartTool(EToolSide WhichSide, const FString& ToolTypeIdentifier) const;
-	bool HasActiveTool(EToolSide WhichSide) const;
-	FString GetActiveToolName(EToolSide WhichSide) const;
-	bool ActiveToolHasAccept(EToolSide WhichSide) const;
-	bool CanAcceptActiveTool(EToolSide WhichSide) const;
-	bool CanCancelActiveTool(EToolSide WhichSide) const;
-	bool CanCompleteActiveTool(EToolSide WhichSide) const;
-	bool StartTool(EToolSide WhichSide, const FString& ToolTypeIdentifier);
-	void EndTool(EToolSide WhichSide, EToolShutdownType ShutdownType);
-	bool IsToolActive(EToolSide WhichSide, const FString ToolIdentifier) const;
+	INTERACTIVETOOLSFRAMEWORK_API virtual void Shutdown();
+
+	INTERACTIVETOOLSFRAMEWORK_API virtual void DeactivateActiveTool(EToolSide WhichSide, EToolShutdownType ShutdownType);
+	INTERACTIVETOOLSFRAMEWORK_API virtual void DeactivateAllActiveTools(EToolShutdownType ShutdownType);
+
+	INTERACTIVETOOLSFRAMEWORK_API bool CanStartTool(EToolSide WhichSide, const FString& ToolTypeIdentifier) const;
+	INTERACTIVETOOLSFRAMEWORK_API bool HasActiveTool(EToolSide WhichSide) const;
+	INTERACTIVETOOLSFRAMEWORK_API FString GetActiveToolName(EToolSide WhichSide) const;
+	INTERACTIVETOOLSFRAMEWORK_API bool ActiveToolHasAccept(EToolSide WhichSide) const;
+	INTERACTIVETOOLSFRAMEWORK_API bool CanAcceptActiveTool(EToolSide WhichSide) const;
+	INTERACTIVETOOLSFRAMEWORK_API bool CanCancelActiveTool(EToolSide WhichSide) const;
+	INTERACTIVETOOLSFRAMEWORK_API bool CanCompleteActiveTool(EToolSide WhichSide) const;
+	INTERACTIVETOOLSFRAMEWORK_API bool StartTool(EToolSide WhichSide, const FString& ToolTypeIdentifier);
+	INTERACTIVETOOLSFRAMEWORK_API void EndTool(EToolSide WhichSide, EToolShutdownType ShutdownType);
+	INTERACTIVETOOLSFRAMEWORK_API bool IsToolActive(EToolSide WhichSide, const FString ToolIdentifier) const;
 
 public:
 	// forwards message to OnToolNotificationMessage delegate
-	virtual void PostToolNotificationMessage(const FText& Message);
-	virtual void PostToolWarningMessage(const FText& Message);
+	INTERACTIVETOOLSFRAMEWORK_API virtual void PostToolNotificationMessage(const FText& Message);
+	INTERACTIVETOOLSFRAMEWORK_API virtual void PostToolWarningMessage(const FText& Message);
 
 	DECLARE_MULTICAST_DELEGATE_OneParam(FToolsContextToolNotification, const FText&);
 	FToolsContextToolNotification OnToolNotificationMessage;
@@ -154,4 +154,4 @@ protected:
 	// todo: deprecate and remove this, can now be accomplished via CreateToolManagerFunc()
 	UPROPERTY()
 	TSoftClassPtr<UInteractiveToolManager> ToolManagerClass;
-};
\ No newline at end of file
+};
diff --git a/Engine/Source/Runtime/InteractiveToolsFramework/Public/MultiSelectionTool.h b/Engine/Source/Runtime/InteractiveToolsFramework/Public/MultiSelectionTool.h
--- a/Engine/Source/Runtime/InteractiveToolsFramework/Public/MultiSelectionTool.h
+++ b/Engine/Source/Runtime/InteractiveToolsFramework/Public/MultiSelectionTool.h
@@ -9,8 +9,8 @@
 
 #include "MultiSelectionTool.generated.h"
 
-UCLASS(Transient)
-class INTERACTIVETOOLSFRAMEWORK_API UMultiSelectionTool : public UInteractiveTool, public IInteractiveToolCameraFocusAPI
+UCLASS(Transient, MinimalAPI)
+class UMultiSelectionTool : public UInteractiveTool, public IInteractiveToolCameraFocusAPI
 {
 GENERATED_BODY()
 public:
@@ -47,10 +47,10 @@ protected:
 
 public:
 	// IInteractiveToolCameraFocusAPI implementation
-	virtual bool SupportsWorldSpaceFocusBox() override;
-	virtual FBox GetWorldSpaceFocusBox() override;
-	virtual bool SupportsWorldSpaceFocusPoint() override;
-	virtual bool GetWorldSpaceFocusPoint(const FRay& WorldRay, FVector& PointOut) override;
+	INTERACTIVETOOLSFRAMEWORK_API virtual bool SupportsWorldSpaceFocusBox() override;
+	INTERACTIVETOOLSFRAMEWORK_API virtual FBox GetWorldSpaceFocusBox() override;
+	INTERACTIVETOOLSFRAMEWORK_API virtual bool SupportsWorldSpaceFocusPoint() override;
+	INTERACTIVETOOLSFRAMEWORK_API virtual bool GetWorldSpaceFocusPoint(const FRay& WorldRay, FVector& PointOut) override;
 };
 
 
diff --git a/Engine/Source/Runtime/InteractiveToolsFramework/Public/SceneQueries/SceneSnappingManager.h b/Engine/Source/Runtime/InteractiveToolsFramework/Public/SceneQueries/SceneSnappingManager.h
--- a/Engine/Source/Runtime/InteractiveToolsFramework/Public/SceneQueries/SceneSnappingManager.h
+++ b/Engine/Source/Runtime/InteractiveToolsFramework/Public/SceneQueries/SceneSnappingManager.h
@@ -20,7 +20,7 @@ class UInteractiveGizmoManager;
 
 
 
-struct INTERACTIVETOOLSFRAMEWORK_API FSceneQueryVisibilityFilter
+struct FSceneQueryVisibilityFilter
 {
 	/** Optional: components to consider invisible even if they aren't. */
 	const TArray<const UPrimitiveComponent*>* ComponentsToIgnore = nullptr;
@@ -29,14 +29,14 @@ struct INTERACTIVETOOLSFRAMEWORK_API FSceneQueryVisibilityFilter
 	const TArray<const UPrimitiveComponent*>* InvisibleComponentsToInclude = nullptr;
 
 	/** @return true if the Component is currently configured as visible (does not consider ComponentsToIgnore or InvisibleComponentsToInclude lists) */
-	bool IsVisible(const UPrimitiveComponent* Component) const;
+	INTERACTIVETOOLSFRAMEWORK_API bool IsVisible(const UPrimitiveComponent* Component) const;
 };
 
 
 /**
 * Configuration variables for a USceneSnappingManager hit query request.
 */
-struct INTERACTIVETOOLSFRAMEWORK_API FSceneHitQueryRequest
+struct FSceneHitQueryRequest
 {
 	/** scene query ray */
 	FRay3d WorldRay;
@@ -50,7 +50,7 @@ struct INTERACTIVETOOLSFRAMEWORK_API FSceneHitQueryRequest
 /**
 * Computed result of a USceneSnappingManager hit query request
 */
-struct INTERACTIVETOOLSFRAMEWORK_API FSceneHitQueryResult
+struct FSceneHitQueryResult
 {
 	/** Actor that owns hit target */
 	AActor* TargetActor = nullptr;
@@ -69,7 +69,7 @@ struct INTERACTIVETOOLSFRAMEWORK_API FSceneHitQueryResult
 
 	FHitResult HitResult;
 
-	void InitializeHitResult(const FSceneHitQueryRequest& FromRequest);
+	INTERACTIVETOOLSFRAMEWORK_API void InitializeHitResult(const FSceneHitQueryRequest& FromRequest);
 };
 
 
@@ -105,7 +105,7 @@ ENUM_CLASS_FLAGS(ESceneSnapQueryTargetType);
 /**
  * Configuration variables for a USceneSnappingManager snap query request.
  */
-struct INTERACTIVETOOLSFRAMEWORK_API FSceneSnapQueryRequest
+struct FSceneSnapQueryRequest
 {
 	/** What type of snap query geometry is this */
 	ESceneSnapQueryType RequestType = ESceneSnapQueryType::Position;
@@ -140,7 +140,7 @@ struct INTERACTIVETOOLSFRAMEWORK_API FSceneSnapQueryRequest
 /**
  * Computed result of a USceneSnappingManager snap query request
  */
-struct INTERACTIVETOOLSFRAMEWORK_API FSceneSnapQueryResult
+struct FSceneSnapQueryResult
 {
 	/** Actor that owns snap target */
 	AActor* TargetActor = nullptr;
@@ -173,8 +173,8 @@ struct INTERACTIVETOOLSFRAMEWORK_API FSceneSnapQueryResult
  * 
  * See UModelingSceneSnappingManager for a sample implementation.
  */
-UCLASS()
-class INTERACTIVETOOLSFRAMEWORK_API USceneSnappingManager : public UObject
+UCLASS(MinimalAPI)
+class USceneSnappingManager : public UObject
 {
 	GENERATED_BODY()
 public:
@@ -208,10 +208,10 @@ public:
 	/**
 	 * @return existing USceneSnappingManager registered in context store via the ToolManager, or nullptr if not found
 	 */
-	static USceneSnappingManager* Find(UInteractiveToolManager* ToolManager);
+	static INTERACTIVETOOLSFRAMEWORK_API USceneSnappingManager* Find(UInteractiveToolManager* ToolManager);
 
 	/**
 	 * @return existing USceneSnappingManager registered in context store via the ToolManager, or nullptr if not found
 	 */
-	static USceneSnappingManager* Find(UInteractiveGizmoManager* GizmoManager);
-};
\ No newline at end of file
+	static INTERACTIVETOOLSFRAMEWORK_API USceneSnappingManager* Find(UInteractiveGizmoManager* GizmoManager);
+};
diff --git a/Engine/Source/Runtime/InteractiveToolsFramework/Public/SelectionSet.h b/Engine/Source/Runtime/InteractiveToolsFramework/Public/SelectionSet.h
--- a/Engine/Source/Runtime/InteractiveToolsFramework/Public/SelectionSet.h
+++ b/Engine/Source/Runtime/InteractiveToolsFramework/Public/SelectionSet.h
@@ -24,8 +24,8 @@ DECLARE_MULTICAST_DELEGATE_OneParam(FSelectionSetModifiedSignature, USelectionSe
  * For example the UMeshSelectionSet implementation stores information on selected
  * triangles, vertices, etc. 
  */
-UCLASS(Transient)
-class INTERACTIVETOOLSFRAMEWORK_API USelectionSet : public UObject
+UCLASS(Transient, MinimalAPI)
+class USelectionSet : public UObject
 {
 	GENERATED_BODY()
 
@@ -66,13 +66,13 @@ enum class EMeshSelectionElementType
  * UMeshSelectionSet is an implementation of USelectionSet that represents selections on indexed meshes.
  * Vertices, Edges, Faces, and Groups can be selected.
  */
-UCLASS(Transient)
-class INTERACTIVETOOLSFRAMEWORK_API UMeshSelectionSet : public USelectionSet
+UCLASS(Transient, MinimalAPI)
+class UMeshSelectionSet : public USelectionSet
 {
 	GENERATED_BODY()
 
 public:
-	UMeshSelectionSet();
+	INTERACTIVETOOLSFRAMEWORK_API UMeshSelectionSet();
 
 	UPROPERTY()
 	TArray<int32> Vertices;
@@ -89,31 +89,31 @@ public:
 
 public:
 	/** @return current selected elements of the given ElementType */
-	TArray<int>& GetElements(EMeshSelectionElementType ElementType);
+	INTERACTIVETOOLSFRAMEWORK_API TArray<int>& GetElements(EMeshSelectionElementType ElementType);
 
 	/** @return current selected elements of the given ElementType */
-	const TArray<int>& GetElements(EMeshSelectionElementType ElementType) const;
+	INTERACTIVETOOLSFRAMEWORK_API const TArray<int>& GetElements(EMeshSelectionElementType ElementType) const;
 
 	/**
 	 * Add list of Indices to the current selection for the given ElementType
 	 * @warning no duplicate detection is currently done
 	 */
-	void AddIndices(EMeshSelectionElementType ElementType, const TArray<int32>& Indices);
+	INTERACTIVETOOLSFRAMEWORK_API void AddIndices(EMeshSelectionElementType ElementType, const TArray<int32>& Indices);
 
 	/**
 	 * Add list of Indices to the current selection for the given ElementType
 	 * @warning no duplicate detection is currently done
 	 */
-	void AddIndices(EMeshSelectionElementType ElementType, const TSet<int32>& Indices);
+	INTERACTIVETOOLSFRAMEWORK_API void AddIndices(EMeshSelectionElementType ElementType, const TSet<int32>& Indices);
 
 
 	/**
 	 * Remove list of Indices from the current selection for the given ElementType
 	 */
-	void RemoveIndices(EMeshSelectionElementType ElementType, const TArray<int32>& Indices);
+	INTERACTIVETOOLSFRAMEWORK_API void RemoveIndices(EMeshSelectionElementType ElementType, const TArray<int32>& Indices);
 
 	/**
 	 * Remove list of Indices from the current selection for the given ElementType
 	 */
-	void RemoveIndices(EMeshSelectionElementType ElementType, const TSet<int32>& Indices);
-};
\ No newline at end of file
+	INTERACTIVETOOLSFRAMEWORK_API void RemoveIndices(EMeshSelectionElementType ElementType, const TSet<int32>& Indices);
+};
diff --git a/Engine/Source/Runtime/InteractiveToolsFramework/Public/SingleSelectionTool.h b/Engine/Source/Runtime/InteractiveToolsFramework/Public/SingleSelectionTool.h
--- a/Engine/Source/Runtime/InteractiveToolsFramework/Public/SingleSelectionTool.h
+++ b/Engine/Source/Runtime/InteractiveToolsFramework/Public/SingleSelectionTool.h
@@ -9,8 +9,8 @@
 
 #include "SingleSelectionTool.generated.h"
 
-UCLASS(Transient)
-class INTERACTIVETOOLSFRAMEWORK_API USingleSelectionTool : public UInteractiveTool, public IInteractiveToolCameraFocusAPI
+UCLASS(Transient, MinimalAPI)
+class USingleSelectionTool : public UInteractiveTool, public IInteractiveToolCameraFocusAPI
 {
 GENERATED_BODY()
 public:
@@ -53,9 +53,9 @@ protected:
 
 public:
 	// IInteractiveToolCameraFocusAPI implementation
-	virtual bool SupportsWorldSpaceFocusBox() override;
-	virtual FBox GetWorldSpaceFocusBox() override;
-	virtual bool SupportsWorldSpaceFocusPoint() override;
-	virtual bool GetWorldSpaceFocusPoint(const FRay& WorldRay, FVector& PointOut) override;
+	INTERACTIVETOOLSFRAMEWORK_API virtual bool SupportsWorldSpaceFocusBox() override;
+	INTERACTIVETOOLSFRAMEWORK_API virtual FBox GetWorldSpaceFocusBox() override;
+	INTERACTIVETOOLSFRAMEWORK_API virtual bool SupportsWorldSpaceFocusPoint() override;
+	INTERACTIVETOOLSFRAMEWORK_API virtual bool GetWorldSpaceFocusPoint(const FRay& WorldRay, FVector& PointOut) override;
 
 };
diff --git a/Engine/Source/Runtime/InteractiveToolsFramework/Public/TargetInterfaces/AssetBackedTarget.h b/Engine/Source/Runtime/InteractiveToolsFramework/Public/TargetInterfaces/AssetBackedTarget.h
--- a/Engine/Source/Runtime/InteractiveToolsFramework/Public/TargetInterfaces/AssetBackedTarget.h
+++ b/Engine/Source/Runtime/InteractiveToolsFramework/Public/TargetInterfaces/AssetBackedTarget.h
@@ -9,13 +9,13 @@
 
 #include "AssetBackedTarget.generated.h"
 
-UINTERFACE()
-class INTERACTIVETOOLSFRAMEWORK_API UAssetBackedTarget : public UInterface
+UINTERFACE(MinimalAPI)
+class UAssetBackedTarget : public UInterface
 {
 	GENERATED_BODY()
 };
 
-class INTERACTIVETOOLSFRAMEWORK_API IAssetBackedTarget
+class IAssetBackedTarget
 {
 	GENERATED_BODY()
 
@@ -30,4 +30,4 @@ public:
 		IAssetBackedTarget* OtherAsset = Cast<IAssetBackedTarget>(OtherTarget);
 		return OtherAsset && GetSourceData() == OtherAsset->GetSourceData();
 	}
-};
\ No newline at end of file
+};
diff --git a/Engine/Source/Runtime/InteractiveToolsFramework/Public/TargetInterfaces/MaterialProvider.h b/Engine/Source/Runtime/InteractiveToolsFramework/Public/TargetInterfaces/MaterialProvider.h
--- a/Engine/Source/Runtime/InteractiveToolsFramework/Public/TargetInterfaces/MaterialProvider.h
+++ b/Engine/Source/Runtime/InteractiveToolsFramework/Public/TargetInterfaces/MaterialProvider.h
@@ -12,7 +12,7 @@ class UMaterialInterface;
 /**
  * FComponentMaterialSet is the set of materials assigned to a component (ie Material Slots on a StaticMesh)
  */
-struct INTERACTIVETOOLSFRAMEWORK_API FComponentMaterialSet
+struct FComponentMaterialSet
 {
 	TArray<UMaterialInterface*> Materials;
 
@@ -22,13 +22,13 @@ struct INTERACTIVETOOLSFRAMEWORK_API FComponentMaterialSet
 	}
 };
 
-UINTERFACE()
-class INTERACTIVETOOLSFRAMEWORK_API UMaterialProvider : public UInterface
+UINTERFACE(MinimalAPI)
+class UMaterialProvider : public UInterface
 {
 	GENERATED_BODY()
 };
 
-class INTERACTIVETOOLSFRAMEWORK_API IMaterialProvider
+class IMaterialProvider
 {
 	GENERATED_BODY()
 
@@ -68,4 +68,4 @@ public:
 	 * @return true if successful
 	 */
 	virtual bool CommitMaterialSetUpdate(const FComponentMaterialSet& MaterialSet, bool bApplyToAsset) = 0;
-};
\ No newline at end of file
+};
diff --git a/Engine/Source/Runtime/InteractiveToolsFramework/Public/TargetInterfaces/MeshDescriptionCommitter.h b/Engine/Source/Runtime/InteractiveToolsFramework/Public/TargetInterfaces/MeshDescriptionCommitter.h
--- a/Engine/Source/Runtime/InteractiveToolsFramework/Public/TargetInterfaces/MeshDescriptionCommitter.h
+++ b/Engine/Source/Runtime/InteractiveToolsFramework/Public/TargetInterfaces/MeshDescriptionCommitter.h
@@ -10,13 +10,13 @@
 
 #include "MeshDescriptionCommitter.generated.h"
 
-UINTERFACE()
-class INTERACTIVETOOLSFRAMEWORK_API UMeshDescriptionCommitter : public UInterface
+UINTERFACE(MinimalAPI)
+class UMeshDescriptionCommitter : public UInterface
 {
 	GENERATED_BODY()
 };
 
-class INTERACTIVETOOLSFRAMEWORK_API IMeshDescriptionCommitter
+class IMeshDescriptionCommitter
 {
 	GENERATED_BODY()
 
@@ -84,4 +84,4 @@ protected:
 	*/
 	virtual void CommitMeshDescription(const FCommitter& Committer, const FCommitMeshParameters& CommitParams = FCommitMeshParameters()) = 0;
 
-};
\ No newline at end of file
+};
diff --git a/Engine/Source/Runtime/InteractiveToolsFramework/Public/TargetInterfaces/MeshDescriptionProvider.h b/Engine/Source/Runtime/InteractiveToolsFramework/Public/TargetInterfaces/MeshDescriptionProvider.h
--- a/Engine/Source/Runtime/InteractiveToolsFramework/Public/TargetInterfaces/MeshDescriptionProvider.h
+++ b/Engine/Source/Runtime/InteractiveToolsFramework/Public/TargetInterfaces/MeshDescriptionProvider.h
@@ -9,13 +9,13 @@
 
 #include "MeshDescriptionProvider.generated.h"
 
-UINTERFACE()
-class INTERACTIVETOOLSFRAMEWORK_API UMeshDescriptionProvider : public UInterface
+UINTERFACE(MinimalAPI)
+class UMeshDescriptionProvider : public UInterface
 {
 	GENERATED_BODY()
 };
 
-class INTERACTIVETOOLSFRAMEWORK_API IMeshDescriptionProvider
+class IMeshDescriptionProvider
 {
 	GENERATED_BODY()
 
@@ -73,4 +73,4 @@ public:
 		return false;
 	}
 
-};
\ No newline at end of file
+};
diff --git a/Engine/Source/Runtime/InteractiveToolsFramework/Public/TargetInterfaces/MeshTargetInterfaceTypes.h b/Engine/Source/Runtime/InteractiveToolsFramework/Public/TargetInterfaces/MeshTargetInterfaceTypes.h
--- a/Engine/Source/Runtime/InteractiveToolsFramework/Public/TargetInterfaces/MeshTargetInterfaceTypes.h
+++ b/Engine/Source/Runtime/InteractiveToolsFramework/Public/TargetInterfaces/MeshTargetInterfaceTypes.h
@@ -29,7 +29,7 @@ typedef EMeshLODIdentifier EStaticMeshEditingLOD;
  * FGetMeshParameters is used by ToolTarget Interfaces/Implementations that support returning a mesh, to allow clients
  * to specify various options like a specific LOD of a mesh, which attributes are required, etc.
  */
-struct INTERACTIVETOOLSFRAMEWORK_API FGetMeshParameters
+struct FGetMeshParameters
 {
 	FGetMeshParameters() = default;
 	explicit FGetMeshParameters(bool bHaveRequestLOD, EMeshLODIdentifier RequestLOD)
@@ -53,7 +53,7 @@ struct INTERACTIVETOOLSFRAMEWORK_API FGetMeshParameters
  * FCommitMeshParameters is used by ToolTarget Interfaces/Implementations that support setting a mesh (eg on a StaticMesh Asset),
  * to allow the client to specify various options
  */
-struct INTERACTIVETOOLSFRAMEWORK_API FCommitMeshParameters
+struct FCommitMeshParameters
 {
 	FCommitMeshParameters() = default;
 	explicit FCommitMeshParameters(bool bHaveTargetLOD, EMeshLODIdentifier TargetLOD) :
@@ -64,4 +64,4 @@ struct INTERACTIVETOOLSFRAMEWORK_API FCommitMeshParameters
 	bool bHaveTargetLOD = false;
 	/** Specify which LOD to set in a mesh-containing ToolTarget. Ignored unless bHaveTargetLOD == true. */
 	EMeshLODIdentifier TargetLOD = EMeshLODIdentifier::Default;
-};
\ No newline at end of file
+};
diff --git a/Engine/Source/Runtime/InteractiveToolsFramework/Public/TargetInterfaces/PhysicsDataSource.h b/Engine/Source/Runtime/InteractiveToolsFramework/Public/TargetInterfaces/PhysicsDataSource.h
--- a/Engine/Source/Runtime/InteractiveToolsFramework/Public/TargetInterfaces/PhysicsDataSource.h
+++ b/Engine/Source/Runtime/InteractiveToolsFramework/Public/TargetInterfaces/PhysicsDataSource.h
@@ -10,8 +10,8 @@
 class UBodySetup;
 class IInterface_CollisionDataProvider;
 
-UINTERFACE()
-class INTERACTIVETOOLSFRAMEWORK_API UPhysicsDataSource : public UInterface
+UINTERFACE(MinimalAPI)
+class UPhysicsDataSource : public UInterface
 {
 	GENERATED_BODY()
 };
@@ -19,7 +19,7 @@ class INTERACTIVETOOLSFRAMEWORK_API UPhysicsDataSource : public UInterface
 /**
  * IPhysicsDataSource is a ToolTarget Interface that provides read/write access to physics-related data structures.
  */
-class INTERACTIVETOOLSFRAMEWORK_API IPhysicsDataSource
+class IPhysicsDataSource
 {
 	GENERATED_BODY()
 
diff --git a/Engine/Source/Runtime/InteractiveToolsFramework/Public/TargetInterfaces/PrimitiveComponentBackedTarget.h b/Engine/Source/Runtime/InteractiveToolsFramework/Public/TargetInterfaces/PrimitiveComponentBackedTarget.h
--- a/Engine/Source/Runtime/InteractiveToolsFramework/Public/TargetInterfaces/PrimitiveComponentBackedTarget.h
+++ b/Engine/Source/Runtime/InteractiveToolsFramework/Public/TargetInterfaces/PrimitiveComponentBackedTarget.h
@@ -15,13 +15,13 @@ class AActor;
 
 struct FHitResult;
 
-UINTERFACE()
-class INTERACTIVETOOLSFRAMEWORK_API UPrimitiveComponentBackedTarget : public UInterface
+UINTERFACE(MinimalAPI)
+class UPrimitiveComponentBackedTarget : public UInterface
 {
 	GENERATED_BODY()
 };
 
-class INTERACTIVETOOLSFRAMEWORK_API IPrimitiveComponentBackedTarget
+class IPrimitiveComponentBackedTarget
 {
 	GENERATED_BODY()
 
@@ -52,4 +52,4 @@ public:
 	 * @return true if ray intersected Component
 	 */
 	virtual bool HitTestComponent(const FRay& WorldRay, FHitResult& OutHit) const = 0;
-};
\ No newline at end of file
+};
diff --git a/Engine/Source/Runtime/InteractiveToolsFramework/Public/TargetInterfaces/SkeletalMeshBackedTarget.h b/Engine/Source/Runtime/InteractiveToolsFramework/Public/TargetInterfaces/SkeletalMeshBackedTarget.h
--- a/Engine/Source/Runtime/InteractiveToolsFramework/Public/TargetInterfaces/SkeletalMeshBackedTarget.h
+++ b/Engine/Source/Runtime/InteractiveToolsFramework/Public/TargetInterfaces/SkeletalMeshBackedTarget.h
@@ -9,13 +9,13 @@
 
 #include "SkeletalMeshBackedTarget.generated.h"
 
-UINTERFACE()
-class INTERACTIVETOOLSFRAMEWORK_API USkeletalMeshBackedTarget : public UAssetBackedTarget
+UINTERFACE(MinimalAPI)
+class USkeletalMeshBackedTarget : public UAssetBackedTarget
 {
 	GENERATED_BODY()
 };
 
-class INTERACTIVETOOLSFRAMEWORK_API ISkeletalMeshBackedTarget : public IAssetBackedTarget
+class ISkeletalMeshBackedTarget : public IAssetBackedTarget
 {
 	GENERATED_BODY()
 
diff --git a/Engine/Source/Runtime/InteractiveToolsFramework/Public/TargetInterfaces/StaticMeshBackedTarget.h b/Engine/Source/Runtime/InteractiveToolsFramework/Public/TargetInterfaces/StaticMeshBackedTarget.h
--- a/Engine/Source/Runtime/InteractiveToolsFramework/Public/TargetInterfaces/StaticMeshBackedTarget.h
+++ b/Engine/Source/Runtime/InteractiveToolsFramework/Public/TargetInterfaces/StaticMeshBackedTarget.h
@@ -10,13 +10,13 @@
 
 #include "StaticMeshBackedTarget.generated.h"
 
-UINTERFACE()
-class INTERACTIVETOOLSFRAMEWORK_API UStaticMeshBackedTarget : public UAssetBackedTarget
+UINTERFACE(MinimalAPI)
+class UStaticMeshBackedTarget : public UAssetBackedTarget
 {
 	GENERATED_BODY()
 };
 
-class INTERACTIVETOOLSFRAMEWORK_API IStaticMeshBackedTarget : public IAssetBackedTarget
+class IStaticMeshBackedTarget : public IAssetBackedTarget
 {
 	GENERATED_BODY()
 
@@ -33,4 +33,4 @@ public:
 	 * @return the underlying source StaticMesh asset for this Target
 	 */
 	virtual UStaticMesh* GetStaticMesh() const = 0;
-};
\ No newline at end of file
+};
diff --git a/Engine/Source/Runtime/InteractiveToolsFramework/Public/ToolContextInterfaces.h b/Engine/Source/Runtime/InteractiveToolsFramework/Public/ToolContextInterfaces.h
--- a/Engine/Source/Runtime/InteractiveToolsFramework/Public/ToolContextInterfaces.h
+++ b/Engine/Source/Runtime/InteractiveToolsFramework/Public/ToolContextInterfaces.h
@@ -44,7 +44,7 @@ struct FTypedElementHandle;
  * to construct a Tool. This information comes from a level above the Tools framework,
  * and depends on the context we are in (Editor vs Runtime, for example).
  */
-struct INTERACTIVETOOLSFRAMEWORK_API FToolBuilderState
+struct FToolBuilderState
 {
 	/** The current UWorld */
 	UWorld* World = nullptr;
@@ -77,7 +77,7 @@ struct INTERACTIVETOOLSFRAMEWORK_API FToolBuilderState
  * FViewCameraState is a bucket of state information that a Tool might
  * need to implement interactions that depend on the current scene view.
  */
-struct INTERACTIVETOOLSFRAMEWORK_API FViewCameraState
+struct FViewCameraState
 {
 	/** Current camera/head position */
 	FVector Position;
@@ -156,7 +156,7 @@ enum class EToolContextTransformGizmoMode : uint8
 /**
  * Snapping configuration settings/state
  */
-struct INTERACTIVETOOLSFRAMEWORK_API FToolContextSnappingConfiguration
+struct FToolContextSnappingConfiguration
 {
 	/** Specify whether position snapping should be applied */
 	bool bEnablePositionGridSnapping = false;
diff --git a/Engine/Source/Runtime/InteractiveToolsFramework/Public/ToolDataVisualizer.h b/Engine/Source/Runtime/InteractiveToolsFramework/Public/ToolDataVisualizer.h
--- a/Engine/Source/Runtime/InteractiveToolsFramework/Public/ToolDataVisualizer.h
+++ b/Engine/Source/Runtime/InteractiveToolsFramework/Public/ToolDataVisualizer.h
@@ -14,7 +14,7 @@ class IToolsContextRenderAPI;
  * 
  * Currently uses PDI drawing but may use different back-ends in the future
  */
-class INTERACTIVETOOLSFRAMEWORK_API FToolDataVisualizer
+class FToolDataVisualizer
 {
 public:
 	/** Default color used for drawing lines */
@@ -33,7 +33,7 @@ public:
 	bool bDepthTested = true;
 
 public:
-	FToolDataVisualizer();
+	INTERACTIVETOOLSFRAMEWORK_API FToolDataVisualizer();
 	virtual ~FToolDataVisualizer() {}
 
 
@@ -43,13 +43,13 @@ public:
 
 
 	/** This must be called every frame to allow Visualizer to extract necessary rendering data/objects */
-	void BeginFrame(IToolsContextRenderAPI* RenderAPI, const FViewCameraState& CameraState);
+	INTERACTIVETOOLSFRAMEWORK_API void BeginFrame(IToolsContextRenderAPI* RenderAPI, const FViewCameraState& CameraState);
 
 	/** This must be called every frame to allow Visualizer to extract necessary rendering data/objects */
-	void BeginFrame(IToolsContextRenderAPI* RenderAPI);
+	INTERACTIVETOOLSFRAMEWORK_API void BeginFrame(IToolsContextRenderAPI* RenderAPI);
 
 	/** this should be called at the end of every frame to allow for necessary cleanup */
-	void EndFrame();
+	INTERACTIVETOOLSFRAMEWORK_API void EndFrame();
 
 
 	//
@@ -57,13 +57,13 @@ public:
 	// 
 
 	/** Clear transform stack and push the given Transform */
-	void SetTransform(const FTransform& Transform);
+	INTERACTIVETOOLSFRAMEWORK_API void SetTransform(const FTransform& Transform);
 	/** Push a Transform onto the transform stack */
-	void PushTransform(const FTransform& Transform);
+	INTERACTIVETOOLSFRAMEWORK_API void PushTransform(const FTransform& Transform);
 	/** Pop a transform from the transform stack */
-	void PopTransform();
+	INTERACTIVETOOLSFRAMEWORK_API void PopTransform();
 	/** Clear the transform stack to identity */
-	void PopAllTransforms();
+	INTERACTIVETOOLSFRAMEWORK_API void PopAllTransforms();
 
 	/** @return input Point transformed by transform stack */
 	FVector TransformP(const FVector& Point)
@@ -246,22 +246,22 @@ protected:
 	FTransform TotalTransform;
 
 	// actually does the line drawing; assumes A and B are already transformed
-	virtual void InternalDrawTransformedLine(const FVector& A, const FVector& B, const FLinearColor& Color, float LineThickness, bool bDepthTested);
+	INTERACTIVETOOLSFRAMEWORK_API virtual void InternalDrawTransformedLine(const FVector& A, const FVector& B, const FLinearColor& Color, float LineThickness, bool bDepthTested);
 	// actually does the point drawing; assumes Position is already transformed
-	virtual void InternalDrawTransformedPoint(const FVector& Position, const FLinearColor& Color, float PointSize, bool bDepthTested);
+	INTERACTIVETOOLSFRAMEWORK_API virtual void InternalDrawTransformedPoint(const FVector& Position, const FLinearColor& Color, float PointSize, bool bDepthTested);
 
 
 	// actually does the circle drawing
-	virtual void InternalDrawCircle(const FVector& Position, const FVector& Normal, float Radius, int Steps, const FLinearColor& Color, float LineThickness, bool bDepthTested);
+	INTERACTIVETOOLSFRAMEWORK_API virtual void InternalDrawCircle(const FVector& Position, const FVector& Normal, float Radius, int Steps, const FLinearColor& Color, float LineThickness, bool bDepthTested);
 	// actually does the circle drawing
-	virtual void InternalDrawViewFacingCircle(const FVector& Position, float Radius, int Steps, const FLinearColor& Color, float LineThickness, bool bDepthTested);
+	INTERACTIVETOOLSFRAMEWORK_API virtual void InternalDrawViewFacingCircle(const FVector& Position, float Radius, int Steps, const FLinearColor& Color, float LineThickness, bool bDepthTested);
 	// actually does the cylinder drawing
-	virtual void InternalDrawWireCylinder(const FVector& Position, const FVector& Normal, float Radius, float Height, int Steps, const FLinearColor& Color, float LineThickness, bool bDepthTested);
+	INTERACTIVETOOLSFRAMEWORK_API virtual void InternalDrawWireCylinder(const FVector& Position, const FVector& Normal, float Radius, float Height, int Steps, const FLinearColor& Color, float LineThickness, bool bDepthTested);
 	// actually does the box drawing
-	virtual void InternalDrawWireBox(const FBox& Box, const FLinearColor& Color, float LineThickness, bool bDepthTested);
+	INTERACTIVETOOLSFRAMEWORK_API virtual void InternalDrawWireBox(const FBox& Box, const FLinearColor& Color, float LineThickness, bool bDepthTested);
 	// actually does the square drawing
-	virtual void InternalDrawSquare(const FVector& Center, const FVector& SideA, const FVector& SideB, const FLinearColor& Color, float LineThickness, bool bDepthTested);
+	INTERACTIVETOOLSFRAMEWORK_API virtual void InternalDrawSquare(const FVector& Center, const FVector& SideA, const FVector& SideB, const FLinearColor& Color, float LineThickness, bool bDepthTested);
 	// actually does the X drawing
-	virtual void InternalDrawViewFacingX(const FVector& Position, float Width, const FLinearColor& Color, float LineThickness, bool bDepthTested);
+	INTERACTIVETOOLSFRAMEWORK_API virtual void InternalDrawViewFacingX(const FVector& Position, float Width, const FLinearColor& Color, float LineThickness, bool bDepthTested);
 
-};
\ No newline at end of file
+};
diff --git a/Engine/Source/Runtime/InteractiveToolsFramework/Public/ToolTargetManager.h b/Engine/Source/Runtime/InteractiveToolsFramework/Public/ToolTargetManager.h
--- a/Engine/Source/Runtime/InteractiveToolsFramework/Public/ToolTargetManager.h
+++ b/Engine/Source/Runtime/InteractiveToolsFramework/Public/ToolTargetManager.h
@@ -32,8 +32,8 @@
  * 
  * See the class comment for UToolTarget for more info.
  */
-UCLASS(Transient)
-class INTERACTIVETOOLSFRAMEWORK_API UToolTargetManager : public UObject
+UCLASS(Transient, MinimalAPI)
+class UToolTargetManager : public UObject
 {
 	GENERATED_BODY()
 
@@ -46,12 +46,12 @@ public:
 	/**
 	 * Add a new ToolTargetFactory
 	 */
-	virtual void AddTargetFactory(UToolTargetFactory* Factory);
+	INTERACTIVETOOLSFRAMEWORK_API virtual void AddTargetFactory(UToolTargetFactory* Factory);
 
 	/**
 	 * Find the first ToolTargetFactory that passes the Predicate function
 	 */
-	virtual UToolTargetFactory* FindFirstFactoryByPredicate(TFunctionRef<bool(UToolTargetFactory*)> Predicate);
+	INTERACTIVETOOLSFRAMEWORK_API virtual UToolTargetFactory* FindFirstFactoryByPredicate(TFunctionRef<bool(UToolTargetFactory*)> Predicate);
 
 	/**
 	 * Find the first ToolTargetFactory of a given type
@@ -66,7 +66,7 @@ public:
 
 
 	/** Examines stored target factories to see if one can build the requested type of target. */
-	virtual bool CanBuildTarget(UObject* SourceObject, const FToolTargetTypeRequirements& TargetRequirements) const;
+	INTERACTIVETOOLSFRAMEWORK_API virtual bool CanBuildTarget(UObject* SourceObject, const FToolTargetTypeRequirements& TargetRequirements) const;
 
 	/** 
 	 * Uses one of the stored factories to build a tool target out of the given input object
@@ -76,7 +76,7 @@ public:
 	 * 
 	 * @return Tool target that staisfies given requirements, or nullptr if none could be created.
 	 */
-	virtual UToolTarget* BuildTarget(UObject* SourceObject, const FToolTargetTypeRequirements& TargetRequirements);
+	INTERACTIVETOOLSFRAMEWORK_API virtual UToolTarget* BuildTarget(UObject* SourceObject, const FToolTargetTypeRequirements& TargetRequirements);
 
 	/** Much like BuildTarget, but casts the target to the template argument before returning. */
 	template<typename CastToType>
@@ -90,14 +90,14 @@ public:
 	 * Looks through the currently selected components and actors and counts the number of
 	 * inputs that could be used to create qualifying tool targets.
 	 */
-	virtual int32 CountSelectedAndTargetable(const FToolBuilderState& SceneState,
+	INTERACTIVETOOLSFRAMEWORK_API virtual int32 CountSelectedAndTargetable(const FToolBuilderState& SceneState,
 		const FToolTargetTypeRequirements& TargetRequirements) const;
 
 	/**
 	 * Looks through the currently selected components and actors and counts the number of
 	 * inputs that could be used to create qualifying tool targets.
 	 */
-	virtual void EnumerateSelectedAndTargetableComponents(const FToolBuilderState& SceneState,
+	INTERACTIVETOOLSFRAMEWORK_API virtual void EnumerateSelectedAndTargetableComponents(const FToolBuilderState& SceneState,
 		const FToolTargetTypeRequirements& TargetRequirements,
 		TFunctionRef<void(UActorComponent*)> ComponentFunc) const;
 
@@ -105,7 +105,7 @@ public:
 	 * Looks through the currently selected components and actors and counts the number of
 	 * inputs that could be used to create qualifying tool targets with an additional test predicate.
 	 */
-	virtual int32 CountSelectedAndTargetableWithPredicate(const FToolBuilderState& SceneState,
+	INTERACTIVETOOLSFRAMEWORK_API virtual int32 CountSelectedAndTargetableWithPredicate(const FToolBuilderState& SceneState,
 		const FToolTargetTypeRequirements& TargetRequirements,
 		TFunctionRef<bool(UActorComponent&)> ComponentPred) const;
 
@@ -114,21 +114,21 @@ public:
 	 * Looks through the currently selected components and actors and builds a target out of
 	 * the first encountered element that satisfies the requirements. 
 	 */
-	virtual UToolTarget* BuildFirstSelectedTargetable(const FToolBuilderState& SceneState,
+	INTERACTIVETOOLSFRAMEWORK_API virtual UToolTarget* BuildFirstSelectedTargetable(const FToolBuilderState& SceneState,
 		const FToolTargetTypeRequirements& TargetRequirements);
 
 	/**
 	 * Looks through the current selected components and actors and builds all targets that
 	 * satisfy the requirements.
 	 */
-	virtual TArray<TObjectPtr<UToolTarget>> BuildAllSelectedTargetable(const FToolBuilderState& SceneState,
+	INTERACTIVETOOLSFRAMEWORK_API virtual TArray<TObjectPtr<UToolTarget>> BuildAllSelectedTargetable(const FToolBuilderState& SceneState,
 		const FToolTargetTypeRequirements& TargetRequirements);
 
 	/** Initialize the ToolTargetManager. UInteractiveToolsContext calls this, you should not. */
-	virtual void Initialize();
+	INTERACTIVETOOLSFRAMEWORK_API virtual void Initialize();
 
 	/** Shutdown the ToolTargetManager. Called by UInteractiveToolsContext. */
-	virtual void Shutdown();
+	INTERACTIVETOOLSFRAMEWORK_API virtual void Shutdown();
 
 protected:
 
@@ -146,4 +146,4 @@ protected:
 	TArray<TObjectPtr<UToolTargetFactory>> Factories;
 
 	// More state will go here if the manager deals with target caching.
-};
\ No newline at end of file
+};
diff --git a/Engine/Source/Runtime/InteractiveToolsFramework/Public/ToolTargets/PrimitiveComponentToolTarget.h b/Engine/Source/Runtime/InteractiveToolsFramework/Public/ToolTargets/PrimitiveComponentToolTarget.h
--- a/Engine/Source/Runtime/InteractiveToolsFramework/Public/ToolTargets/PrimitiveComponentToolTarget.h
+++ b/Engine/Source/Runtime/InteractiveToolsFramework/Public/ToolTargets/PrimitiveComponentToolTarget.h
@@ -13,47 +13,47 @@
  * A tool target to share some reusable code for tool targets that are
  * backed by primitive components. 
  */
-UCLASS(Transient)
-class INTERACTIVETOOLSFRAMEWORK_API UPrimitiveComponentToolTarget : public UToolTarget, public IPrimitiveComponentBackedTarget
+UCLASS(Transient, MinimalAPI)
+class UPrimitiveComponentToolTarget : public UToolTarget, public IPrimitiveComponentBackedTarget
 {
 	GENERATED_BODY()
 public:
 
 	// UToolTarget
-	virtual bool IsValid() const override;
+	INTERACTIVETOOLSFRAMEWORK_API virtual bool IsValid() const override;
 
 	// IPrimitiveComponentBackedTarget implementation
-	virtual UPrimitiveComponent* GetOwnerComponent() const override;
-	virtual AActor* GetOwnerActor() const override;
-	virtual void SetOwnerVisibility(bool bVisible) const override;
-	virtual FTransform GetWorldTransform() const override;
-	virtual bool HitTestComponent(const FRay& WorldRay, FHitResult& OutHit) const override;
+	INTERACTIVETOOLSFRAMEWORK_API virtual UPrimitiveComponent* GetOwnerComponent() const override;
+	INTERACTIVETOOLSFRAMEWORK_API virtual AActor* GetOwnerActor() const override;
+	INTERACTIVETOOLSFRAMEWORK_API virtual void SetOwnerVisibility(bool bVisible) const override;
+	INTERACTIVETOOLSFRAMEWORK_API virtual FTransform GetWorldTransform() const override;
+	INTERACTIVETOOLSFRAMEWORK_API virtual bool HitTestComponent(const FRay& WorldRay, FHitResult& OutHit) const override;
 
 	// UObject
-	void BeginDestroy() override;
+	INTERACTIVETOOLSFRAMEWORK_API void BeginDestroy() override;
 
 protected:
 	friend class UPrimitiveComponentToolTargetFactory;
 	
-	void InitializeComponent(UPrimitiveComponent* ComponentIn);
+	INTERACTIVETOOLSFRAMEWORK_API void InitializeComponent(UPrimitiveComponent* ComponentIn);
 
 	TWeakObjectPtr<UPrimitiveComponent> Component;
 
 private:
 #if WITH_EDITOR
-	void OnObjectsReplaced(const FCoreUObjectDelegates::FReplacementObjectMap& Map);
+	INTERACTIVETOOLSFRAMEWORK_API void OnObjectsReplaced(const FCoreUObjectDelegates::FReplacementObjectMap& Map);
 #endif
 };
 
 /**
  * Factory for UPrimitiveComponentToolTarget to be used by the target manager.
  */
-UCLASS(Transient)
-class INTERACTIVETOOLSFRAMEWORK_API UPrimitiveComponentToolTargetFactory : public UToolTargetFactory
+UCLASS(Transient, MinimalAPI)
+class UPrimitiveComponentToolTargetFactory : public UToolTargetFactory
 {
 	GENERATED_BODY()
 public:
-	virtual bool CanBuildTarget(UObject* SourceObject, const FToolTargetTypeRequirements& TargetTypeInfo) const override;
-	virtual UToolTarget* BuildTarget(UObject* SourceObject, const FToolTargetTypeRequirements& TargetTypeInfo) override;
+	INTERACTIVETOOLSFRAMEWORK_API virtual bool CanBuildTarget(UObject* SourceObject, const FToolTargetTypeRequirements& TargetTypeInfo) const override;
+	INTERACTIVETOOLSFRAMEWORK_API virtual UToolTarget* BuildTarget(UObject* SourceObject, const FToolTargetTypeRequirements& TargetTypeInfo) override;
 };
 
diff --git a/Engine/Source/Runtime/InteractiveToolsFramework/Public/ToolTargets/ToolTarget.h b/Engine/Source/Runtime/InteractiveToolsFramework/Public/ToolTargets/ToolTarget.h
--- a/Engine/Source/Runtime/InteractiveToolsFramework/Public/ToolTargets/ToolTarget.h
+++ b/Engine/Source/Runtime/InteractiveToolsFramework/Public/ToolTargets/ToolTarget.h
@@ -35,14 +35,14 @@ class UClass;
  * has the interfaces the tools needs. The tools cast the target to those interfaces to use
  * them.
  */
-UCLASS(Transient, Abstract)
-class INTERACTIVETOOLSFRAMEWORK_API UToolTarget : public UObject
+UCLASS(Transient, Abstract, MinimalAPI)
+class UToolTarget : public UObject
 {
 	GENERATED_BODY()
 public:
 
 	/** @return true if target is still valid. May become invalid for various reasons (eg Component was deleted out from under us) */
-	virtual bool IsValid() const PURE_VIRTUAL(UToolTarget::IsValid, return false;);
+	INTERACTIVETOOLSFRAMEWORK_API virtual bool IsValid() const PURE_VIRTUAL(UToolTarget::IsValid, return false;);
 };
 
 
@@ -50,7 +50,7 @@ public:
  * A structure used to specify the requirements of a tool for its target. E.g., a tool
  * may need a target that has interfaces x,y,z.
  */
-class INTERACTIVETOOLSFRAMEWORK_API FToolTargetTypeRequirements
+class FToolTargetTypeRequirements
 {
 public:
 
@@ -76,9 +76,9 @@ public:
 		return *this;
 	}
 
-	bool AreSatisfiedBy(UClass* Class) const;
+	INTERACTIVETOOLSFRAMEWORK_API bool AreSatisfiedBy(UClass* Class) const;
 
-	bool AreSatisfiedBy(UToolTarget* ToolTarget) const;
+	INTERACTIVETOOLSFRAMEWORK_API bool AreSatisfiedBy(UToolTarget* ToolTarget) const;
 };
 
 
@@ -87,14 +87,14 @@ public:
  * out of inputs without knowing anything about the inputs itself, as long as it
  * has a factory registered that is able to process the input.
  */
-UCLASS(Transient, Abstract)
-class INTERACTIVETOOLSFRAMEWORK_API UToolTargetFactory : public UObject
+UCLASS(Transient, Abstract, MinimalAPI)
+class UToolTargetFactory : public UObject
 {
 	GENERATED_BODY()
 
 public:
 
-	virtual bool CanBuildTarget(UObject* SourceObject, const FToolTargetTypeRequirements& TargetTypeInfo) const PURE_VIRTUAL(UToolTargetFactory::CanBuildTarget, return false;);
+	INTERACTIVETOOLSFRAMEWORK_API virtual bool CanBuildTarget(UObject* SourceObject, const FToolTargetTypeRequirements& TargetTypeInfo) const PURE_VIRTUAL(UToolTargetFactory::CanBuildTarget, return false;);
 
-	virtual UToolTarget* BuildTarget(UObject* SourceObject, const FToolTargetTypeRequirements& TargetTypeInfo) PURE_VIRTUAL(UToolTargetFactory::BuildTarget, return nullptr;);
+	INTERACTIVETOOLSFRAMEWORK_API virtual UToolTarget* BuildTarget(UObject* SourceObject, const FToolTargetTypeRequirements& TargetTypeInfo) PURE_VIRTUAL(UToolTargetFactory::BuildTarget, return nullptr;);
 };
diff --git a/Engine/Source/Runtime/Interchange/Core/Public/InterchangeFactoryBase.h b/Engine/Source/Runtime/Interchange/Core/Public/InterchangeFactoryBase.h
--- a/Engine/Source/Runtime/Interchange/Core/Public/InterchangeFactoryBase.h
+++ b/Engine/Source/Runtime/Interchange/Core/Public/InterchangeFactoryBase.h
@@ -49,8 +49,8 @@ enum class EInterchangeFactoryAssetType : uint8
  * 1. ImportSceneObject_GameThread - Create an actor in a level.
  */
 
-UCLASS(BlueprintType, Blueprintable, Abstract)
-class INTERCHANGECORE_API UInterchangeFactoryBase : public UObject
+UCLASS(BlueprintType, Blueprintable, Abstract, MinimalAPI)
+class UInterchangeFactoryBase : public UObject
 {
 	GENERATED_BODY()
 public:
diff --git a/Engine/Source/Runtime/Interchange/Core/Public/InterchangePipelineBase.h b/Engine/Source/Runtime/Interchange/Core/Public/InterchangePipelineBase.h
--- a/Engine/Source/Runtime/Interchange/Core/Public/InterchangePipelineBase.h
+++ b/Engine/Source/Runtime/Interchange/Core/Public/InterchangePipelineBase.h
@@ -101,8 +101,8 @@ struct FInterchangePipelinePropertyStates
 	FInterchangePipelinePropertyStatePerContext ReimportStates;
 };
 
-UCLASS(BlueprintType, Blueprintable, editinlinenew, Abstract)
-class INTERCHANGECORE_API UInterchangePipelineBase : public UObject
+UCLASS(BlueprintType, Blueprintable, editinlinenew, Abstract, MinimalAPI)
+class UInterchangePipelineBase : public UObject
 {
 	GENERATED_BODY()
 
@@ -110,7 +110,7 @@ public:
 
 	UE_DEPRECATED(5.2, "This function is replace by ScriptedExecutePipeline.")
 	UFUNCTION(BlueprintNativeEvent, BlueprintCallable, Category = "Interchange | Pipeline")
-	void ScriptedExecutePreImportPipeline(UInterchangeBaseNodeContainer* BaseNodeContainer, const TArray<UInterchangeSourceData*>& SourceDatas);
+	INTERCHANGECORE_API void ScriptedExecutePreImportPipeline(UInterchangeBaseNodeContainer* BaseNodeContainer, const TArray<UInterchangeSourceData*>& SourceDatas);
 
 	UE_DEPRECATED(5.2, "This internal public function is deprecated and not replace by any.")
 	void ScriptedExecutePreImportPipeline_Implementation(UInterchangeBaseNodeContainer* BaseNodeContainer, const TArray<UInterchangeSourceData*>& SourceDatas)
@@ -128,7 +128,7 @@ public:
 	 * @note - the FTaskPipeline is calling this function not the virtual one that is call by the default implementation.
 	 */
 	UFUNCTION(BlueprintNativeEvent, BlueprintCallable, Category = "Interchange | Pipeline")
-	void ScriptedExecutePipeline(UInterchangeBaseNodeContainer* BaseNodeContainer, const TArray<UInterchangeSourceData*>& SourceDatas);
+	INTERCHANGECORE_API void ScriptedExecutePipeline(UInterchangeBaseNodeContainer* BaseNodeContainer, const TArray<UInterchangeSourceData*>& SourceDatas);
 	/** The default implementation (call if the blueprint do not have any implementation) will call the virtual ExecutePipeline */
 	void ScriptedExecutePipeline_Implementation(UInterchangeBaseNodeContainer* BaseNodeContainer, const TArray<UInterchangeSourceData*>& SourceDatas)
 	{
@@ -141,7 +141,7 @@ public:
 	 * @note - the FTaskPreCompletion task is calling this function not the virtual one that is call by the default implementation.
 	 */
 	UFUNCTION(BlueprintNativeEvent, BlueprintCallable, Category = "Interchange | Pipeline")
-	void ScriptedExecutePostFactoryPipeline(const UInterchangeBaseNodeContainer* BaseNodeContainer, const FString& FactoryNodeKey, UObject* CreatedAsset, bool bIsAReimport);
+	INTERCHANGECORE_API void ScriptedExecutePostFactoryPipeline(const UInterchangeBaseNodeContainer* BaseNodeContainer, const FString& FactoryNodeKey, UObject* CreatedAsset, bool bIsAReimport);
 	/** The default implementation (call if the blueprint do not have any implementation) will call the virtual ExecutePostFactoryPipeline */
 	void ScriptedExecutePostFactoryPipeline_Implementation(const UInterchangeBaseNodeContainer* BaseNodeContainer, const FString& FactoryNodeKey, UObject* CreatedAsset, bool bIsAReimport)
 	{
@@ -155,7 +155,7 @@ public:
 	 * @note - the FTaskPipelinePostImport is calling this function not the virtual one that is call by the default implementation.
 	 */
 	UFUNCTION(BlueprintNativeEvent, BlueprintCallable, Category = "Interchange | Pipeline")
-	void ScriptedExecutePostImportPipeline(const UInterchangeBaseNodeContainer* BaseNodeContainer, const FString& FactoryNodeKey, UObject* CreatedAsset, bool bIsAReimport);
+	INTERCHANGECORE_API void ScriptedExecutePostImportPipeline(const UInterchangeBaseNodeContainer* BaseNodeContainer, const FString& FactoryNodeKey, UObject* CreatedAsset, bool bIsAReimport);
 	/** The default implementation (call if the blueprint do not have any implementation) will call the virtual ExecutePostImportPipeline */
 	void ScriptedExecutePostImportPipeline_Implementation(const UInterchangeBaseNodeContainer* BaseNodeContainer, const FString& FactoryNodeKey, UObject* CreatedAsset, bool bIsAReimport)
 	{
@@ -168,7 +168,7 @@ public:
 	 * the Interchange manager is calling this function not the virtual one that is call by the default implementation.
 	 */
 	UFUNCTION(BlueprintCallable, BlueprintNativeEvent, Category = "Interchange | Pipeline")
-	void ScriptedExecuteExportPipeline(UInterchangeBaseNodeContainer* BaseNodeContainer);
+	INTERCHANGECORE_API void ScriptedExecuteExportPipeline(UInterchangeBaseNodeContainer* BaseNodeContainer);
 
 	/** The default implementation (call if the blueprint do not have any implementation) will call the virtual ExecuteExportPipeline */
 	void ScriptedExecuteExportPipeline_Implementation(UInterchangeBaseNodeContainer* BaseNodeContainer)
@@ -181,7 +181,7 @@ public:
 	 * the Interchange manager is calling this function not the virtual one that is call by the default implementation.
 	 */
 	UFUNCTION(BlueprintNativeEvent, BlueprintCallable, Category = "Interchange | Pipeline")
-	bool ScriptedCanExecuteOnAnyThread(EInterchangePipelineTask PipelineTask);
+	INTERCHANGECORE_API bool ScriptedCanExecuteOnAnyThread(EInterchangePipelineTask PipelineTask);
 
 	/** The default implementation (call if the blueprint do not have any implementation) will call the virtual CanExecuteAsync */
 	bool ScriptedCanExecuteOnAnyThread_Implementation(EInterchangePipelineTask PipelineTask)
@@ -194,7 +194,7 @@ public:
 	 * the Interchange framework is calling this function not the virtual one that is called by the default implementation.
 	 */
 	UFUNCTION(BlueprintNativeEvent, BlueprintCallable, Category = "Interchange | Pipeline")
-	void ScriptedSetReimportSourceIndex(UClass* ReimportObjectClass, const int32 SourceFileIndex);
+	INTERCHANGECORE_API void ScriptedSetReimportSourceIndex(UClass* ReimportObjectClass, const int32 SourceFileIndex);
 
 	/** The default implementation (call if the blueprint do not have any implementation) will call the virtual SetReimportContentFromSourceIndex */
 	void ScriptedSetReimportSourceIndex_Implementation(UClass* ReimportObjectClass, const int32 SourceFileIndex)
@@ -214,9 +214,9 @@ public:
 		return true;
 	}
 
-	void LoadSettings(const FName PipelineStackName);
+	INTERCHANGECORE_API void LoadSettings(const FName PipelineStackName);
 
-	void SaveSettings(const FName PipelineStackName);
+	INTERCHANGECORE_API void SaveSettings(const FName PipelineStackName);
 
 	/**
 	 * This function is called before showing the import dialog it is not called doing a re-import.
@@ -240,11 +240,11 @@ public:
 	 * @Param ReimportType - Tell pipeline what re-import type the user want to achieve.
 	 * @Param ReimportAsset - This is an optional parameter which is set when re-importing an asset.
 	 */
-	virtual void AdjustSettingsForContext(EInterchangePipelineContext ReimportType, TObjectPtr<UObject> ReimportAsset);
-	virtual void AdjustSettingsFromCache();
+	INTERCHANGECORE_API virtual void AdjustSettingsForContext(EInterchangePipelineContext ReimportType, TObjectPtr<UObject> ReimportAsset);
+	INTERCHANGECORE_API virtual void AdjustSettingsFromCache();
 
 	/** Transfer the source pipeline adjust settings to this pipeline. */
-	void TransferAdjustSettings(UInterchangePipelineBase* SourcePipeline);
+	INTERCHANGECORE_API void TransferAdjustSettings(UInterchangePipelineBase* SourcePipeline);
 
 	/**
 	 * This function is used to add the given message object directly into the results for this operation.
@@ -271,31 +271,31 @@ public:
 	/**
 	 * Return a const property states pointer. Return nullptr if the property do not exist
 	 */
-	const FInterchangePipelinePropertyStates* GetPropertyStates(const FName PropertyPath) const;
+	INTERCHANGECORE_API const FInterchangePipelinePropertyStates* GetPropertyStates(const FName PropertyPath) const;
 
 	/**
 	 * Return a mutable property states pointer. Return nullptr if the property do not exist
 	 */
-	FInterchangePipelinePropertyStates* GetMutablePropertyStates(const FName PropertyPath);
+	INTERCHANGECORE_API FInterchangePipelinePropertyStates* GetMutablePropertyStates(const FName PropertyPath);
 
 	/**
 	 * Return true if the property has valid states, false if no states was set for the property
 	 */
 	UFUNCTION(BlueprintCallable, Category = "Interchange | Pipeline")
-	bool DoesPropertyStatesExist(const FName PropertyPath) const;
+	INTERCHANGECORE_API bool DoesPropertyStatesExist(const FName PropertyPath) const;
 
 	/**
 	 * Return a mutable property states Reference. Add the property states if it doesnt exist.
 	 */
 	UFUNCTION(BlueprintCallable, Category = "Interchange | Pipeline")
-	FInterchangePipelinePropertyStates& FindOrAddPropertyStates(const FName PropertyPath);
+	INTERCHANGECORE_API FInterchangePipelinePropertyStates& FindOrAddPropertyStates(const FName PropertyPath);
 
 	/**
 	 * Return the property name of the properties states map
 	 */
-	static FName GetPropertiesStatesPropertyName();
+	static INTERCHANGECORE_API FName GetPropertiesStatesPropertyName();
 
-	static FName GetResultsPropertyName();
+	static INTERCHANGECORE_API FName GetResultsPropertyName();
 
 	bool CanEditPropertiesStates() { return bAllowPropertyStatesEdition; }
 	bool IsReimportContext() { return bIsReimportContext; }
@@ -377,15 +377,15 @@ protected:
 	{
 	}
 
-	void LoadSettingsInternal(const FName PipelineStackName, const FString& ConfigFilename, TMap<FName, FInterchangePipelinePropertyStates>& ParentPropertiesStates);
+	INTERCHANGECORE_API void LoadSettingsInternal(const FName PipelineStackName, const FString& ConfigFilename, TMap<FName, FInterchangePipelinePropertyStates>& ParentPropertiesStates);
 
-	void SaveSettingsInternal(const FName PipelineStackName, const FString& ConfigFilename);
+	INTERCHANGECORE_API void SaveSettingsInternal(const FName PipelineStackName, const FString& ConfigFilename);
 
-	UInterchangePipelineBase* GetMostPipelineOuter() const;
+	INTERCHANGECORE_API UInterchangePipelineBase* GetMostPipelineOuter() const;
 
-	static void InternalToggleVisibilityPropertiesOfMetaDataValue(UInterchangePipelineBase* OuterMostPipeline, UInterchangePipelineBase* Pipeline, bool bDoTransientSubPipeline, const FString& MetaDataKey, const FString& MetaDataValue, const bool bVisibilityState);
-	static void HidePropertiesOfCategory(UInterchangePipelineBase* OuterMostPipeline, UInterchangePipelineBase* Pipeline, const FString& HideCategoryName, bool bDoTransientSubPipeline = false);
-	static void HidePropertiesOfSubCategory(UInterchangePipelineBase* OuterMostPipeline, UInterchangePipelineBase* Pipeline, const FString& HideSubCategoryName, bool bDoTransientSubPipeline = false);
+	static INTERCHANGECORE_API void InternalToggleVisibilityPropertiesOfMetaDataValue(UInterchangePipelineBase* OuterMostPipeline, UInterchangePipelineBase* Pipeline, bool bDoTransientSubPipeline, const FString& MetaDataKey, const FString& MetaDataValue, const bool bVisibilityState);
+	static INTERCHANGECORE_API void HidePropertiesOfCategory(UInterchangePipelineBase* OuterMostPipeline, UInterchangePipelineBase* Pipeline, const FString& HideCategoryName, bool bDoTransientSubPipeline = false);
+	static INTERCHANGECORE_API void HidePropertiesOfSubCategory(UInterchangePipelineBase* OuterMostPipeline, UInterchangePipelineBase* Pipeline, const FString& HideSubCategoryName, bool bDoTransientSubPipeline = false);
 	
 	/**
 	 * If true, the property editor for this pipeline instance will allow properties states edition.
diff --git a/Engine/Source/Runtime/Interchange/Core/Public/InterchangeResult.h b/Engine/Source/Runtime/Interchange/Core/Public/InterchangeResult.h
--- a/Engine/Source/Runtime/Interchange/Core/Public/InterchangeResult.h
+++ b/Engine/Source/Runtime/Interchange/Core/Public/InterchangeResult.h
@@ -28,8 +28,8 @@ enum class EInterchangeResultType
 /**
  * Base class for the result from an Interchange operation.
  */
-UCLASS(Experimental)
-class INTERCHANGECORE_API UInterchangeResult : public UObject
+UCLASS(Experimental, MinimalAPI)
+class UInterchangeResult : public UObject
 {
 	GENERATED_BODY()
 
@@ -48,12 +48,12 @@ public:
 	/**
 	 * Serialize this UInterchangeResult object to Json.
 	 */
-	FString ToJson();
+	INTERCHANGECORE_API FString ToJson();
 
 	/**
 	 * Create a new UInterchangeResult-derived object from a Json representation.
 	 */
-	static UInterchangeResult* FromJson(const FString& JsonString);
+	static INTERCHANGECORE_API UInterchangeResult* FromJson(const FString& JsonString);
 
 
 	UPROPERTY()
@@ -73,8 +73,8 @@ public:
 /**
  * Class representing a succesful result.
  */
-UCLASS()
-class INTERCHANGECORE_API UInterchangeResultSuccess : public UInterchangeResult
+UCLASS(MinimalAPI)
+class UInterchangeResultSuccess : public UInterchangeResult
 {
 	GENERATED_BODY()
 
@@ -88,8 +88,8 @@ public:
 /**
  * Base class representing a warning result.
  */
-UCLASS()
-class INTERCHANGECORE_API UInterchangeResultWarning : public UInterchangeResult
+UCLASS(MinimalAPI)
+class UInterchangeResultWarning : public UInterchangeResult
 {
 	GENERATED_BODY()
 
@@ -102,8 +102,8 @@ public:
 /**
  * Base class representing a error result.
  */
-UCLASS()
-class INTERCHANGECORE_API UInterchangeResultError : public UInterchangeResult
+UCLASS(MinimalAPI)
+class UInterchangeResultError : public UInterchangeResult
 {
 	GENERATED_BODY()
 
@@ -113,8 +113,8 @@ public:
 };
 
 
-UCLASS()
-class INTERCHANGECORE_API UInterchangeResultWarning_Generic : public UInterchangeResultWarning
+UCLASS(MinimalAPI)
+class UInterchangeResultWarning_Generic : public UInterchangeResultWarning
 {
 	GENERATED_BODY()
 
@@ -127,8 +127,8 @@ public:
 };
 
 
-UCLASS()
-class INTERCHANGECORE_API UInterchangeResultError_Generic : public UInterchangeResultError
+UCLASS(MinimalAPI)
+class UInterchangeResultError_Generic : public UInterchangeResultError
 {
 	GENERATED_BODY()
 
@@ -140,8 +140,8 @@ public:
 	FText Text;
 };
 
-UCLASS()
-class INTERCHANGECORE_API UInterchangeResultError_ReimportFail : public UInterchangeResultError
+UCLASS(MinimalAPI)
+class UInterchangeResultError_ReimportFail : public UInterchangeResultError
 {
 	GENERATED_BODY()
 
diff --git a/Engine/Source/Runtime/Interchange/Core/Public/InterchangeResultsContainer.h b/Engine/Source/Runtime/Interchange/Core/Public/InterchangeResultsContainer.h
--- a/Engine/Source/Runtime/Interchange/Core/Public/InterchangeResultsContainer.h
+++ b/Engine/Source/Runtime/Interchange/Core/Public/InterchangeResultsContainer.h
@@ -16,8 +16,8 @@
 #include "InterchangeResultsContainer.generated.h"
 
 
-UCLASS(Experimental)
-class INTERCHANGECORE_API UInterchangeResultsContainer : public UObject
+UCLASS(Experimental, MinimalAPI)
+class UInterchangeResultsContainer : public UObject
 {
 	GENERATED_BODY()
 
@@ -35,12 +35,12 @@ public:
 	/**
 	 * Empties the results container
 	 */
-	void Empty();
+	INTERCHANGECORE_API void Empty();
 
 	/**
 	 * Appends the given results container to this one
 	 */
-	void Append(UInterchangeResultsContainer* Other);
+	INTERCHANGECORE_API void Append(UInterchangeResultsContainer* Other);
 
 	/**
 	 * Creates a UInterchangeResult of the given type, adds it to the container and returns it.
@@ -66,7 +66,7 @@ public:
 	/**
 	 * Finalizes the container, prior to passing it to the UI display
 	 */
-	void Finalize();
+	INTERCHANGECORE_API void Finalize();
 
 	/**
 	 * Return the contained array (by value, for thread safety).
diff --git a/Engine/Source/Runtime/Interchange/Core/Public/InterchangeSourceData.h b/Engine/Source/Runtime/Interchange/Core/Public/InterchangeSourceData.h
--- a/Engine/Source/Runtime/Interchange/Core/Public/InterchangeSourceData.h
+++ b/Engine/Source/Runtime/Interchange/Core/Public/InterchangeSourceData.h
@@ -18,8 +18,8 @@
  * Stream (TODO)
  */
 
-UCLASS(BlueprintType, Blueprintable)
-class INTERCHANGECORE_API UInterchangeSourceData : public UObject
+UCLASS(BlueprintType, Blueprintable, MinimalAPI)
+class UInterchangeSourceData : public UObject
 {
 	 GENERATED_BODY()
 public:
diff --git a/Engine/Source/Runtime/Interchange/Core/Public/InterchangeTranslatorBase.h b/Engine/Source/Runtime/Interchange/Core/Public/InterchangeTranslatorBase.h
--- a/Engine/Source/Runtime/Interchange/Core/Public/InterchangeTranslatorBase.h
+++ b/Engine/Source/Runtime/Interchange/Core/Public/InterchangeTranslatorBase.h
@@ -40,23 +40,23 @@ enum class EInterchangeTranslatorAssetType : uint8
 };
 ENUM_CLASS_FLAGS(EInterchangeTranslatorAssetType);
 
-UCLASS(BlueprintType, Blueprintable, Abstract)
-class INTERCHANGECORE_API UInterchangeTranslatorBase : public UObject
+UCLASS(BlueprintType, Blueprintable, Abstract, MinimalAPI)
+class UInterchangeTranslatorBase : public UObject
 {
 	GENERATED_BODY()
 public:
 
 	/** return true if the translator can translate the given source data. */
-	virtual bool CanImportSourceData(const UInterchangeSourceData* InSourceData) const;
+	INTERCHANGECORE_API virtual bool CanImportSourceData(const UInterchangeSourceData* InSourceData) const;
 
 	/** Specifies the capabilities of a translator. */
 	virtual EInterchangeTranslatorType GetTranslatorType() const { return EInterchangeTranslatorType::Assets; }
 
 	bool DoesSupportAssetType(EInterchangeTranslatorAssetType AssetType) const { return EnumHasAllFlags(GetSupportedAssetTypes(), AssetType); }
-	virtual EInterchangeTranslatorAssetType GetSupportedAssetTypes() const PURE_VIRTUAL(UInterchangeTranslatorBase::GetSupportedAssetTypes, return EInterchangeTranslatorAssetType::None;);
+	INTERCHANGECORE_API virtual EInterchangeTranslatorAssetType GetSupportedAssetTypes() const PURE_VIRTUAL(UInterchangeTranslatorBase::GetSupportedAssetTypes, return EInterchangeTranslatorAssetType::None;);
 
 	/** List of formats supported by the translator. Each entry is of the form "ext;Description" where ext is the file extension. */
-	virtual TArray<FString> GetSupportedFormats() const PURE_VIRTUAL(UInterchangeTranslatorBase::GetSupportedExtensions, return TArray<FString>(););
+	INTERCHANGECORE_API virtual TArray<FString> GetSupportedFormats() const PURE_VIRTUAL(UInterchangeTranslatorBase::GetSupportedExtensions, return TArray<FString>(););
 
 	/**
 	 * Translate the associated source data into node(s) that are hold in the specified nodes container.
diff --git a/Engine/Source/Runtime/Interchange/Core/Public/InterchangeWriterBase.h b/Engine/Source/Runtime/Interchange/Core/Public/InterchangeWriterBase.h
--- a/Engine/Source/Runtime/Interchange/Core/Public/InterchangeWriterBase.h
+++ b/Engine/Source/Runtime/Interchange/Core/Public/InterchangeWriterBase.h
@@ -9,8 +9,8 @@
 
 #include "InterchangeWriterBase.generated.h"
 
-UCLASS(BlueprintType, Blueprintable, Abstract, Experimental)
-class INTERCHANGECORE_API UInterchangeWriterBase : public UObject
+UCLASS(BlueprintType, Blueprintable, Abstract, Experimental, MinimalAPI)
+class UInterchangeWriterBase : public UObject
 {
 	GENERATED_BODY()
 public:
diff --git a/Engine/Source/Runtime/Interchange/Core/Public/Nodes/InterchangeBaseNode.h b/Engine/Source/Runtime/Interchange/Core/Public/Nodes/InterchangeBaseNode.h
--- a/Engine/Source/Runtime/Interchange/Core/Public/Nodes/InterchangeBaseNode.h
+++ b/Engine/Source/Runtime/Interchange/Core/Public/Nodes/InterchangeBaseNode.h
@@ -150,16 +150,16 @@ namespace UE::Interchange
 		* Node that derive from UInterchangeBaseNode can also add a struct that derive from this one to add there static data
 		* @note: The static data are mainly for holding Attribute keys. All attributes that are always available for a node should be in this class or a derived class.
 		*/
-	struct INTERCHANGECORE_API FBaseNodeStaticData
+	struct FBaseNodeStaticData
 	{
-		static const FAttributeKey& UniqueIDKey();
-		static const FAttributeKey& DisplayLabelKey();
-		static const FAttributeKey& ParentIDKey();
-		static const FAttributeKey& IsEnabledKey();
-		static const FAttributeKey& TargetAssetIDsKey();
-		static const FAttributeKey& ClassTypeAttributeKey();
-		static const FAttributeKey& AssetNameKey();
-		static const FAttributeKey& NodeContainerTypeKey();
+		static INTERCHANGECORE_API const FAttributeKey& UniqueIDKey();
+		static INTERCHANGECORE_API const FAttributeKey& DisplayLabelKey();
+		static INTERCHANGECORE_API const FAttributeKey& ParentIDKey();
+		static INTERCHANGECORE_API const FAttributeKey& IsEnabledKey();
+		static INTERCHANGECORE_API const FAttributeKey& TargetAssetIDsKey();
+		static INTERCHANGECORE_API const FAttributeKey& ClassTypeAttributeKey();
+		static INTERCHANGECORE_API const FAttributeKey& AssetNameKey();
+		static INTERCHANGECORE_API const FAttributeKey& NodeContainerTypeKey();
 	};
 
 } //ns UE::Interchange
@@ -178,15 +178,15 @@ enum class EInterchangeNodeContainerType : uint8
  * See UE::Interchange::EAttributeTypes to know the supported template types
  * This is an abstract class. This is the base class of the interchange node graph format, all class in this format should derive from this class
  */
-UCLASS(BlueprintType)
-class INTERCHANGECORE_API UInterchangeBaseNode : public UObject
+UCLASS(BlueprintType, MinimalAPI)
+class UInterchangeBaseNode : public UObject
 {
 	GENERATED_BODY()
 
 public:
 	static constexpr const TCHAR* HierarchySeparator = TEXT("\\");
 
-	UInterchangeBaseNode();
+	INTERCHANGECORE_API UInterchangeBaseNode();
 
 	/**
 	 * Initialize the base data of the node
@@ -195,32 +195,32 @@ public:
 	 *
 	 */
 	UFUNCTION(BlueprintCallable, Category = "Interchange | Node")
-	void InitializeNode(const FString& UniqueID, const FString& DisplayLabel, const EInterchangeNodeContainerType NodeContainerType);
+	INTERCHANGECORE_API void InitializeNode(const FString& UniqueID, const FString& DisplayLabel, const EInterchangeNodeContainerType NodeContainerType);
 
 	/**
 	 * Return the node type name of the class, we use this when reporting error
 	 */
-	virtual FString GetTypeName() const;
+	INTERCHANGECORE_API virtual FString GetTypeName() const;
 
 	/**
 	 * Icon name, use to retrieve the brush when we display the node in any UI
 	 */
-	virtual FName GetIconName() const;
+	INTERCHANGECORE_API virtual FName GetIconName() const;
 
 	/**
 	 * UI that inspect node attribute call this to give a readable name to attribute key
 	 */
-	virtual FString GetKeyDisplayName(const UE::Interchange::FAttributeKey& NodeAttributeKey) const;
+	INTERCHANGECORE_API virtual FString GetKeyDisplayName(const UE::Interchange::FAttributeKey& NodeAttributeKey) const;
 
 	/**
 	 * UI that inspect node attribute call this to display or hide an attribute
 	 */
-	virtual bool ShouldHideAttribute(const UE::Interchange::FAttributeKey& NodeAttributeKey) const;
+	INTERCHANGECORE_API virtual bool ShouldHideAttribute(const UE::Interchange::FAttributeKey& NodeAttributeKey) const;
 
 	/**
 	 * UI that inspect node attribute call this to display the attribute under the returned category
 	 */
-	virtual FString GetAttributeCategory(const UE::Interchange::FAttributeKey& NodeAttributeKey) const;
+	INTERCHANGECORE_API virtual FString GetAttributeCategory(const UE::Interchange::FAttributeKey& NodeAttributeKey) const;
 
 	/**
 	 * Add an attribute to the node
@@ -246,14 +246,14 @@ public:
 	 * @param nodeAttributeKey - The key of the searched attribute
 	 *
 	 */
-	virtual bool HasAttribute(const UE::Interchange::FAttributeKey& NodeAttributeKey) const;
+	INTERCHANGECORE_API virtual bool HasAttribute(const UE::Interchange::FAttributeKey& NodeAttributeKey) const;
 
 	/**
 	 * This function return an attribute type for the specified Key. Return type None if the key is invalid
 	 *
 	 * @param NodeAttributeKey - The key of the attribute
 	 */
-	virtual UE::Interchange::EAttributeTypes GetAttributeType(const UE::Interchange::FAttributeKey& NodeAttributeKey) const;
+	INTERCHANGECORE_API virtual UE::Interchange::EAttributeTypes GetAttributeType(const UE::Interchange::FAttributeKey& NodeAttributeKey) const;
 
 	/**
 	 * This function return an  attribute handle for the specified Key.
@@ -270,109 +270,109 @@ public:
 	/**
 	 * Use this function to query all the node attribute keys
 	 */
-	void GetAttributeKeys(TArray<UE::Interchange::FAttributeKey>& AttributeKeys) const;
+	INTERCHANGECORE_API void GetAttributeKeys(TArray<UE::Interchange::FAttributeKey>& AttributeKeys) const;
 
 	/**
 	 * Remove any attribute from this node. Return false if we cannot remove it. If the attribute do not exist it will return true.
 	 */
 	UFUNCTION(BlueprintCallable, Category = "Interchange | Node")
-	bool RemoveAttribute(const FString& NodeAttributeKey);
+	INTERCHANGECORE_API bool RemoveAttribute(const FString& NodeAttributeKey);
 
 	/**
 	 * Add a boolean attribute to this node. Return false if the attribute do not exist or if we cannot add it
 	 */
 	UFUNCTION(BlueprintCallable, Category = "Interchange | Node")
-	bool AddBooleanAttribute(const FString& NodeAttributeKey, const bool& Value);
+	INTERCHANGECORE_API bool AddBooleanAttribute(const FString& NodeAttributeKey, const bool& Value);
 	
 	/**
 	 * Get a boolean attribute from this node. Return false if the attribute do not exist
 	 */
 	UFUNCTION(BlueprintCallable, Category = "Interchange | Node")
-	bool GetBooleanAttribute(const FString& NodeAttributeKey, bool& OutValue) const;
+	INTERCHANGECORE_API bool GetBooleanAttribute(const FString& NodeAttributeKey, bool& OutValue) const;
 
 	/**
 	 * Add a int32 attribute to this node. Return false if the attribute do not exist or if we cannot add it
 	 */
 	UFUNCTION(BlueprintCallable, Category = "Interchange | Node")
-	bool AddInt32Attribute(const FString& NodeAttributeKey, const int32& Value);
+	INTERCHANGECORE_API bool AddInt32Attribute(const FString& NodeAttributeKey, const int32& Value);
 	
 	/**
 	 * Get a int32 attribute from this node. Return false if the attribute do not exist
 	 */
 	UFUNCTION(BlueprintCallable, Category = "Interchange | Node")
-	bool GetInt32Attribute(const FString& NodeAttributeKey, int32& OutValue) const;
+	INTERCHANGECORE_API bool GetInt32Attribute(const FString& NodeAttributeKey, int32& OutValue) const;
 
 	/**
 	 * Add a float attribute to this node. Return false if the attribute do not exist or if we cannot add it
 	 */
 	UFUNCTION(BlueprintCallable, Category = "Interchange | Node")
-	bool AddFloatAttribute(const FString& NodeAttributeKey, const float& Value);
+	INTERCHANGECORE_API bool AddFloatAttribute(const FString& NodeAttributeKey, const float& Value);
 	
 	/**
 	 * Get a float attribute from this node. Return false if the attribute do not exist
 	 */
 	UFUNCTION(BlueprintCallable, Category = "Interchange | Node")
-	bool GetFloatAttribute(const FString& NodeAttributeKey, float& OutValue) const;
+	INTERCHANGECORE_API bool GetFloatAttribute(const FString& NodeAttributeKey, float& OutValue) const;
 
 	/**
 	 * Add a double attribute to this node. Return false if the attribute do not exist or if we cannot add it
 	 */
 	UFUNCTION(BlueprintCallable, Category = "Interchange | Node")
-	bool AddDoubleAttribute(const FString& NodeAttributeKey, const double& Value);
+	INTERCHANGECORE_API bool AddDoubleAttribute(const FString& NodeAttributeKey, const double& Value);
 	
 	/**
 	 * Get a double attribute from this node. Return false if the attribute do not exist
 	 */
 	UFUNCTION(BlueprintCallable, Category = "Interchange | Node")
-	bool GetDoubleAttribute(const FString& NodeAttributeKey, double& OutValue) const;
+	INTERCHANGECORE_API bool GetDoubleAttribute(const FString& NodeAttributeKey, double& OutValue) const;
 
 	/**
 	 * Add a string attribute to this node. Return false if the attribute do not exist or if we cannot add it
 	 */
 	UFUNCTION(BlueprintCallable, Category = "Interchange | Node")
-	bool AddStringAttribute(const FString& NodeAttributeKey, const FString& Value);
+	INTERCHANGECORE_API bool AddStringAttribute(const FString& NodeAttributeKey, const FString& Value);
 	
 	/**
 	 * Get a string attribute from this node. Return false if the attribute do not exist
 	 */
 	UFUNCTION(BlueprintCallable, Category = "Interchange | Node")
-	bool GetStringAttribute(const FString& NodeAttributeKey, FString& OutValue) const;
+	INTERCHANGECORE_API bool GetStringAttribute(const FString& NodeAttributeKey, FString& OutValue) const;
 
 	/**
 	 * Add a guid attribute to this node. Return false if the attribute do not exist or if we cannot add it
 	 */
 	UFUNCTION(BlueprintCallable, Category = "Interchange | Node")
-	bool AddGuidAttribute(const FString& NodeAttributeKey, const FGuid& Value);
+	INTERCHANGECORE_API bool AddGuidAttribute(const FString& NodeAttributeKey, const FGuid& Value);
 
 	/**
 	 * Get a guid attribute from this node. Return false if the attribute do not exist
 	 */
 	UFUNCTION(BlueprintCallable, Category = "Interchange | Node")
-	bool GetGuidAttribute(const FString& NodeAttributeKey, FGuid& OutValue) const;
+	INTERCHANGECORE_API bool GetGuidAttribute(const FString& NodeAttributeKey, FGuid& OutValue) const;
 
 	/**
 	 * Add a FLinearColor attribute to this node. Return false if the attribute do not exist or if we cannot add it
 	 */
 	UFUNCTION(BlueprintCallable, Category = "Interchange | Node")
-	bool AddLinearColorAttribute(const FString& NodeAttributeKey, const FLinearColor& Value);
+	INTERCHANGECORE_API bool AddLinearColorAttribute(const FString& NodeAttributeKey, const FLinearColor& Value);
 	
 	/**
 	 * Get a FLinearColor attribute from this node. Return false if the attribute do not exist
 	 */
 	UFUNCTION(BlueprintCallable, Category = "Interchange | Node")
-	bool GetLinearColorAttribute(const FString& NodeAttributeKey, FLinearColor& OutValue) const;
+	INTERCHANGECORE_API bool GetLinearColorAttribute(const FString& NodeAttributeKey, FLinearColor& OutValue) const;
 
 	/**
 	 * Add a Vector2 attribute to this node. Return false if the attribute do not exist or if we cannot add it
 	 */
 	UFUNCTION(BlueprintCallable, Category = "Interchange | Node")
-	bool AddVector2Attribute(const FString& NodeAttributeKey, const FVector2f& Value);
+	INTERCHANGECORE_API bool AddVector2Attribute(const FString& NodeAttributeKey, const FVector2f& Value);
 
 	/**
 	 * Get a Vector2 attribute from this node. Return false if the attribute do not exist
 	 */
 	UFUNCTION(BlueprintCallable, Category = "Interchange | Node")
-	bool GetVector2Attribute(const FString& NodeAttributeKey, FVector2f& OutValue) const;
+	INTERCHANGECORE_API bool GetVector2Attribute(const FString& NodeAttributeKey, FVector2f& OutValue) const;
 
 	template<typename AttributeType>
 	AttributeType GetAttributeChecked(const FString& NodeAttributeKey) const
@@ -400,63 +400,63 @@ public:
 	 * Return the unique id pass in the constructor.
 	 */
 	UFUNCTION(BlueprintCallable, Category = "Interchange | Node")
-	FString GetUniqueID() const;
+	INTERCHANGECORE_API FString GetUniqueID() const;
 
 	/**
 	 * Return the display label.
 	 */
 	UFUNCTION(BlueprintCallable, Category = "Interchange | Node")
-	FString GetDisplayLabel() const;
+	INTERCHANGECORE_API FString GetDisplayLabel() const;
 
 	/**
 	 * Change the display label.
 	 */
 	UFUNCTION(BlueprintCallable, Category = "Interchange | Node")
-	bool SetDisplayLabel(const FString& DisplayName);
+	INTERCHANGECORE_API bool SetDisplayLabel(const FString& DisplayName);
 
 	/**
 	 * Return the parent unique id. In case the attribute does not exist it will return InvalidNodeUid()
 	 */
 	UFUNCTION(BlueprintCallable, Category = "Interchange | Node")
-	FString GetParentUid() const;
+	INTERCHANGECORE_API FString GetParentUid() const;
 
 	/**
 	 * Set the parent unique id.
 	 */
 	UE_DEPRECATED(5.1, "This function should not be use - Use UInterchangeBaseNodeContainer::SetNodeParentUid function to create a parent hierarchy. This will ensure the container cache is always in sync")
 	UFUNCTION(BlueprintCallable, Category = "Interchange | Node")
-	bool SetParentUid(const FString& ParentUid);
+	INTERCHANGECORE_API bool SetParentUid(const FString& ParentUid);
 
 	/**
 	 * Get number of target assets relating to this object.
 	 */
 	UFUNCTION(BlueprintCallable, Category = "Interchange | Node")
-	int32 GetTargetNodeCount() const;
+	INTERCHANGECORE_API int32 GetTargetNodeCount() const;
 
 	/**
 	 * Get target assets relating to this object.
 	 */
 	UFUNCTION(BlueprintCallable, Category = "Interchange | Node")
-	void GetTargetNodeUids(TArray<FString>& OutTargetAssets) const;
+	INTERCHANGECORE_API void GetTargetNodeUids(TArray<FString>& OutTargetAssets) const;
 
 	/**
 	 * Add asset node UID relating to this object.
 	 */
 	UFUNCTION(BlueprintCallable, Category = "Interchange | Node")
-	bool AddTargetNodeUid(const FString& AssetUid) const;
+	INTERCHANGECORE_API bool AddTargetNodeUid(const FString& AssetUid) const;
 
 	/**
 	 * Remove asset node UID relating to this object.
 	 */
 	UFUNCTION(BlueprintCallable, Category = "Interchange | Node")
-	bool RemoveTargetNodeUid(const FString& AssetUid) const;
+	INTERCHANGECORE_API bool RemoveTargetNodeUid(const FString& AssetUid) const;
 
 	/**
 	 * IsEnable true mean that the node will be import/export, if false it will be discarded.
 	 * Return false if this node was disabled. Return true if the attribute is not there or if it was enabled.
 	 */
 	UFUNCTION(BlueprintCallable, Category = "Interchange | Node")
-	bool IsEnabled() const;
+	INTERCHANGECORE_API bool IsEnabled() const;
 
 	/**
 	 * Set the IsEnable attribute to determine if this node should be part of the import/export process
@@ -464,13 +464,13 @@ public:
 	 * @return true if it was able to set the attribute, false otherwise.
 	 */
 	UFUNCTION(BlueprintCallable, Category = "Interchange | Node")
-	bool SetEnabled(const bool bIsEnabled);
+	INTERCHANGECORE_API bool SetEnabled(const bool bIsEnabled);
 
 	/**
 	 * Return the node container type which define the purpose of the node (Factory node, translated scene node or translated asset node).
 	 */
 	UFUNCTION(BlueprintCallable, Category = "Interchange | Node")
-	EInterchangeNodeContainerType GetNodeContainerType() const;
+	INTERCHANGECORE_API EInterchangeNodeContainerType GetNodeContainerType() const;
 
 	/**
 	 * Return a FGuid build from the FSHA1 of all the attribute data contain in the node.
@@ -479,34 +479,34 @@ public:
 	 * the order we add the attributes.
 	 * This function interface is pure virtual
 	 */
-	virtual FGuid GetHash() const;
+	INTERCHANGECORE_API virtual FGuid GetHash() const;
 
 	/**
 	 * Optional, Any node that can import/export an asset should set the proper name we will give to the asset.
 	 * If the attribute was never set, it will return GetDisplayLabel.
 	 */
 	UFUNCTION(BlueprintCallable, Category = "Interchange | Node")
-	virtual FString GetAssetName() const;
+	INTERCHANGECORE_API virtual FString GetAssetName() const;
 
 	/**
 	 * Set the name we want for the imported asset this node represent. The asset factory will call GetAssetName()
 	 */
 	UFUNCTION(BlueprintCallable, Category = "Interchange | Node")
-	virtual bool SetAssetName(const FString& AssetName);
+	INTERCHANGECORE_API virtual bool SetAssetName(const FString& AssetName);
 
 	/** Return the invalid unique ID */
-	static FString InvalidNodeUid();
+	static INTERCHANGECORE_API FString InvalidNodeUid();
 
 	/**
 	 * Serialize the node, by default only the attribute storage is serialize for a node
 	 */
-	virtual void Serialize(FArchive& Ar) override;
+	INTERCHANGECORE_API virtual void Serialize(FArchive& Ar) override;
 
-	static void CompareNodeStorage(const UInterchangeBaseNode* NodeA, const UInterchangeBaseNode* NodeB, TArray<UE::Interchange::FAttributeKey>& RemovedAttributes, TArray<UE::Interchange::FAttributeKey>& AddedAttributes, TArray<UE::Interchange::FAttributeKey>& ModifiedAttributes);
+	static INTERCHANGECORE_API void CompareNodeStorage(const UInterchangeBaseNode* NodeA, const UInterchangeBaseNode* NodeB, TArray<UE::Interchange::FAttributeKey>& RemovedAttributes, TArray<UE::Interchange::FAttributeKey>& AddedAttributes, TArray<UE::Interchange::FAttributeKey>& ModifiedAttributes);
 	
-	static void CopyStorageAttributes(const UInterchangeBaseNode* SourceNode, UInterchangeBaseNode* DestinationNode, TArray<UE::Interchange::FAttributeKey>& AttributeKeys);
+	static INTERCHANGECORE_API void CopyStorageAttributes(const UInterchangeBaseNode* SourceNode, UInterchangeBaseNode* DestinationNode, TArray<UE::Interchange::FAttributeKey>& AttributeKeys);
 
-	static void CopyStorage(const UInterchangeBaseNode* SourceNode, UInterchangeBaseNode* DestinationNode);
+	static INTERCHANGECORE_API void CopyStorage(const UInterchangeBaseNode* SourceNode, UInterchangeBaseNode* DestinationNode);
 	
 	virtual void AppendAssetRegistryTags(TArray<FAssetRegistryTag>& OutTags) const
 	{
diff --git a/Engine/Source/Runtime/Interchange/Core/Public/Nodes/InterchangeBaseNodeContainer.h b/Engine/Source/Runtime/Interchange/Core/Public/Nodes/InterchangeBaseNodeContainer.h
--- a/Engine/Source/Runtime/Interchange/Core/Public/Nodes/InterchangeBaseNodeContainer.h
+++ b/Engine/Source/Runtime/Interchange/Core/Public/Nodes/InterchangeBaseNodeContainer.h
@@ -33,12 +33,12 @@ template <typename T> struct TObjectPtr;
  *
  * Translators are filling this container and the Import/Export managers are reading it to execute the import/export process
  */
- UCLASS(BlueprintType)
-class INTERCHANGECORE_API UInterchangeBaseNodeContainer : public UObject
+ UCLASS(BlueprintType, MinimalAPI)
+class UInterchangeBaseNodeContainer : public UObject
 {
 	GENERATED_BODY()
 public:
-	UInterchangeBaseNodeContainer();
+	INTERCHANGECORE_API UInterchangeBaseNodeContainer();
 
 	/**
 	 * Add a node in the container, the node will be add into a TMap.
@@ -48,17 +48,17 @@ public:
 	 *
 	 */
 	UFUNCTION(BlueprintCallable, Category = "Interchange | Node Container")
-	FString AddNode(UInterchangeBaseNode* Node);
+	INTERCHANGECORE_API FString AddNode(UInterchangeBaseNode* Node);
 
 	UFUNCTION(BlueprintCallable, Category = "Interchange | Node Container")
-	void ReplaceNode(const FString& NodeUniqueID, UInterchangeFactoryBaseNode* NewNode);
+	INTERCHANGECORE_API void ReplaceNode(const FString& NodeUniqueID, UInterchangeFactoryBaseNode* NewNode);
 
 	/** Return true if the node unique ID exist in the container */
 	UFUNCTION(BlueprintCallable, Category = "Interchange | Node Container")
-	bool IsNodeUidValid(const FString& NodeUniqueID) const;
+	INTERCHANGECORE_API bool IsNodeUidValid(const FString& NodeUniqueID) const;
 
 	/** Unordered iteration of the all nodes */
-	void IterateNodes(TFunctionRef<void(const FString&, UInterchangeBaseNode*)> IterationLambda) const;
+	INTERCHANGECORE_API void IterateNodes(TFunctionRef<void(const FString&, UInterchangeBaseNode*)> IterationLambda) const;
 
 	template <typename T>
 	void IterateNodesOfType(TFunctionRef<void(const FString&, T*)> IterationLambda) const
@@ -117,7 +117,7 @@ public:
 	}
 
 	/** Unordered iteration of the all nodes, but I can be stop early by returning true */
-	void BreakableIterateNodes(TFunctionRef<bool(const FString&, UInterchangeBaseNode*)> IterationLambda) const;
+	INTERCHANGECORE_API void BreakableIterateNodes(TFunctionRef<bool(const FString&, UInterchangeBaseNode*)> IterationLambda) const;
 
 	template <typename T>
 	void BreakableIterateNodesOfType(TFunctionRef<bool(const FString&, T*)> IterationLambda) const
@@ -136,57 +136,57 @@ public:
 
 	/** Return all nodes that do not have any parent */
 	UFUNCTION(BlueprintCallable, Category = "Interchange | Node Container")
-	void GetRoots(TArray<FString>& RootNodes) const;
+	INTERCHANGECORE_API void GetRoots(TArray<FString>& RootNodes) const;
 
 	/** Return all nodes that are of the ClassNode type*/
 	UFUNCTION(BlueprintCallable, Category = "Interchange | Node Container")
-	void GetNodes(const UClass* ClassNode, TArray<FString>& OutNodes) const;
+	INTERCHANGECORE_API void GetNodes(const UClass* ClassNode, TArray<FString>& OutNodes) const;
 
 	/** Get a node pointer. Once added to the container, nodes are considered const */
 	UFUNCTION(BlueprintCallable, Category = "Interchange | Node Container")
-	const UInterchangeBaseNode* GetNode(const FString& NodeUniqueID) const;
+	INTERCHANGECORE_API const UInterchangeBaseNode* GetNode(const FString& NodeUniqueID) const;
 
 	/** Get a factory node pointer */
 	UFUNCTION(BlueprintCallable, Category = "Interchange | Node Container")
-	UInterchangeFactoryBaseNode* GetFactoryNode(const FString& NodeUniqueID) const;
+	INTERCHANGECORE_API UInterchangeFactoryBaseNode* GetFactoryNode(const FString& NodeUniqueID) const;
 
 	/** Set node ParentUid */
 	UFUNCTION(BlueprintCallable, Category = "Interchange | Node Container")
-	bool SetNodeParentUid(const FString& NodeUniqueID, const FString& NewParentNodeUid);
+	INTERCHANGECORE_API bool SetNodeParentUid(const FString& NodeUniqueID, const FString& NewParentNodeUid);
 
 	/** Get the node children count */
 	UFUNCTION(BlueprintCallable, Category = "Interchange | Node Container")
-	int32 GetNodeChildrenCount(const FString& NodeUniqueID) const;
+	INTERCHANGECORE_API int32 GetNodeChildrenCount(const FString& NodeUniqueID) const;
 
 	/** Get all children Uid */
 	UFUNCTION(BlueprintCallable, Category = "Interchange | Node Container")
-	TArray<FString> GetNodeChildrenUids(const FString& NodeUniqueID) const;
+	INTERCHANGECORE_API TArray<FString> GetNodeChildrenUids(const FString& NodeUniqueID) const;
 
 	/** Get the node nth const children */
 	UFUNCTION(BlueprintCallable, Category = "Interchange | Node Container")
-	UInterchangeBaseNode* GetNodeChildren(const FString& NodeUniqueID, int32 ChildIndex);
+	INTERCHANGECORE_API UInterchangeBaseNode* GetNodeChildren(const FString& NodeUniqueID, int32 ChildIndex);
 
 	/** Get the node nth const children. Const version */
-	const UInterchangeBaseNode* GetNodeChildren(const FString& NodeUniqueID, int32 ChildIndex) const;
+	INTERCHANGECORE_API const UInterchangeBaseNode* GetNodeChildren(const FString& NodeUniqueID, int32 ChildIndex) const;
 
 	/**
 	 * This function serialize the node container and all node sub-objects point by it.
 	 * Out of process translator like fbx will dump a file containing this data and the editor
 	 * will read the file and regenerate the container from the saved data.
 	 */
-	void SerializeNodeContainerData(FArchive& Ar);
+	INTERCHANGECORE_API void SerializeNodeContainerData(FArchive& Ar);
 
 	/* Serialize the node container into the specified file.*/
 	UFUNCTION(BlueprintCallable, Category = "Interchange | Node Container")
-	void SaveToFile(const FString& Filename);
+	INTERCHANGECORE_API void SaveToFile(const FString& Filename);
 
 	/* Serialize the node container from the specified file.*/
 	UFUNCTION(BlueprintCallable, Category = "Interchange | Node Container")
-	void LoadFromFile(const FString& Filename);
+	INTERCHANGECORE_API void LoadFromFile(const FString& Filename);
 
 	/* Fill the children uids cache to optimize the GetNodeChildrenUids call.*/
 	UFUNCTION(BlueprintCallable, Category = "Interchange | Node Container")
-	void ComputeChildrenCache();
+	INTERCHANGECORE_API void ComputeChildrenCache();
 
 	/* Reset the children uids cache. */
 	UFUNCTION(BlueprintCallable, Category = "Interchange | Node Container")
@@ -197,7 +197,7 @@ public:
 
 private:
 
-	UInterchangeBaseNode* GetNodeChildrenInternal(const FString& NodeUniqueID, int32 ChildIndex);
+	INTERCHANGECORE_API UInterchangeBaseNode* GetNodeChildrenInternal(const FString& NodeUniqueID, int32 ChildIndex);
 
 	/** Flat List of the nodes. Since the nodes are variable size, we store a pointer. */
 	UPROPERTY(VisibleAnywhere, Category = "Interchange | Node Container")
diff --git a/Engine/Source/Runtime/Interchange/Core/Public/Nodes/InterchangeFactoryBaseNode.h b/Engine/Source/Runtime/Interchange/Core/Public/Nodes/InterchangeFactoryBaseNode.h
--- a/Engine/Source/Runtime/Interchange/Core/Public/Nodes/InterchangeFactoryBaseNode.h
+++ b/Engine/Source/Runtime/Interchange/Core/Public/Nodes/InterchangeFactoryBaseNode.h
@@ -130,11 +130,11 @@ namespace UE::Interchange
 	DECLARE_DELEGATE_RetVal_OneParam(bool, FApplyAttributeToAsset, UObject*);
 	DECLARE_DELEGATE_RetVal_OneParam(bool, FFillAttributeToAsset, UObject*);
 
-	struct INTERCHANGECORE_API FFactoryBaseNodeStaticData : public FBaseNodeStaticData
+	struct FFactoryBaseNodeStaticData : public FBaseNodeStaticData
 	{
-		static const FString& FactoryDependenciesBaseKey();
-		static const FAttributeKey& ReimportStrategyFlagsKey();
-		static const FAttributeKey& SkipNodeImportKey();
+		static INTERCHANGECORE_API const FString& FactoryDependenciesBaseKey();
+		static INTERCHANGECORE_API const FAttributeKey& ReimportStrategyFlagsKey();
+		static INTERCHANGECORE_API const FAttributeKey& SkipNodeImportKey();
 	};
 
 } // namespace UE::Interchange
@@ -155,30 +155,30 @@ enum class EReimportStrategyFlags : uint8
  * See UE::Interchange::EAttributeTypes to know the supported template types
  * This is an abstract class. This is the base class of the interchange node graph format, all class in this format should derive from this class
  */
-UCLASS(BlueprintType)
-class INTERCHANGECORE_API UInterchangeFactoryBaseNode : public UInterchangeBaseNode
+UCLASS(BlueprintType, MinimalAPI)
+class UInterchangeFactoryBaseNode : public UInterchangeBaseNode
 {
 	GENERATED_BODY()
 
 public:
-	UInterchangeFactoryBaseNode();
+	INTERCHANGECORE_API UInterchangeFactoryBaseNode();
 
-	virtual FString GetKeyDisplayName(const UE::Interchange::FAttributeKey& NodeAttributeKey) const override;
-	virtual FString GetAttributeCategory(const UE::Interchange::FAttributeKey& NodeAttributeKey) const override;
+	INTERCHANGECORE_API virtual FString GetKeyDisplayName(const UE::Interchange::FAttributeKey& NodeAttributeKey) const override;
+	INTERCHANGECORE_API virtual FString GetAttributeCategory(const UE::Interchange::FAttributeKey& NodeAttributeKey) const override;
 
 	/**
 	 * Return the reimport strategy flags.
 	 *
 	 */
 	UFUNCTION(BlueprintCallable, Category = "Interchange | Node")
-	EReimportStrategyFlags GetReimportStrategyFlags() const;
+	INTERCHANGECORE_API EReimportStrategyFlags GetReimportStrategyFlags() const;
 
 	/**
 	 * Change the reimport strategy flags.
 	 *
 	 */
 	UFUNCTION(BlueprintCallable, Category = "Interchange | Node")
-	bool SetReimportStrategyFlags(const EReimportStrategyFlags& ReimportStrategyFlags);
+	INTERCHANGECORE_API bool SetReimportStrategyFlags(const EReimportStrategyFlags& ReimportStrategyFlags);
 
 	/**
 	 * Return true if this node should skip the factory import process, false otherwise.
@@ -188,80 +188,80 @@ public:
 	 * - An asset can already exist and is being imported by another concurrent import task (user importing multiple file in same time in the same content folder)
 	 */
 	UFUNCTION(BlueprintCallable, Category = "Interchange | Node")
-	bool ShouldSkipNodeImport() const;
+	INTERCHANGECORE_API bool ShouldSkipNodeImport() const;
 
 	/**
 	 * Add the skip node attribute. Use this function to cancel the creation of the unreal asset, see ShouldSkipNodeImport for more documentation.
 	 */
 	UFUNCTION(BlueprintCallable, Category = "Interchange | Node")
-	bool SetSkipNodeImport();
+	INTERCHANGECORE_API bool SetSkipNodeImport();
 
 	/**
 	 * Remove the skip node attribute. See ShouldSkipNodeImport for more documentation.
 	 */
 	UFUNCTION(BlueprintCallable, Category = "Interchange | Node")
-	bool UnsetSkipNodeImport();
+	INTERCHANGECORE_API bool UnsetSkipNodeImport();
 
 	/**
 	 * Return the UClass of the object we represent so we can find factory/writer
 	 */
-	virtual class UClass* GetObjectClass() const;
+	INTERCHANGECORE_API virtual class UClass* GetObjectClass() const;
 
 	/**
 	 * Return the custom sub-path under PackageBasePath, where the assets will be created.
 	 */
 	UFUNCTION(BlueprintCallable, Category = "Interchange | Node")
-	bool GetCustomSubPath(FString& AttributeValue) const;
+	INTERCHANGECORE_API bool GetCustomSubPath(FString& AttributeValue) const;
 
 	/**
 	 * Set the custom sub-path under PackageBasePath where the assets will be created.
 	 */
 	UFUNCTION(BlueprintCallable, Category = "Interchange | Node")
-	bool SetCustomSubPath(const FString& AttributeValue);
+	INTERCHANGECORE_API bool SetCustomSubPath(const FString& AttributeValue);
 
 	/**
 	 * This function allow to retrieve the number of factory dependencies for this object.
 	 */
 	UFUNCTION(BlueprintCallable, Category = "Interchange | Node")
-	int32 GetFactoryDependenciesCount() const;
+	INTERCHANGECORE_API int32 GetFactoryDependenciesCount() const;
 
 	/**
 	 * This function allow to retrieve the dependency for this object.
 	 */
 	UFUNCTION(BlueprintCallable, Category = "Interchange | Node")
-	void GetFactoryDependencies(TArray<FString>& OutDependencies ) const;
+	INTERCHANGECORE_API void GetFactoryDependencies(TArray<FString>& OutDependencies ) const;
 
 	/**
 	 * This function allow to retrieve one dependency for this object.
 	 */
 	UFUNCTION(BlueprintCallable, Category = "Interchange | Node")
-	void GetFactoryDependency(const int32 Index, FString& OutDependency) const;
+	INTERCHANGECORE_API void GetFactoryDependency(const int32 Index, FString& OutDependency) const;
 
 	/**
 	 * Add one dependency to this object.
 	 */
 	UFUNCTION(BlueprintCallable, Category = "Interchange | Node")
-	bool AddFactoryDependencyUid(const FString& DependencyUid);
+	INTERCHANGECORE_API bool AddFactoryDependencyUid(const FString& DependencyUid);
 
 	/**
 	 * Remove one dependency from this object.
 	 */
 	UFUNCTION(BlueprintCallable, Category = "Interchange | Node")
-	bool RemoveFactoryDependencyUid(const FString& DependencyUid);
+	INTERCHANGECORE_API bool RemoveFactoryDependencyUid(const FString& DependencyUid);
 
 	/**
 	 * Return the custom ReferenceObject. The UObject this factory node has created.
 	 */
 	UFUNCTION(BlueprintCallable, Category = "Interchange | Node")
-	bool GetCustomReferenceObject(FSoftObjectPath& AttributeValue) const;
+	INTERCHANGECORE_API bool GetCustomReferenceObject(FSoftObjectPath& AttributeValue) const;
 
 	/**
 	 * Set the custom ReferenceObject. The UObject this factory node has created.
 	 */
 	UFUNCTION(BlueprintCallable, Category = "Interchange | Node")
-	bool SetCustomReferenceObject(const FSoftObjectPath& AttributeValue);
+	INTERCHANGECORE_API bool SetCustomReferenceObject(const FSoftObjectPath& AttributeValue);
 
-	static FString BuildFactoryNodeUid(const FString& TranslatedNodeUid);
+	static INTERCHANGECORE_API FString BuildFactoryNodeUid(const FString& TranslatedNodeUid);
 
 	/**
 	 * Adds the delegates that will read and write the attribute value to a UObject.
@@ -308,9 +308,9 @@ public:
 	 * get the delegate execute so it apply the attribute to the UObject property.
 	 * See the macros IMPLEMENT_NODE_ATTRIBUTE_SETTER at the top of the file to know how delegates are setup for property.
 	 */
-	void ApplyAllCustomAttributeToObject(UObject* Object) const;
+	INTERCHANGECORE_API void ApplyAllCustomAttributeToObject(UObject* Object) const;
 
-	void FillAllCustomAttributeFromObject(UObject* Object) const;
+	INTERCHANGECORE_API void FillAllCustomAttributeFromObject(UObject* Object) const;
 
 	/**
 	 * Copies all the custom attributes from SourceNode to this node and
@@ -318,7 +318,7 @@ public:
 	 * @param SourceNode		The source factory node to copy from.
 	 * @param Object			The object to fill custom attributes' values from, if applicable.
 	 */
-	static UInterchangeFactoryBaseNode* DuplicateWithObject(const UInterchangeFactoryBaseNode* SourceNode, UObject* Object);
+	static INTERCHANGECORE_API UInterchangeFactoryBaseNode* DuplicateWithObject(const UInterchangeFactoryBaseNode* SourceNode, UObject* Object);
 
 protected:
 	/**
