[Engine]

diff --git a/Engine/Source/Runtime/Interchange/Core/Public/Nodes/InterchangeSourceNode.h b/Engine/Source/Runtime/Interchange/Core/Public/Nodes/InterchangeSourceNode.h
--- a/Engine/Source/Runtime/Interchange/Core/Public/Nodes/InterchangeSourceNode.h
+++ b/Engine/Source/Runtime/Interchange/Core/Public/Nodes/InterchangeSourceNode.h
@@ -22,8 +22,8 @@ struct FFrame;
 /**
  * This class allow a translator to add general source data that describe the whole source. Pipeline can use this information.
  */
-UCLASS(BlueprintType, Experimental)
-class INTERCHANGECORE_API UInterchangeSourceNode : public UInterchangeBaseNode
+UCLASS(BlueprintType, Experimental, MinimalAPI)
+class UInterchangeSourceNode : public UInterchangeBaseNode
 {
 	GENERATED_BODY()
 
@@ -37,74 +37,74 @@ public:
 	 *
 	 */
 	UFUNCTION(BlueprintCallable, Category = "Interchange | Node | source")
-	void InitializeSourceNode(const FString& UniqueID, const FString& DisplayLabel);
+	INTERCHANGECORE_API void InitializeSourceNode(const FString& UniqueID, const FString& DisplayLabel);
 
 	/**
 	 * Return the node type name of the class, we use this when reporting error
 	 */
-	virtual FString GetTypeName() const override;
+	INTERCHANGECORE_API virtual FString GetTypeName() const override;
 
 	/* The translators that want to modify the common data should ensure they create the unique common pipeline node. */
-	static UInterchangeSourceNode* FindOrCreateUniqueInstance(UInterchangeBaseNodeContainer* NodeContainer);
+	static INTERCHANGECORE_API UInterchangeSourceNode* FindOrCreateUniqueInstance(UInterchangeBaseNodeContainer* NodeContainer);
 
 	/* If the unique instance doesn't exist it will return nullptr. This function should be use by the pipelines, to avoid creating a node. */
-	static const UInterchangeSourceNode* GetUniqueInstance(const UInterchangeBaseNodeContainer* NodeContainer);
+	static INTERCHANGECORE_API const UInterchangeSourceNode* GetUniqueInstance(const UInterchangeBaseNodeContainer* NodeContainer);
 
 	/** Query the source frame rate numerator. */
 	UFUNCTION(BlueprintCallable, Category = "Interchange | Node | Source")
-	bool GetCustomSourceFrameRateNumerator(int32& AttributeValue) const;
+	INTERCHANGECORE_API bool GetCustomSourceFrameRateNumerator(int32& AttributeValue) const;
 
 	/** Store the source frame rate numerator. */
 	UFUNCTION(BlueprintCallable, Category = "Interchange | Node | Source")
-	bool SetCustomSourceFrameRateNumerator(const int32& AttributeValue);
+	INTERCHANGECORE_API bool SetCustomSourceFrameRateNumerator(const int32& AttributeValue);
 
 	/** Query the source frame rate denominator. */
 	UFUNCTION(BlueprintCallable, Category = "Interchange | Node | Source")
-	bool GetCustomSourceFrameRateDenominator(int32& AttributeValue) const;
+	INTERCHANGECORE_API bool GetCustomSourceFrameRateDenominator(int32& AttributeValue) const;
 
 	/** Store the source frame rate denominator. */
 	UFUNCTION(BlueprintCallable, Category = "Interchange | Node | Source")
-	bool SetCustomSourceFrameRateDenominator(const int32& AttributeValue);
+	INTERCHANGECORE_API bool SetCustomSourceFrameRateDenominator(const int32& AttributeValue);
 
 	/** Query the source time line start. */
 	UFUNCTION(BlueprintCallable, Category = "Interchange | Node | Source")
-	bool GetCustomSourceTimelineStart(double& AttributeValue) const;
+	INTERCHANGECORE_API bool GetCustomSourceTimelineStart(double& AttributeValue) const;
 
 	/** Store the source time line start. */
 	UFUNCTION(BlueprintCallable, Category = "Interchange | Node | Source")
-	bool SetCustomSourceTimelineStart(const double& AttributeValue);
+	INTERCHANGECORE_API bool SetCustomSourceTimelineStart(const double& AttributeValue);
 
 	/** Query the source time line end. */
 	UFUNCTION(BlueprintCallable, Category = "Interchange | Node | Source")
-	bool GetCustomSourceTimelineEnd(double& AttributeValue) const;
+	INTERCHANGECORE_API bool GetCustomSourceTimelineEnd(double& AttributeValue) const;
 
 	/** Store the source time line end. */
 	UFUNCTION(BlueprintCallable, Category = "Interchange | Node | Source")
-	bool SetCustomSourceTimelineEnd(const double& AttributeValue);
+	INTERCHANGECORE_API bool SetCustomSourceTimelineEnd(const double& AttributeValue);
 
 	/** Query the source animated time start. */
 	UFUNCTION(BlueprintCallable, Category = "Interchange | Node | Source")
-	bool GetCustomAnimatedTimeStart(double& AttributeValue) const;
+	INTERCHANGECORE_API bool GetCustomAnimatedTimeStart(double& AttributeValue) const;
 
 	/** Store the source animated time start. */
 	UFUNCTION(BlueprintCallable, Category = "Interchange | Node | Source")
-	bool SetCustomAnimatedTimeStart(const double& AttributeValue);
+	INTERCHANGECORE_API bool SetCustomAnimatedTimeStart(const double& AttributeValue);
 
 	/** Query the source animated time end. */
 	UFUNCTION(BlueprintCallable, Category = "Interchange | Node | Source")
-	bool GetCustomAnimatedTimeEnd(double& AttributeValue) const;
+	INTERCHANGECORE_API bool GetCustomAnimatedTimeEnd(double& AttributeValue) const;
 
 	/** Store the source animated time end. */
 	UFUNCTION(BlueprintCallable, Category = "Interchange | Node | Source")
-	bool SetCustomAnimatedTimeEnd(const double& AttributeValue);
+	INTERCHANGECORE_API bool SetCustomAnimatedTimeEnd(const double& AttributeValue);
 
 	/** Query the source animated time end. */
 	UFUNCTION(BlueprintCallable, Category = "Interchange | Node | Source")
-	bool GetCustomImportUnusedMaterial(bool& AttributeValue) const;
+	INTERCHANGECORE_API bool GetCustomImportUnusedMaterial(bool& AttributeValue) const;
 
 	/** Store the source animated time end. */
 	UFUNCTION(BlueprintCallable, Category = "Interchange | Node | Source")
-	bool SetCustomImportUnusedMaterial(const bool& AttributeValue);
+	INTERCHANGECORE_API bool SetCustomImportUnusedMaterial(const bool& AttributeValue);
 
 private:
 	const UE::Interchange::FAttributeKey Macro_CustomSourceFrameRateNumeratorKey = UE::Interchange::FAttributeKey(TEXT("SourceFrameRateNumerator"));
diff --git a/Engine/Source/Runtime/Interchange/Core/Public/Nodes/InterchangeUserDefinedAttribute.h b/Engine/Source/Runtime/Interchange/Core/Public/Nodes/InterchangeUserDefinedAttribute.h
--- a/Engine/Source/Runtime/Interchange/Core/Public/Nodes/InterchangeUserDefinedAttribute.h
+++ b/Engine/Source/Runtime/Interchange/Core/Public/Nodes/InterchangeUserDefinedAttribute.h
@@ -37,8 +37,8 @@ struct FInterchangeUserDefinedAttributeInfo
  * Any user defined attribute have: name, value and a optional AnimationPayloadKey (FRichCurve which is a float curve).
  * Value type must be supported by the UE::Interchange::EAttributeTypes enumeration.
  */
-UCLASS(BlueprintType, Experimental)
-class INTERCHANGECORE_API UInterchangeUserDefinedAttributesAPI : public UObject
+UCLASS(BlueprintType, Experimental, MinimalAPI)
+class UInterchangeUserDefinedAttributesAPI : public UObject
 {
 	GENERATED_BODY()
  
@@ -56,19 +56,19 @@ public:
 	static bool CreateUserDefinedAttribute(UInterchangeBaseNode* InterchangeNode, const FString& UserDefinedAttributeName, const ValueType& Value, const TOptional<FString>& PayloadKey);
 
 	UFUNCTION(BlueprintCallable, Category = "Interchange | Node | UserDefinedAttribute")
-	static bool CreateUserDefinedAttribute_Boolean(UInterchangeBaseNode* InterchangeNode, const FString& UserDefinedAttributeName, const bool& Value, const FString& PayloadKey);
+	static INTERCHANGECORE_API bool CreateUserDefinedAttribute_Boolean(UInterchangeBaseNode* InterchangeNode, const FString& UserDefinedAttributeName, const bool& Value, const FString& PayloadKey);
 	
 	UFUNCTION(BlueprintCallable, Category = "Interchange | Node | UserDefinedAttribute")
-	static bool CreateUserDefinedAttribute_Float(UInterchangeBaseNode* InterchangeNode, const FString& UserDefinedAttributeName, const float& Value, const FString& PayloadKey);
+	static INTERCHANGECORE_API bool CreateUserDefinedAttribute_Float(UInterchangeBaseNode* InterchangeNode, const FString& UserDefinedAttributeName, const float& Value, const FString& PayloadKey);
 
 	UFUNCTION(BlueprintCallable, Category = "Interchange | Node | UserDefinedAttribute")
-	static bool CreateUserDefinedAttribute_Double(UInterchangeBaseNode* InterchangeNode, const FString& UserDefinedAttributeName, const double& Value, const FString& PayloadKey);
+	static INTERCHANGECORE_API bool CreateUserDefinedAttribute_Double(UInterchangeBaseNode* InterchangeNode, const FString& UserDefinedAttributeName, const double& Value, const FString& PayloadKey);
 
 	UFUNCTION(BlueprintCallable, Category = "Interchange | Node | UserDefinedAttribute")
-	static bool CreateUserDefinedAttribute_Int32(UInterchangeBaseNode* InterchangeNode, const FString& UserDefinedAttributeName, const int32& Value, const FString& PayloadKey);
+	static INTERCHANGECORE_API bool CreateUserDefinedAttribute_Int32(UInterchangeBaseNode* InterchangeNode, const FString& UserDefinedAttributeName, const int32& Value, const FString& PayloadKey);
 
 	UFUNCTION(BlueprintCallable, Category = "Interchange | Node | UserDefinedAttribute")
-	static bool CreateUserDefinedAttribute_FString(UInterchangeBaseNode* InterchangeNode, const FString& UserDefinedAttributeName, const FString& Value, const FString& PayloadKey);
+	static INTERCHANGECORE_API bool CreateUserDefinedAttribute_FString(UInterchangeBaseNode* InterchangeNode, const FString& UserDefinedAttributeName, const FString& Value, const FString& PayloadKey);
 
 	/**
 	 * Remove the specified user defined attribute
@@ -78,7 +78,7 @@ public:
 	 *        Payload key will point on a FRichCurve payload.
 	 */
 	UFUNCTION(BlueprintCallable, Category = "Interchange | Node | UserDefinedAttribute")
-	static bool RemoveUserDefinedAttribute(UInterchangeBaseNode* InterchangeNode, const FString& UserDefinedAttributeName);
+	static INTERCHANGECORE_API bool RemoveUserDefinedAttribute(UInterchangeBaseNode* InterchangeNode, const FString& UserDefinedAttributeName);
 
 	/**
 	 * Get user defined attribute value and optional payload key
@@ -92,32 +92,32 @@ public:
 	static bool GetUserDefinedAttribute(const UInterchangeBaseNode* InterchangeNode, const FString& UserDefinedAttributeName, ValueType& OutValue, TOptional<FString>& OutPayloadKey);
 
 	UFUNCTION(BlueprintCallable, Category = "Interchange | Node | UserDefinedAttribute")
-	static bool GetUserDefinedAttribute_Boolean(const UInterchangeBaseNode* InterchangeNode, const FString& UserDefinedAttributeName, bool& OutValue, FString& OutPayloadKey);
+	static INTERCHANGECORE_API bool GetUserDefinedAttribute_Boolean(const UInterchangeBaseNode* InterchangeNode, const FString& UserDefinedAttributeName, bool& OutValue, FString& OutPayloadKey);
 
 	UFUNCTION(BlueprintCallable, Category = "Interchange | Node | UserDefinedAttribute")
-	static bool GetUserDefinedAttribute_Float(const UInterchangeBaseNode* InterchangeNode, const FString& UserDefinedAttributeName, float& OutValue, FString& OutPayloadKey);
+	static INTERCHANGECORE_API bool GetUserDefinedAttribute_Float(const UInterchangeBaseNode* InterchangeNode, const FString& UserDefinedAttributeName, float& OutValue, FString& OutPayloadKey);
 
 	UFUNCTION(BlueprintCallable, Category = "Interchange | Node | UserDefinedAttribute")
-	static bool GetUserDefinedAttribute_Double(const UInterchangeBaseNode* InterchangeNode, const FString& UserDefinedAttributeName, double& OutValue, FString& OutPayloadKey);
+	static INTERCHANGECORE_API bool GetUserDefinedAttribute_Double(const UInterchangeBaseNode* InterchangeNode, const FString& UserDefinedAttributeName, double& OutValue, FString& OutPayloadKey);
 
 	UFUNCTION(BlueprintCallable, Category = "Interchange | Node | UserDefinedAttribute")
-	static bool GetUserDefinedAttribute_Int32(const UInterchangeBaseNode* InterchangeNode, const FString& UserDefinedAttributeName, int32& OutValue, FString& OutPayloadKey);
+	static INTERCHANGECORE_API bool GetUserDefinedAttribute_Int32(const UInterchangeBaseNode* InterchangeNode, const FString& UserDefinedAttributeName, int32& OutValue, FString& OutPayloadKey);
 
 	UFUNCTION(BlueprintCallable, Category = "Interchange | Node | UserDefinedAttribute")
-	static bool GetUserDefinedAttribute_FString(const UInterchangeBaseNode* InterchangeNode, const FString& UserDefinedAttributeName, FString& OutValue, FString& OutPayloadKey);
+	static INTERCHANGECORE_API bool GetUserDefinedAttribute_FString(const UInterchangeBaseNode* InterchangeNode, const FString& UserDefinedAttributeName, FString& OutValue, FString& OutPayloadKey);
 
-	static TArray<FInterchangeUserDefinedAttributeInfo> GetUserDefinedAttributeInfos(const UInterchangeBaseNode* InterchangeNode);
+	static INTERCHANGECORE_API TArray<FInterchangeUserDefinedAttributeInfo> GetUserDefinedAttributeInfos(const UInterchangeBaseNode* InterchangeNode);
 
 	UFUNCTION(BlueprintCallable, Category = "Interchange | Node | UserDefinedAttribute")
-	static void GetUserDefinedAttributeInfos(const UInterchangeBaseNode* InterchangeNode, TArray<FInterchangeUserDefinedAttributeInfo>& UserDefinedAttributeInfos);
+	static INTERCHANGECORE_API void GetUserDefinedAttributeInfos(const UInterchangeBaseNode* InterchangeNode, TArray<FInterchangeUserDefinedAttributeInfo>& UserDefinedAttributeInfos);
  
 	UFUNCTION(BlueprintCallable, Category = "Interchange | Node | UserDefinedAttribute")
-	static void DuplicateAllUserDefinedAttribute(const UInterchangeBaseNode* InterchangeSourceNode, UInterchangeBaseNode* InterchangeDestinationNode, bool bAddSourceNodeName);
+	static INTERCHANGECORE_API void DuplicateAllUserDefinedAttribute(const UInterchangeBaseNode* InterchangeSourceNode, UInterchangeBaseNode* InterchangeDestinationNode, bool bAddSourceNodeName);
 
 private:
-	static const FString UserDefinedAttributeBaseKey;
-	static const FString UserDefinedAttributeValuePostKey;
-	static const FString UserDefinedAttributePayLoadPostKey;
+	static INTERCHANGECORE_API const FString UserDefinedAttributeBaseKey;
+	static INTERCHANGECORE_API const FString UserDefinedAttributeValuePostKey;
+	static INTERCHANGECORE_API const FString UserDefinedAttributePayLoadPostKey;
 };
 
 template<typename ValueType>
diff --git a/Engine/Source/Runtime/Interchange/Core/Public/Types/AttributeStorage.h b/Engine/Source/Runtime/Interchange/Core/Public/Types/AttributeStorage.h
--- a/Engine/Source/Runtime/Interchange/Core/Public/Types/AttributeStorage.h
+++ b/Engine/Source/Runtime/Interchange/Core/Public/Types/AttributeStorage.h
@@ -425,7 +425,7 @@ namespace UE
 		 * The storage is multi thread safe, it use a mutex to lock the storage for every read/write operation.
 		 * The hash of the storage is deterministic because it sort the attributes before calculating the hash.
 		 */
-		class INTERCHANGECORE_API FAttributeStorage
+		class FAttributeStorage
 		{
 		public:
 			/**
@@ -510,9 +510,9 @@ namespace UE
 
 			FAttributeStorage() = default;
 
-			FAttributeStorage(const FAttributeStorage& Other);
+			INTERCHANGECORE_API FAttributeStorage(const FAttributeStorage& Other);
 
-			FAttributeStorage& operator=(const FAttributeStorage& Other);
+			INTERCHANGECORE_API FAttributeStorage& operator=(const FAttributeStorage& Other);
 
 			/**
 			 * Register an attribute in the storage. Return success if the attribute was properly added, or there is an existing
@@ -579,7 +579,7 @@ namespace UE
 			 * - Key do not exist
 			 * - Internal storage structure remove error
 			 */
-			EAttributeStorageResult UnregisterAttribute(const FAttributeKey& ElementAttributeKey);
+			INTERCHANGECORE_API EAttributeStorageResult UnregisterAttribute(const FAttributeKey& ElementAttributeKey);
 
 			/**
 			 * Return an attribute handle for the specified attribute. This handle is compile type check and is use to Get and Set the attribute value type.
@@ -601,7 +601,7 @@ namespace UE
 			 * @param ElementAttributeKey - is the storage key (the path) of the attribute.
 			 *
 			 */
-			EAttributeTypes GetAttributeType(const FAttributeKey& ElementAttributeKey) const;
+			INTERCHANGECORE_API EAttributeTypes GetAttributeType(const FAttributeKey& ElementAttributeKey) const;
 
 			/**
 			 * Return true if the attribute key point on an existing attribute in the storage. Return false otherwise
@@ -609,13 +609,13 @@ namespace UE
 			 * @param ElementAttributeKey - is the storage key (the path) of the attribute.
 			 *
 			 */
-			bool ContainAttribute(const FAttributeKey& ElementAttributeKey) const;
+			INTERCHANGECORE_API bool ContainAttribute(const FAttributeKey& ElementAttributeKey) const;
 
 			/**
 			 * This function allow the user to retrieve the array of keys to iterate and do reflection on the storage content
 			 *
 			 */
-			void GetAttributeKeys(TArray<FAttributeKey>& AttributeKeys) const;
+			INTERCHANGECORE_API void GetAttributeKeys(TArray<FAttributeKey>& AttributeKeys) const;
 	
 			/**
 			 * Return a FGuid build from the FSHA1 of the specified attribute data. If the Attribute do not exist it return an empty FGUID
@@ -623,7 +623,7 @@ namespace UE
 			 * @param ElementAttributeKey is the storage key (the path) of the attribute.
 			 *
 			 */
-			FGuid GetAttributeHash(const FAttributeKey& ElementAttributeKey) const;
+			INTERCHANGECORE_API FGuid GetAttributeHash(const FAttributeKey& ElementAttributeKey) const;
 	
 			/**
 			 * This function fill the OutGuid with the hash of the specified attribute. Return true if the attribute exist and the OutGuid was assign, return false otherwise without touching the OutGuid.
@@ -632,7 +632,7 @@ namespace UE
 			 * @param OutGuid - is where we put the attribute hash.
 			 *
 			 */
-			bool GetAttributeHash(const FAttributeKey& ElementAttributeKey, FGuid& OutGuid) const;
+			INTERCHANGECORE_API bool GetAttributeHash(const FAttributeKey& ElementAttributeKey, FGuid& OutGuid) const;
 
 			/**
 			 * Return a FGuid build from the FSHA1 of all the attribute data contain in the node.
@@ -641,7 +641,7 @@ namespace UE
 			 * @note the attribute are sorted by key when building the FSHA1 data. The hash will be deterministic for the same data whatever
 			 * the order we add the attributes.
 			 */
-			FGuid GetStorageHash() const;
+			INTERCHANGECORE_API FGuid GetStorageHash() const;
 
 			/**
 			 * Compare two storage to know which properties was modified/added/removed.
@@ -653,7 +653,7 @@ namespace UE
 			 * @param ModifiedAttributes - All attributes that are in both storage base storage but have a different hash (different value)
 			 *
 			 */
-			static void CompareStorage(const FAttributeStorage& BaseStorage, const FAttributeStorage& VersionStorage, TArray<FAttributeKey>& RemovedAttributes, TArray<FAttributeKey>& AddedAttributes, TArray<FAttributeKey>& ModifiedAttributes);
+			static INTERCHANGECORE_API void CompareStorage(const FAttributeStorage& BaseStorage, const FAttributeStorage& VersionStorage, TArray<FAttributeKey>& RemovedAttributes, TArray<FAttributeKey>& AddedAttributes, TArray<FAttributeKey>& ModifiedAttributes);
 
 			/**
 			 * Copy an array of attributes from the source storage to the destination storage. If the attribute already exist in the destination, the value will be updated.
@@ -664,7 +664,7 @@ namespace UE
 			 * @param AttributeKeys - All attributes that must be copy from the source to the destination.
 			 *
 			 */
-			static void CopyStorageAttributes(const FAttributeStorage& SourceStorage, FAttributeStorage& DestinationStorage, const TArray<FAttributeKey>& AttributeKeys);
+			static INTERCHANGECORE_API void CopyStorageAttributes(const FAttributeStorage& SourceStorage, FAttributeStorage& DestinationStorage, const TArray<FAttributeKey>& AttributeKeys);
 
 			/**
 			 * Return the defrag ratio. This ratio is use to know when we need to defrag the storage
@@ -677,7 +677,7 @@ namespace UE
 			}
 	
 			/** Set the defrag ratio. See the GetDefragRatio for the defrag documentation. */
-			void SetDefragRatio(const float InDefragRatio);
+			INTERCHANGECORE_API void SetDefragRatio(const float InDefragRatio);
 	
 			friend FArchive& operator<<(FArchive& Ar, FAttributeStorage& Storage)
 			{
@@ -694,7 +694,7 @@ namespace UE
 			 * @param NewAttributeCount: The number of attribute we want to reserve. Zero value do not reserve attribute count
 			 * @param NewStorageSize: The size on the storage all those new attribute will need. Passing zero value do not reserve storage size
 			 */
-			void Reserve(int64 NewAttributeCount, int64 NewStorageSize);
+			INTERCHANGECORE_API void Reserve(int64 NewAttributeCount, int64 NewStorageSize);
 
 		protected:
 			/** Structure use to hold the attribute information store in the attribute allocation table. */
@@ -774,9 +774,9 @@ namespace UE
 			}
 	
 			/** Defrag the storage by using memmove on the attribute store after a hole in the storage. */
-			void DefragInternal();
+			INTERCHANGECORE_API void DefragInternal();
 
-			static FGuid GetValueHash(const uint8* Value, uint64 ValueSize);
+			static INTERCHANGECORE_API FGuid GetValueHash(const uint8* Value, uint64 ValueSize);
 
 			/** Return the storage size of the value. */
 			template<typename T>
@@ -915,13 +915,13 @@ namespace UE
 				return GenericArrayGetAttribute(ElementAttributeKey, OutValue);
 			}
 	
-			void ExtractFStringAttributeFromStorage(const uint8* StorageData, const FAttributeAllocationInfo* AttributeAllocationInfo, FString& OutValue) const;
+			INTERCHANGECORE_API void ExtractFStringAttributeFromStorage(const uint8* StorageData, const FAttributeAllocationInfo* AttributeAllocationInfo, FString& OutValue) const;
 
-			EAttributeStorageResult GetAttribute(const FAttributeKey& ElementAttributeKey, FString& OutValue, TSpecializeType<FString >) const;
+			INTERCHANGECORE_API EAttributeStorageResult GetAttribute(const FAttributeKey& ElementAttributeKey, FString& OutValue, TSpecializeType<FString >) const;
 
-			EAttributeStorageResult GetAttribute(const FAttributeKey& ElementAttributeKey, FName& OutValue, TSpecializeType<FName >) const;
+			INTERCHANGECORE_API EAttributeStorageResult GetAttribute(const FAttributeKey& ElementAttributeKey, FName& OutValue, TSpecializeType<FName >) const;
 
-			EAttributeStorageResult GetAttribute(const FAttributeKey& ElementAttributeKey, FSoftObjectPath& OutValue, TSpecializeType<FSoftObjectPath >) const;
+			INTERCHANGECORE_API EAttributeStorageResult GetAttribute(const FAttributeKey& ElementAttributeKey, FSoftObjectPath& OutValue, TSpecializeType<FSoftObjectPath >) const;
 
 			template<typename T>
 			static uint64 GetValueSize(const T& Value, TSpecializeType<T>)
@@ -934,7 +934,7 @@ namespace UE
 				return sizeof(T);
 			}
 	
-			static uint64 GetValueSize(const FString& Value, TSpecializeType<FString >);
+			static INTERCHANGECORE_API uint64 GetValueSize(const FString& Value, TSpecializeType<FString >);
 
 			static uint64 GetValueSize(const FName& Value, TSpecializeType<FName>)
 			{
@@ -948,9 +948,9 @@ namespace UE
 				return GetValueSize(ValueStr, TSpecializeType<FString >());
 			}
 
-			static uint64 GetValueSize(const TArray<uint8>& Value, TSpecializeType<TArray<uint8> >);
+			static INTERCHANGECORE_API uint64 GetValueSize(const TArray<uint8>& Value, TSpecializeType<TArray<uint8> >);
 
-			static uint64 GetValueSize(const TArray64<uint8>& Value, TSpecializeType<TArray64<uint8> >);
+			static INTERCHANGECORE_API uint64 GetValueSize(const TArray64<uint8>& Value, TSpecializeType<TArray64<uint8> >);
 	
 			// End specialize templates
 
diff --git a/Engine/Source/Runtime/Interchange/Engine/Public/InterchangeAssetImportData.h b/Engine/Source/Runtime/Interchange/Engine/Public/InterchangeAssetImportData.h
--- a/Engine/Source/Runtime/Interchange/Engine/Public/InterchangeAssetImportData.h
+++ b/Engine/Source/Runtime/Interchange/Engine/Public/InterchangeAssetImportData.h
@@ -14,8 +14,8 @@
 
 class UInterchangePipelineBase;
 
-UCLASS(BlueprintType)
-class INTERCHANGEENGINE_API UInterchangeAssetImportData : public UAssetImportData
+UCLASS(BlueprintType, MinimalAPI)
+class UInterchangeAssetImportData : public UAssetImportData
 {
 	GENERATED_BODY()
 public:
diff --git a/Engine/Source/Runtime/Interchange/Engine/Public/InterchangeBlueprintPipelineBase.h b/Engine/Source/Runtime/Interchange/Engine/Public/InterchangeBlueprintPipelineBase.h
--- a/Engine/Source/Runtime/Interchange/Engine/Public/InterchangeBlueprintPipelineBase.h
+++ b/Engine/Source/Runtime/Interchange/Engine/Public/InterchangeBlueprintPipelineBase.h
@@ -10,8 +10,8 @@
 
 #include "InterchangeBlueprintPipelineBase.generated.h"
 
-UCLASS(BlueprintType)
-class INTERCHANGEENGINE_API UInterchangeBlueprintPipelineBase : public UBlueprint
+UCLASS(BlueprintType, MinimalAPI)
+class UInterchangeBlueprintPipelineBase : public UBlueprint
 {
 	GENERATED_BODY()
 
@@ -22,4 +22,4 @@ public:
 		//We must make sure the GeneratedClass is generated after the blueprint is loaded
 		bRecompileOnLoad = true;
 	}
-};
\ No newline at end of file
+};
diff --git a/Engine/Source/Runtime/Interchange/Engine/Public/InterchangeFilePickerBase.h b/Engine/Source/Runtime/Interchange/Engine/Public/InterchangeFilePickerBase.h
--- a/Engine/Source/Runtime/Interchange/Engine/Public/InterchangeFilePickerBase.h
+++ b/Engine/Source/Runtime/Interchange/Engine/Public/InterchangeFilePickerBase.h
@@ -27,8 +27,8 @@ struct FInterchangeFilePickerParameters
 	FString DefaultPath = TEXT("");
 };
 
-UCLASS(Abstract, BlueprintType, Blueprintable)
-class INTERCHANGEENGINE_API UInterchangeFilePickerBase : public UObject
+UCLASS(Abstract, BlueprintType, Blueprintable, MinimalAPI)
+class UInterchangeFilePickerBase : public UObject
 {
 	GENERATED_BODY()
 
@@ -38,7 +38,7 @@ public:
 	 * Non virtual helper to allow blueprint to implement event base function to implement the function FilePickerForTranslatorAssetType,
 	 */
 	UFUNCTION(BlueprintNativeEvent, BlueprintCallable, Category = "Interchange | Utilities")
-	bool ScriptedFilePickerForTranslatorAssetType(const EInterchangeTranslatorAssetType TranslatorAssetType, FInterchangeFilePickerParameters& Parameters, TArray<FString>& OutFilenames);
+	INTERCHANGEENGINE_API bool ScriptedFilePickerForTranslatorAssetType(const EInterchangeTranslatorAssetType TranslatorAssetType, FInterchangeFilePickerParameters& Parameters, TArray<FString>& OutFilenames);
 
 	/** The default implementation (call if the blueprint do not have any implementation) will call the virtual FilePickerForTranslatorAssetType */
 	bool ScriptedFilePickerForTranslatorAssetType_Implementation(const EInterchangeTranslatorAssetType TranslatorAssetType, FInterchangeFilePickerParameters& Parameters, TArray<FString>& OutFilenames)
@@ -51,7 +51,7 @@ public:
 	 * Non virtual helper to allow blueprint to implement event base function to implement the function FilePickerForTranslatorType,
 	 */
 	UFUNCTION(BlueprintNativeEvent, BlueprintCallable, Category = "Interchange | Utilities")
-	bool ScriptedFilePickerForTranslatorType(const EInterchangeTranslatorType TranslatorType, FInterchangeFilePickerParameters& Parameters, TArray<FString>& OutFilenames);
+	INTERCHANGEENGINE_API bool ScriptedFilePickerForTranslatorType(const EInterchangeTranslatorType TranslatorType, FInterchangeFilePickerParameters& Parameters, TArray<FString>& OutFilenames);
 
 	/** The default implementation (call if the blueprint do not have any implementation) will call the virtual FilePickerForTranslatorType */
 	bool ScriptedFilePickerForTranslatorType_Implementation(const EInterchangeTranslatorType TranslatorAssetType, FInterchangeFilePickerParameters& Parameters, TArray<FString>& OutFilenames)
diff --git a/Engine/Source/Runtime/Interchange/Engine/Public/InterchangeImportCommon.h b/Engine/Source/Runtime/Interchange/Engine/Public/InterchangeImportCommon.h
--- a/Engine/Source/Runtime/Interchange/Engine/Public/InterchangeImportCommon.h
+++ b/Engine/Source/Runtime/Interchange/Engine/Public/InterchangeImportCommon.h
@@ -23,11 +23,11 @@ namespace UE::Interchange
 	* All the code we cannot put in the base factory class because of dependencies (like Engine dep)
 	* Will be available here.
 	*/
-	class INTERCHANGEENGINE_API FFactoryCommon
+	class FFactoryCommon
 	{
 	public:
 
-		struct INTERCHANGEENGINE_API FUpdateImportAssetDataParameters
+		struct FUpdateImportAssetDataParameters
 		{
 			UObject* AssetImportDataOuter = nullptr;
 			UAssetImportData* AssetImportData = nullptr;
@@ -36,7 +36,7 @@ namespace UE::Interchange
 			UInterchangeBaseNodeContainer* NodeContainer;
 			const TArray<UObject*> Pipelines;
 
-			FUpdateImportAssetDataParameters(UObject* InAssetImportDataOuter
+			INTERCHANGEENGINE_API FUpdateImportAssetDataParameters(UObject* InAssetImportDataOuter
 																, UAssetImportData* InAssetImportData
 																, const UInterchangeSourceData* InSourceData
 																, FString InNodeUniqueID
@@ -49,7 +49,7 @@ namespace UE::Interchange
 		 * If the AssetImportData is null it will create one. If the AssetImportData is not an UInterchangeAssetImportData it will create a new one.
 		 * @return The source data that should be stored on the asset or nullptr if a parameter is invalid
 		 */
-		static UAssetImportData* UpdateImportAssetData(FUpdateImportAssetDataParameters& Parameters);
+		static INTERCHANGEENGINE_API UAssetImportData* UpdateImportAssetData(FUpdateImportAssetDataParameters& Parameters);
 
 		/**
 		 * Update the AssetImportData of the specified asset in the parameters. Also update the node container and the node unique id.
@@ -57,15 +57,15 @@ namespace UE::Interchange
 		 * The file source update is done by calling the function parameter CustomFileSourceUpdate, so its the client responsability to properly update the file source.
 		 * @return The source data that should be stored on the asset or nullptr if a parameter is invalid
 		 */
-		static UAssetImportData* UpdateImportAssetData(FUpdateImportAssetDataParameters& Parameters, TFunctionRef<void(UInterchangeAssetImportData*)> CustomFileSourceUpdate);
+		static INTERCHANGEENGINE_API UAssetImportData* UpdateImportAssetData(FUpdateImportAssetDataParameters& Parameters, TFunctionRef<void(UInterchangeAssetImportData*)> CustomFileSourceUpdate);
 
 #if WITH_EDITORONLY_DATA
-		struct INTERCHANGEENGINE_API FSetImportAssetDataParameters : public FUpdateImportAssetDataParameters
+		struct FSetImportAssetDataParameters : public FUpdateImportAssetDataParameters
 		{
 			// Allow the factory to provide is own list of source files.
 			TArray<FAssetImportInfo::FSourceFile> SourceFiles;
 
-			FSetImportAssetDataParameters(UObject* InAssetImportDataOuter
+			INTERCHANGEENGINE_API FSetImportAssetDataParameters(UObject* InAssetImportDataOuter
 				, UAssetImportData* InAssetImportData
 				, const UInterchangeSourceData* InSourceData
 				, FString InNodeUniqueID
@@ -78,30 +78,30 @@ namespace UE::Interchange
 		 * If the AssetImportData is null it will create one. If the AssetImportData is not an UInterchangeAssetImportData it will create a new one.
 		 * @return The source data that should be stored on the asset or nullptr if a parameter is invalid
 		 */
-		static UAssetImportData* SetImportAssetData(FSetImportAssetDataParameters& Parameters);
+		static INTERCHANGEENGINE_API UAssetImportData* SetImportAssetData(FSetImportAssetDataParameters& Parameters);
 
 		/**
 		 * Fills the OutSourceFilenames array with the list of source files contained in the asset source data.
 		 * Returns true if the operation was successful.
 		 */
-		static bool GetSourceFilenames(const UAssetImportData* AssetImportData, TArray<FString>& OutSourceFilenames);
+		static INTERCHANGEENGINE_API bool GetSourceFilenames(const UAssetImportData* AssetImportData, TArray<FString>& OutSourceFilenames);
 
 		/**
 		 * Sets the SourceFileName value at the specified index.
 		 */
-		static bool SetSourceFilename(UAssetImportData* AssetImportData, const FString& SourceFilename, int32 SourceIndex, const FString& SourceLabel = FString());
+		static INTERCHANGEENGINE_API bool SetSourceFilename(UAssetImportData* AssetImportData, const FString& SourceFilename, int32 SourceIndex, const FString& SourceLabel = FString());
 
 		/**
 		 * Set the object's reimport source at the specified index value.
 		 */
-		static bool SetReimportSourceIndex(const UObject* Object, UAssetImportData* AssetImportData, int32 SourceIndex);
+		static INTERCHANGEENGINE_API bool SetReimportSourceIndex(const UObject* Object, UAssetImportData* AssetImportData, int32 SourceIndex);
 
 #endif // WITH_EDITORONLY_DATA
 
 		/**
 		 * Apply the current strategy to the PipelineAssetNode
 		 */
-		static void ApplyReimportStrategyToAsset(UObject* Asset
+		static INTERCHANGEENGINE_API void ApplyReimportStrategyToAsset(UObject* Asset
 											, const UInterchangeFactoryBaseNode* PreviousAssetNode
 											, const UInterchangeFactoryBaseNode* CurrentAssetNode
 											, UInterchangeFactoryBaseNode* PipelineAssetNode);
@@ -115,7 +115,7 @@ namespace UE::Interchange
 		 * @oaram AssetName: Asset name of the actual object to reimport
 		 * @oaram SubPathString: Optional subobject name
 		 */
-		static 	UObject* GetObjectToReimport(UObject* ReimportObject, const UInterchangeFactoryBaseNode& FactoryNode, const FString& PackageName, const FString& AssetName, const FString& SubPathString = FString());
+		static 	INTERCHANGEENGINE_API UObject* GetObjectToReimport(UObject* ReimportObject, const UInterchangeFactoryBaseNode& FactoryNode, const FString& PackageName, const FString& AssetName, const FString& SubPathString = FString());
 
 		/**
 		 * If the ReimportObject is a UInterchangeSceneImportAsset, returns the factory node
@@ -126,11 +126,11 @@ namespace UE::Interchange
 		 * @oaram AssetName: Asset name of the actual object to reimport
 		 * @oaram SubPathString: Optional subobject name
 		 */
-		static 	const UInterchangeFactoryBaseNode* GetFactoryNode(UObject* ReimportObject, const FString& PackageName, const FString& AssetName, const FString& SubPathString = FString());
+		static 	INTERCHANGEENGINE_API const UInterchangeFactoryBaseNode* GetFactoryNode(UObject* ReimportObject, const FString& PackageName, const FString& AssetName, const FString& SubPathString = FString());
 		
 		/**
 		 * Factory helper use to find the object we have create in the game thread when we do the async part of the factory import.
 		 */
-		static UObject* AsyncFindObject(UInterchangeFactoryBaseNode* FactoryNode, const UClass* FactoryClass, UObject* Parent, const FString& AssetName);
+		static INTERCHANGEENGINE_API UObject* AsyncFindObject(UInterchangeFactoryBaseNode* FactoryNode, const UClass* FactoryClass, UObject* Parent, const FString& AssetName);
 	};
 } //ns UE::Interchange
diff --git a/Engine/Source/Runtime/Interchange/Engine/Public/InterchangeManager.h b/Engine/Source/Runtime/Interchange/Engine/Public/InterchangeManager.h
--- a/Engine/Source/Runtime/Interchange/Engine/Public/InterchangeManager.h
+++ b/Engine/Source/Runtime/Interchange/Engine/Public/InterchangeManager.h
@@ -45,29 +45,29 @@ namespace UE
 {
 	namespace Interchange
 	{
-		class INTERCHANGEENGINE_API FScopedInterchangeImportEnableState
+		class FScopedInterchangeImportEnableState
 		{
 		public:
-			explicit FScopedInterchangeImportEnableState(const bool bScopeValue);
-			~FScopedInterchangeImportEnableState();
+			INTERCHANGEENGINE_API explicit FScopedInterchangeImportEnableState(const bool bScopeValue);
+			INTERCHANGEENGINE_API ~FScopedInterchangeImportEnableState();
 		private:
 			bool bOriginalInterchangeImportEnableState;
 		};
 
-		class INTERCHANGEENGINE_API FScopedSourceData
+		class FScopedSourceData
 		{
 		public:
-			explicit FScopedSourceData(const FString& Filename);
-			UInterchangeSourceData* GetSourceData() const;
+			INTERCHANGEENGINE_API explicit FScopedSourceData(const FString& Filename);
+			INTERCHANGEENGINE_API UInterchangeSourceData* GetSourceData() const;
 		private:
 			TStrongObjectPtr<UInterchangeSourceData> SourceDataPtr = nullptr;
 		};
 
-		class INTERCHANGEENGINE_API FScopedTranslator
+		class FScopedTranslator
 		{
 		public:
-			explicit FScopedTranslator(const UInterchangeSourceData* SourceData);
-			UInterchangeTranslatorBase* GetTranslator();
+			INTERCHANGEENGINE_API explicit FScopedTranslator(const UInterchangeSourceData* SourceData);
+			INTERCHANGEENGINE_API UInterchangeTranslatorBase* GetTranslator();
 
 		private:
 			TStrongObjectPtr<UInterchangeTranslatorBase> ScopedTranslatorPtr = nullptr;
@@ -95,10 +95,10 @@ namespace UE
 			UObject* ReimportObject = nullptr;
 		};
 
-		class INTERCHANGEENGINE_API FImportResult : protected FGCObject
+		class FImportResult : protected FGCObject
 		{
 		public:
-			FImportResult();
+			INTERCHANGEENGINE_API FImportResult();
 
 			FImportResult(FImportResult&&) = delete;
 			FImportResult& operator=(FImportResult&&) = delete;
@@ -116,30 +116,30 @@ namespace UE
 				Done
 			};
 
-			EStatus GetStatus() const;
+			INTERCHANGEENGINE_API EStatus GetStatus() const;
 
-			bool IsValid() const;
+			INTERCHANGEENGINE_API bool IsValid() const;
 
-			void SetInProgress();
-			void SetDone();
-			void WaitUntilDone();
+			INTERCHANGEENGINE_API void SetInProgress();
+			INTERCHANGEENGINE_API void SetDone();
+			INTERCHANGEENGINE_API void WaitUntilDone();
 
 			// Assets are only made available once they have been completely imported (passed through the entire import pipeline)
 			// While the status isn't EStatus::Done, the list can grow between subsequent calls.
 			// FAssetImportResult holds a reference to the assets so that they aren't garbage collected.
-			const TArray< UObject* >& GetImportedObjects() const;
+			INTERCHANGEENGINE_API const TArray< UObject* >& GetImportedObjects() const;
 
 			// Helper to get the first asset of a certain class. Use when expecting a single asset of that class to be imported since the order isn't deterministic.
-			UObject* GetFirstAssetOfClass(UClass* InClass) const;
+			INTERCHANGEENGINE_API UObject* GetFirstAssetOfClass(UClass* InClass) const;
 
 			// Return the results of this asset import operation
 			UInterchangeResultsContainer* GetResults() const { return Results; }
 
 			// Adds an asset to the list of imported assets.
-			void AddImportedObject(UObject* ImportedObject);
+			INTERCHANGEENGINE_API void AddImportedObject(UObject* ImportedObject);
 
 			// Callback when the status switches to done.
-			void OnDone(TFunction< void(FImportResult&) > Callback);
+			INTERCHANGEENGINE_API void OnDone(TFunction< void(FImportResult&) > Callback);
 
 			// Internal delegates (use the FImportAssetParameters when calling the interchange import functions to set those)
 			FOnObjectImportDoneDynamic OnObjectDone;
@@ -150,7 +150,7 @@ namespace UE
 
 		protected:
 			/* FGCObject interface */
-			virtual void AddReferencedObjects(FReferenceCollector& Collector) override;
+			INTERCHANGEENGINE_API virtual void AddReferencedObjects(FReferenceCollector& Collector) override;
 			virtual FString GetReferencerName() const override
 			{
 				return TEXT("UE::Interchange::FImportResult");
@@ -301,8 +301,8 @@ namespace UE
 /**
  * This class is use to pass override pipelines in the ImportAssetTask Options member
  */
-UCLASS(Transient, BlueprintType)
-class INTERCHANGEENGINE_API UInterchangePipelineStackOverride : public UObject
+UCLASS(Transient, BlueprintType, MinimalAPI)
+class UInterchangePipelineStackOverride : public UObject
 {
 	GENERATED_BODY()
 public:
@@ -311,17 +311,17 @@ public:
 	TArray<FSoftObjectPath> OverridePipelines;
 
 	UFUNCTION(BlueprintCallable, Category = "Interchange | PipelineStackOverride")
-	void AddPythonPipeline(UInterchangePythonPipelineBase* PipelineBase);
+	INTERCHANGEENGINE_API void AddPythonPipeline(UInterchangePythonPipelineBase* PipelineBase);
 
 	UFUNCTION(BlueprintCallable, Category = "Interchange | PipelineStackOverride")
-	void AddBlueprintPipeline(UInterchangeBlueprintPipelineBase* PipelineBase);
+	INTERCHANGEENGINE_API void AddBlueprintPipeline(UInterchangeBlueprintPipelineBase* PipelineBase);
 
 	UFUNCTION(BlueprintCallable, Category = "Interchange | PipelineStackOverride")
-	void AddPipeline(UInterchangePipelineBase* PipelineBase);
+	INTERCHANGEENGINE_API void AddPipeline(UInterchangePipelineBase* PipelineBase);
 };
 
 USTRUCT(BlueprintType)
-struct INTERCHANGEENGINE_API FImportAssetParameters
+struct FImportAssetParameters
 {
 	GENERATED_USTRUCT_BODY()
 
@@ -368,8 +368,8 @@ struct INTERCHANGEENGINE_API FImportAssetParameters
 	FOnImportDoneNative OnSceneImportDoneNative;
 };
 
-UCLASS(Transient, BlueprintType)
-class INTERCHANGEENGINE_API UInterchangeManager : public UObject
+UCLASS(Transient, BlueprintType, MinimalAPI)
+class UInterchangeManager : public UObject
 {
 	GENERATED_BODY()
 public:
@@ -386,10 +386,10 @@ public:
 	}
 
 	/** Return the interchange manager singleton.*/
-	static UInterchangeManager& GetInterchangeManager();
+	static INTERCHANGEENGINE_API UInterchangeManager& GetInterchangeManager();
 
 	/** Return the CVar which make interchange enable or not.*/
-	static bool IsInterchangeImportEnabled();
+	static INTERCHANGEENGINE_API bool IsInterchangeImportEnabled();
 
 	/** delegate type fired when new assets have been imported. Note: InCreatedObject can be NULL if import failed. Params: UFactory* InFactory, UObject* InCreatedObject */
 	DECLARE_MULTICAST_DELEGATE_OneParam(FInterchangeOnAssetPostImport, UObject*);
@@ -414,7 +414,7 @@ public:
 	 * @Note if you register multiple time the same class it will return true for every call
 	 * @Note The order in which the translators are registered will be the same as the order used to select a translator to import a file
 	 */
-	bool RegisterTranslator(const UClass* TranslatorClass);
+	INTERCHANGEENGINE_API bool RegisterTranslator(const UClass* TranslatorClass);
 
 	/**
 	 * Any factory must register to the manager
@@ -423,7 +423,7 @@ public:
 	 *
 	 * @Note if you register multiple time the same class it will return true for every call
 	 */
-	bool RegisterFactory(const UClass* Factory);
+	INTERCHANGEENGINE_API bool RegisterFactory(const UClass* Factory);
 
 	/**
 	 * Any writer must register to the manager
@@ -432,22 +432,22 @@ public:
 	 *
 	 * @Note if you register multiple time the same class it will return true for every call
 	 */
-	bool RegisterWriter(const UClass* Writer);
+	INTERCHANGEENGINE_API bool RegisterWriter(const UClass* Writer);
 
 	/**
 	 * Returns the list of supported formats for a given translator type.
 	 */
-	TArray<FString> GetSupportedFormats(const EInterchangeTranslatorType ForTranslatorType) const;
+	INTERCHANGEENGINE_API TArray<FString> GetSupportedFormats(const EInterchangeTranslatorType ForTranslatorType) const;
 
 	/**
 	 * Returns the list of formats supporting the specified translator asset type.
 	 */
-	TArray<FString> GetSupportedAssetTypeFormats(const EInterchangeTranslatorAssetType ForTranslatorAssetType) const;
+	INTERCHANGEENGINE_API TArray<FString> GetSupportedAssetTypeFormats(const EInterchangeTranslatorAssetType ForTranslatorAssetType) const;
 
 	/**
 	 * Returns the list of supported formats for a given Object.
 	 */
-	TArray<FString> GetSupportedFormatsForObject(const UObject* Object) const;
+	INTERCHANGEENGINE_API TArray<FString> GetSupportedFormatsForObject(const UObject* Object) const;
 
 	/**
 	 * Look if there is a registered translator for this source data.
@@ -455,14 +455,14 @@ public:
 	 * @Param SourceData - The source data input we want to translate to Uod
 	 * @return True if there is a registered translator that can handle handle this source data, false otherwise.
 	 */
-	bool CanTranslateSourceData(const UInterchangeSourceData* SourceData) const;
+	INTERCHANGEENGINE_API bool CanTranslateSourceData(const UInterchangeSourceData* SourceData) const;
 
 	/**
 	 * Returns true if Interchange can create that type of assets and is able to translate its source file.
 	 * @Param Object - The object we want to reimport.
 	 * @Param OutFilenames - An array that is filled with the object's source filenames if the operation is successful.
 	 */
-	bool CanReimport(const UObject* Object, TArray<FString>& OutFilenames) const;
+	INTERCHANGEENGINE_API bool CanReimport(const UObject* Object, TArray<FString>& OutFilenames) const;
 
 	/**
 	 * Call this to start an import asset process, the caller must specify a source data.
@@ -474,8 +474,8 @@ public:
 	 * @return true if the import succeed, false otherwise.
 	 */
 	UFUNCTION(BlueprintCallable, Category = "Interchange | Import Manager")
-	bool ImportAsset(const FString& ContentPath, const UInterchangeSourceData* SourceData, const FImportAssetParameters& ImportAssetParameters);
-	UE::Interchange::FAssetImportResultRef ImportAssetAsync(const FString& ContentPath, const UInterchangeSourceData* SourceData, const FImportAssetParameters& ImportAssetParameters);
+	INTERCHANGEENGINE_API bool ImportAsset(const FString& ContentPath, const UInterchangeSourceData* SourceData, const FImportAssetParameters& ImportAssetParameters);
+	INTERCHANGEENGINE_API UE::Interchange::FAssetImportResultRef ImportAssetAsync(const FString& ContentPath, const UInterchangeSourceData* SourceData, const FImportAssetParameters& ImportAssetParameters);
 
 	/**
 	 * Call this to start an import scene process, the caller must specify a source data.
@@ -487,9 +487,9 @@ public:
 	 * @return true if the import succeed, false otherwise.
 	 */
 	UFUNCTION(BlueprintCallable, Category = "Interchange | Import Manager")
-	bool ImportScene(const FString& ContentPath, const UInterchangeSourceData* SourceData, const FImportAssetParameters& ImportAssetParameters);
+	INTERCHANGEENGINE_API bool ImportScene(const FString& ContentPath, const UInterchangeSourceData* SourceData, const FImportAssetParameters& ImportAssetParameters);
 
-	TTuple<UE::Interchange::FAssetImportResultRef, UE::Interchange::FSceneImportResultRef>
+	INTERCHANGEENGINE_API TTuple<UE::Interchange::FAssetImportResultRef, UE::Interchange::FSceneImportResultRef>
 	ImportSceneAsync(const FString& ContentPath, const UInterchangeSourceData* SourceData, const FImportAssetParameters& ImportAssetParameters);
 
 	/**
@@ -500,7 +500,7 @@ public:
 	 * @return true if the import succeed, false otherwise.
 	 */
 	UFUNCTION(BlueprintCallable, Category = "Interchange | Export Manager")
-	bool ExportAsset(const UObject* Asset, bool bIsAutomated = false);
+	INTERCHANGEENGINE_API bool ExportAsset(const UObject* Asset, bool bIsAutomated = false);
 
 	/**
 	 * Call this to start an export scene process, the caller must specify a source data
@@ -510,7 +510,7 @@ public:
 	 * @return true if the import succeed, false otherwise.
 	 */
 	UFUNCTION(BlueprintCallable, Category = "Interchange | Export Manager")
-	bool ExportScene(const UObject* World, bool bIsAutomated = false);
+	INTERCHANGEENGINE_API bool ExportScene(const UObject* World, bool bIsAutomated = false);
 
 	/*
 	* Script helper to create a source data object pointing on a file on disk
@@ -518,7 +518,7 @@ public:
 	* @return: A new UInterchangeSourceData.
 	*/
 	UFUNCTION(BlueprintCallable, Category = "Interchange | Import Manager")
-	static UInterchangeSourceData* CreateSourceData(const FString& InFileName);
+	static INTERCHANGEENGINE_API UInterchangeSourceData* CreateSourceData(const FString& InFileName);
 
 	/**
 	* Script helper to get a registered factory for a specified class
@@ -526,42 +526,42 @@ public:
 	* @return: if found, we return the factory class that is registered. Return NULL if nothing found.
 	*/
 	UFUNCTION(BlueprintCallable, Category = "Interchange | Import Manager")
-	const UClass* GetRegisteredFactoryClass(const UClass* ClassToMake) const;
+	INTERCHANGEENGINE_API const UClass* GetRegisteredFactoryClass(const UClass* ClassToMake) const;
 
 	/**
 	 * Return an FImportAsynHelper pointer. The pointer is deleted when ReleaseAsyncHelper is call.
 	 * @param Data - The data we want to pass to the different import tasks
 	 */
-	TSharedRef<UE::Interchange::FImportAsyncHelper, ESPMode::ThreadSafe> CreateAsyncHelper(const UE::Interchange::FImportAsyncHelperData& Data, const FImportAssetParameters& ImportAssetParameters);
+	INTERCHANGEENGINE_API TSharedRef<UE::Interchange::FImportAsyncHelper, ESPMode::ThreadSafe> CreateAsyncHelper(const UE::Interchange::FImportAsyncHelperData& Data, const FImportAssetParameters& ImportAssetParameters);
 
 	/** Delete the specified AsyncHelper and remove it from the array that was holding it. */
-	void ReleaseAsyncHelper(TWeakPtr<UE::Interchange::FImportAsyncHelper, ESPMode::ThreadSafe> AsyncHelper);
+	INTERCHANGEENGINE_API void ReleaseAsyncHelper(TWeakPtr<UE::Interchange::FImportAsyncHelper, ESPMode::ThreadSafe> AsyncHelper);
 
 	/*
 	 * Return the first translator that can translate the source data.
 	 * @Param SourceData - The source data for which we search a translator.
 	 * @return return a matching translator or nullptr if there is no match.
 	 */
-	UInterchangeTranslatorBase* GetTranslatorForSourceData(const UInterchangeSourceData* SourceData) const;
+	INTERCHANGEENGINE_API UInterchangeTranslatorBase* GetTranslatorForSourceData(const UInterchangeSourceData* SourceData) const;
 
 	/**
 	 * Return true if the Interchange is active (importing or exporting), return false otherwise.
 	 */
-	bool IsInterchangeActive();
+	INTERCHANGEENGINE_API bool IsInterchangeActive();
 
 	/**
 	 * Return false if the Interchange is not active (importing or exporting).
 	 * If the interchange is active, it will display a notification to let the user know they can cancel the asynchronous import/export
 	 * to be able to complete the operation they want to do. (The exit editor operation is calling this)
 	 */
-	bool WarnIfInterchangeIsActive();
+	INTERCHANGEENGINE_API bool WarnIfInterchangeIsActive();
 
 	/**
 	 * Look if there is a translator registered that can translate the source data with the specified PayloadInterface
 	 * @Param SourceData - The source data input we want to translate to Uod
 	 * @return true if the source data can be translated using the specified PayloadInterface, false otherwise.
 	 */
-	bool CanTranslateSourceDataWithPayloadInterface(const UInterchangeSourceData* SourceData, const UClass* PayloadInterfaceClass) const;
+	INTERCHANGEENGINE_API bool CanTranslateSourceDataWithPayloadInterface(const UInterchangeSourceData* SourceData, const UClass* PayloadInterfaceClass) const;
 
 	/*
 	 * Return the first translator that can translate the source data with the specified PayloadInterface.
@@ -569,55 +569,55 @@ public:
 	 * @Param PayloadInterfaceClass - The PayloadInterface that the translator must implement.
 	 * @return return a matching translator implementing the specified PayloadInterface or nullptr if there is no match.
 	 */
-	UInterchangeTranslatorBase* GetTranslatorSupportingPayloadInterfaceForSourceData(const UInterchangeSourceData* SourceData, const UClass* PayloadInterfaceClass) const;
+	INTERCHANGEENGINE_API UInterchangeTranslatorBase* GetTranslatorSupportingPayloadInterfaceForSourceData(const UInterchangeSourceData* SourceData, const UClass* PayloadInterfaceClass) const;
 
 	/**
 	 * Return true if the object is being imported, return false otherwise. If the user import multiple file in the same folder its possible to
 	 * have the same asset name in two different files.
 	 */
-	bool IsObjectBeingImported(UObject* Object) const;
+	INTERCHANGEENGINE_API bool IsObjectBeingImported(UObject* Object) const;
 
 protected:
 
 	/** Return true if we can show some UI */
-	static bool IsAttended();
+	static INTERCHANGEENGINE_API bool IsAttended();
 
 	/*
 	 * Find all Pipeline candidate (c++, blueprint and python).
 	 * @Param SourceData - The source data for which we search a translator.
 	 * @return return a matching translator or nullptr if there is no match.
 	 */
-	void FindPipelineCandidate(TArray<UClass*>& PipelineCandidates);
+	INTERCHANGEENGINE_API void FindPipelineCandidate(TArray<UClass*>& PipelineCandidates);
 
 	/**
 	 * This function cancel all task and finish them has fast as possible.
 	 * We use this if the user cancel the work or if the editor is exiting.
 	 * @note - This is a asynchronous call, tasks will be completed (cancel) soon.
 	 */
-	void CancelAllTasks();
+	INTERCHANGEENGINE_API void CancelAllTasks();
 
 	/**
 	 * Wait synchronously that all tasks are done
 	 */
-	void WaitUntilAllTasksDone(bool bCancel);
+	INTERCHANGEENGINE_API void WaitUntilAllTasksDone(bool bCancel);
 
 
 	/**
 	 * If we set the mode to active we will setup the timer and add the thread that will block the GC.
 	 * If the we set the mode to inactive we will remove the timer and finish the thread that block the GC.
 	 */
-	void SetActiveMode(bool IsActive);
+	INTERCHANGEENGINE_API void SetActiveMode(bool IsActive);
 
 	/**
 	 * Start task until we reach the taskgraph worker number.
 	 * @param bCancelAllTasks - If true we will start all task but with the cancel state set, so task will complete fast and call the completion task
 	 */
-	void StartQueuedTasks(bool bCancelAllTasks = false);
+	INTERCHANGEENGINE_API void StartQueuedTasks(bool bCancelAllTasks = false);
 
 	/**
 	 * Called by the public Import functions
 	 */
-	TTuple<UE::Interchange::FAssetImportResultRef, UE::Interchange::FSceneImportResultRef>
+	INTERCHANGEENGINE_API TTuple<UE::Interchange::FAssetImportResultRef, UE::Interchange::FSceneImportResultRef>
 	ImportInternal(const FString& ContentPath, const UInterchangeSourceData* SourceData, const FImportAssetParameters& ImportAssetParameters, const UE::Interchange::EImportType ImportType);
 
 private:
diff --git a/Engine/Source/Runtime/Interchange/Engine/Public/InterchangeMeshUtilities.h b/Engine/Source/Runtime/Interchange/Engine/Public/InterchangeMeshUtilities.h
--- a/Engine/Source/Runtime/Interchange/Engine/Public/InterchangeMeshUtilities.h
+++ b/Engine/Source/Runtime/Interchange/Engine/Public/InterchangeMeshUtilities.h
@@ -10,8 +10,8 @@
 
 class UInterchangeSourceData;
 
-UCLASS(Experimental)
-class INTERCHANGEENGINE_API UInterchangeMeshUtilities : public UObject
+UCLASS(Experimental, MinimalAPI)
+class UInterchangeMeshUtilities : public UObject
 {
 	GENERATED_BODY()
 public:
@@ -24,7 +24,7 @@ public:
 	 *
 	 * @Note - This function will sort out the SourceData and call ImportCustomLodAsync and wait until the result is available.
 	 */
-	static TFuture<bool> ImportCustomLodAsync(UObject* MeshObject, const int32 LodIndex);
+	static INTERCHANGEENGINE_API TFuture<bool> ImportCustomLodAsync(UObject* MeshObject, const int32 LodIndex);
 
 
 	/**
@@ -38,7 +38,7 @@ public:
 	 * @Note - Since this is the async version, the SourceData parameter is required.
 	 * @Note - This function will search for an available interchange asset factory that can do the job for the MeshObject class
 	 */
-	static TFuture<bool> ImportCustomLodAsync(UObject* MeshObject, const int32 LodIndex, const UInterchangeSourceData* SourceData);
+	static INTERCHANGEENGINE_API TFuture<bool> ImportCustomLodAsync(UObject* MeshObject, const int32 LodIndex, const UInterchangeSourceData* SourceData);
 
 private:
 
diff --git a/Engine/Source/Runtime/Interchange/Engine/Public/InterchangePipelineConfigurationBase.h b/Engine/Source/Runtime/Interchange/Engine/Public/InterchangePipelineConfigurationBase.h
--- a/Engine/Source/Runtime/Interchange/Engine/Public/InterchangePipelineConfigurationBase.h
+++ b/Engine/Source/Runtime/Interchange/Engine/Public/InterchangePipelineConfigurationBase.h
@@ -29,8 +29,8 @@ struct FInterchangeStackInfo
 	TArray<TObjectPtr<UInterchangePipelineBase>> Pipelines;
 };
 
-UCLASS(BlueprintType, Blueprintable, Experimental)
-class INTERCHANGEENGINE_API UInterchangePipelineConfigurationBase : public UObject
+UCLASS(BlueprintType, Blueprintable, Experimental, MinimalAPI)
+class UInterchangePipelineConfigurationBase : public UObject
 {
 	GENERATED_BODY()
 
@@ -40,7 +40,7 @@ public:
 	 * Non virtual helper to allow blueprint to implement event base function to implement the ShowPipelineConfigurationDialog,
 	 */
 	UFUNCTION(BlueprintNativeEvent, BlueprintCallable, Category = "Interchange | Translator")
-	EInterchangePipelineConfigurationDialogResult ScriptedShowPipelineConfigurationDialog(TArray<FInterchangeStackInfo>& PipelineStacks
+	INTERCHANGEENGINE_API EInterchangePipelineConfigurationDialogResult ScriptedShowPipelineConfigurationDialog(TArray<FInterchangeStackInfo>& PipelineStacks
 		, TArray<UInterchangePipelineBase*>& OutPipelines
 		, UInterchangeSourceData* SourceData);
 
@@ -57,7 +57,7 @@ public:
 	 * Non virtual helper to allow blueprint to implement event base function to implement the ShowPipelineConfigurationDialog,
 	 */
 	UFUNCTION(BlueprintNativeEvent, BlueprintCallable, Category = "Interchange | Translator")
-	EInterchangePipelineConfigurationDialogResult ScriptedShowScenePipelineConfigurationDialog(TArray<FInterchangeStackInfo>& PipelineStacks
+	INTERCHANGEENGINE_API EInterchangePipelineConfigurationDialogResult ScriptedShowScenePipelineConfigurationDialog(TArray<FInterchangeStackInfo>& PipelineStacks
 		, TArray<UInterchangePipelineBase*>& OutPipelines
 		, UInterchangeSourceData* SourceData);
 
@@ -74,7 +74,7 @@ public:
 	 * Non virtual helper to allow blueprint to implement event base function to implement the ShowPipelineConfigurationDialog,
 	 */
 	UFUNCTION(BlueprintNativeEvent, BlueprintCallable, Category = "Interchange | Translator")
-	EInterchangePipelineConfigurationDialogResult ScriptedShowReimportPipelineConfigurationDialog(TArray<FInterchangeStackInfo>& PipelineStacks
+	INTERCHANGEENGINE_API EInterchangePipelineConfigurationDialogResult ScriptedShowReimportPipelineConfigurationDialog(TArray<FInterchangeStackInfo>& PipelineStacks
 		, TArray<UInterchangePipelineBase*>& OutPipelines
 		, UInterchangeSourceData* SourceData);
 
diff --git a/Engine/Source/Runtime/Interchange/Engine/Public/InterchangeProjectSettings.h b/Engine/Source/Runtime/Interchange/Engine/Public/InterchangeProjectSettings.h
--- a/Engine/Source/Runtime/Interchange/Engine/Public/InterchangeProjectSettings.h
+++ b/Engine/Source/Runtime/Interchange/Engine/Public/InterchangeProjectSettings.h
@@ -112,16 +112,16 @@ public:
 	TSoftClassPtr <UInterchangePipelineBase> GenericPipelineClass;
 };
 
-class INTERCHANGEENGINE_API FInterchangeProjectSettingsUtils
+class FInterchangeProjectSettingsUtils
 {
 public:
-	static const FInterchangeImportSettings& GetImportSettings(const UInterchangeProjectSettings& InterchangeProjectSettings, const bool bIsSceneImport);
-	static FInterchangeImportSettings& GetMutableImportSettings(UInterchangeProjectSettings& InterchangeProjectSettings, const bool bIsSceneImport);
-	static const FInterchangeImportSettings& GetDefaultImportSettings(const bool bIsSceneImport);
-	static FInterchangeImportSettings& GetMutableDefaultImportSettings(const bool bIsSceneImport);
+	static INTERCHANGEENGINE_API const FInterchangeImportSettings& GetImportSettings(const UInterchangeProjectSettings& InterchangeProjectSettings, const bool bIsSceneImport);
+	static INTERCHANGEENGINE_API FInterchangeImportSettings& GetMutableImportSettings(UInterchangeProjectSettings& InterchangeProjectSettings, const bool bIsSceneImport);
+	static INTERCHANGEENGINE_API const FInterchangeImportSettings& GetDefaultImportSettings(const bool bIsSceneImport);
+	static INTERCHANGEENGINE_API FInterchangeImportSettings& GetMutableDefaultImportSettings(const bool bIsSceneImport);
 
-	static FName GetDefaultPipelineStackName(const bool bIsSceneImport, const UInterchangeSourceData& SourceData);
-	static void SetDefaultPipelineStackName(const bool bIsSceneImport, const UInterchangeSourceData& SourceData, const FName StackName);
+	static INTERCHANGEENGINE_API FName GetDefaultPipelineStackName(const bool bIsSceneImport, const UInterchangeSourceData& SourceData);
+	static INTERCHANGEENGINE_API void SetDefaultPipelineStackName(const bool bIsSceneImport, const UInterchangeSourceData& SourceData, const FName StackName);
 
-	static bool ShouldShowPipelineStacksConfigurationDialog(const bool bIsSceneImport, const UInterchangeSourceData& SourceData);
-};
\ No newline at end of file
+	static INTERCHANGEENGINE_API bool ShouldShowPipelineStacksConfigurationDialog(const bool bIsSceneImport, const UInterchangeSourceData& SourceData);
+};
diff --git a/Engine/Source/Runtime/Interchange/Engine/Public/InterchangePythonPipelineBase.h b/Engine/Source/Runtime/Interchange/Engine/Public/InterchangePythonPipelineBase.h
--- a/Engine/Source/Runtime/Interchange/Engine/Public/InterchangePythonPipelineBase.h
+++ b/Engine/Source/Runtime/Interchange/Engine/Public/InterchangePythonPipelineBase.h
@@ -13,14 +13,14 @@
  * This class represent a python pipeline. It is use by the TSoftClassPtr of the PythonPipeline asset.
  *
  */
-UCLASS(BlueprintType, Abstract, Experimental)
-class INTERCHANGEENGINE_API UInterchangePythonPipelineBase : public UInterchangePipelineBase
+UCLASS(BlueprintType, Abstract, Experimental, MinimalAPI)
+class UInterchangePythonPipelineBase : public UInterchangePipelineBase
 {
 	GENERATED_BODY()
 public:
 
 #if WITH_EDITOR
-	virtual void PostEditChangeProperty(struct FPropertyChangedEvent& PropertyChangedEvent) override;
+	INTERCHANGEENGINE_API virtual void PostEditChangeProperty(struct FPropertyChangedEvent& PropertyChangedEvent) override;
 #endif
 };
 
@@ -43,20 +43,20 @@ struct FPropertyData
  * This restriction exist because python class are transient, so any assets create from a python class cannot be save.
  * 
  */
-UCLASS(BlueprintType, Experimental)
-class INTERCHANGEENGINE_API UInterchangePythonPipelineAsset : public UObject
+UCLASS(BlueprintType, Experimental, MinimalAPI)
+class UInterchangePythonPipelineAsset : public UObject
 {
 	GENERATED_BODY()
 
 public:
-	virtual void PostLoad() override;
+	INTERCHANGEENGINE_API virtual void PostLoad() override;
 #if WITH_EDITOR
-	virtual void PostEditChangeProperty(struct FPropertyChangedEvent& PropertyChangedEvent) override;
+	INTERCHANGEENGINE_API virtual void PostEditChangeProperty(struct FPropertyChangedEvent& PropertyChangedEvent) override;
 #endif
 
-	void GeneratePipeline();
+	INTERCHANGEENGINE_API void GeneratePipeline();
 
-	void SetupFromPipeline(const UInterchangePythonPipelineBase* PythonPipeline, const bool bRegeneratePipeline = true);
+	INTERCHANGEENGINE_API void SetupFromPipeline(const UInterchangePythonPipelineBase* PythonPipeline, const bool bRegeneratePipeline = true);
 
 	/** The python class we want to use has a pipeline */
 	UPROPERTY(EditAnywhere, Category = "Interchange|Python")
diff --git a/Engine/Source/Runtime/Interchange/Engine/Public/InterchangeSceneImportAsset.h b/Engine/Source/Runtime/Interchange/Engine/Public/InterchangeSceneImportAsset.h
--- a/Engine/Source/Runtime/Interchange/Engine/Public/InterchangeSceneImportAsset.h
+++ b/Engine/Source/Runtime/Interchange/Engine/Public/InterchangeSceneImportAsset.h
@@ -14,12 +14,12 @@ class UInterchangeFactoryBaseNode;
 /*
  * Class to hold all the data required to properly re-import a level
  */
-UCLASS()
-class INTERCHANGEENGINE_API UInterchangeSceneImportAsset : public UObject, public IInterface_AssetUserData
+UCLASS(MinimalAPI)
+class UInterchangeSceneImportAsset : public UObject, public IInterface_AssetUserData
 {
 	GENERATED_BODY()
 
-	virtual ~UInterchangeSceneImportAsset();
+	INTERCHANGEENGINE_API virtual ~UInterchangeSceneImportAsset();
 
 public:
 #if WITH_EDITORONLY_DATA
@@ -33,23 +33,23 @@ public:
 #endif // #if WITH_EDITORONLY_DATA
 
 	//~ Begin UObject Interface
-	virtual void GetAssetRegistryTags(TArray<FAssetRegistryTag>& OutTags) const override;
-	virtual void PostLoad() override;
+	INTERCHANGEENGINE_API virtual void GetAssetRegistryTags(TArray<FAssetRegistryTag>& OutTags) const override;
+	INTERCHANGEENGINE_API virtual void PostLoad() override;
 	//~ End UObject Interface
 
 	//~ Begin IInterface_AssetUserData Interface
-	virtual void AddAssetUserData( UAssetUserData* InUserData ) override;
-	virtual void RemoveUserDataOfClass( TSubclassOf<UAssetUserData> InUserDataClass ) override;
-	virtual UAssetUserData* GetAssetUserDataOfClass( TSubclassOf<UAssetUserData> InUserDataClass ) override;
-	virtual const TArray<UAssetUserData*>* GetAssetUserDataArray() const override;
+	INTERCHANGEENGINE_API virtual void AddAssetUserData( UAssetUserData* InUserData ) override;
+	INTERCHANGEENGINE_API virtual void RemoveUserDataOfClass( TSubclassOf<UAssetUserData> InUserDataClass ) override;
+	INTERCHANGEENGINE_API virtual UAssetUserData* GetAssetUserDataOfClass( TSubclassOf<UAssetUserData> InUserDataClass ) override;
+	INTERCHANGEENGINE_API virtual const TArray<UAssetUserData*>* GetAssetUserDataArray() const override;
 	//~ End IInterface_AssetUserData Interface
 
 #if WITH_EDITOR
-	void RegisterWorldRenameCallbacks();
+	INTERCHANGEENGINE_API void RegisterWorldRenameCallbacks();
 #endif
 
 	/** Updates the SceneObjects cache based on the node container stored in AssetImportData */
-	void UpdateSceneObjects();
+	INTERCHANGEENGINE_API void UpdateSceneObjects();
 
 	/**
 	 * Returns the UObject which asset path is '//PackageName.AssetName[:SubPathString]'.
@@ -59,7 +59,7 @@ public:
 	 * @oaram AssetName: Asset name of the actual object to reimport
 	 * @oaram SubPathString: Optional subobject name
 	 */
-	UObject* GetSceneObject(const FString& PackageName, const FString& AssetName, const FString& SubPathString = FString()) const;
+	INTERCHANGEENGINE_API UObject* GetSceneObject(const FString& PackageName, const FString& AssetName, const FString& SubPathString = FString()) const;
 
 	/**
 	 * Returns the factory node associated with the asset which path is '//PackageName.AssetName[:SubPathString]'.
@@ -69,7 +69,7 @@ public:
 	 * @oaram AssetName: Asset name of the actual object to reimport
 	 * @oaram SubPathString: Optional subobject name
 	 */
-	const UInterchangeFactoryBaseNode* GetFactoryNode(const FString& PackageName, const FString& AssetName, const FString& SubPathString = FString()) const;
+	INTERCHANGEENGINE_API const UInterchangeFactoryBaseNode* GetFactoryNode(const FString& PackageName, const FString& AssetName, const FString& SubPathString = FString()) const;
 
 private:
 
diff --git a/Engine/Source/Runtime/Json/Public/Dom/JsonObject.h b/Engine/Source/Runtime/Json/Public/Dom/JsonObject.h
--- a/Engine/Source/Runtime/Json/Public/Dom/JsonObject.h
+++ b/Engine/Source/Runtime/Json/Public/Dom/JsonObject.h
@@ -19,7 +19,7 @@
  * A Json Object is a structure holding an unordered set of name/value pairs.
  * In a Json file, it is represented by everything between curly braces {}.
  */
-class JSON_API FJsonObject
+class FJsonObject
 {
 public:
 
@@ -102,14 +102,14 @@ public:
 	 * @param FieldName The name of the field to set.
 	 * @param Value The value to set.
 	 */
-	void SetField( const FString& FieldName, const TSharedPtr<FJsonValue>& Value );
+	JSON_API void SetField( const FString& FieldName, const TSharedPtr<FJsonValue>& Value );
 
 	/**
 	 * Removes the field with the specified name.
 	 *
 	 * @param FieldName The name of the field to remove.
 	 */
-	void RemoveField(const FString& FieldName);
+	JSON_API void RemoveField(const FString& FieldName);
 
 	/**
 	 * Gets the field with the specified name as a number.
@@ -119,7 +119,7 @@ public:
 	 * @param FieldName The name of the field to get.
 	 * @return The field's value as a number.
 	 */
-	double GetNumberField(const FString& FieldName) const;
+	JSON_API double GetNumberField(const FString& FieldName) const;
 
 	/**
 	 * Gets a numeric field and casts to an int32
@@ -130,46 +130,46 @@ public:
 	}
 
 	/** Get the field named FieldName as a number. Returns false if it doesn't exist or cannot be converted. */
-	bool TryGetNumberField(const FString& FieldName, float& OutNumber) const;
+	JSON_API bool TryGetNumberField(const FString& FieldName, float& OutNumber) const;
 
 	/** Get the field named FieldName as a number. Returns false if it doesn't exist or cannot be converted. */
-	bool TryGetNumberField(const FString& FieldName, double& OutNumber) const;
+	JSON_API bool TryGetNumberField(const FString& FieldName, double& OutNumber) const;
 
 	/** Get the field named FieldName as a number, and makes sure it's within int8 range. Returns false if it doesn't exist or cannot be converted. */
-	bool TryGetNumberField(const FString& FieldName, int8& OutNumber) const;
+	JSON_API bool TryGetNumberField(const FString& FieldName, int8& OutNumber) const;
 
 	/** Get the field named FieldName as a number, and makes sure it's within int16 range. Returns false if it doesn't exist or cannot be converted. */
-	bool TryGetNumberField(const FString& FieldName, int16& OutNumber) const;
+	JSON_API bool TryGetNumberField(const FString& FieldName, int16& OutNumber) const;
 
 	/** Get the field named FieldName as a number, and makes sure it's within int32 range. Returns false if it doesn't exist or cannot be converted. */
-	bool TryGetNumberField(const FString& FieldName, int32& OutNumber) const;
+	JSON_API bool TryGetNumberField(const FString& FieldName, int32& OutNumber) const;
 
 	/** Get the field named FieldName as a number. Returns false if it doesn't exist or cannot be converted. */
-	bool TryGetNumberField(const FString& FieldName, int64& OutNumber) const;
+	JSON_API bool TryGetNumberField(const FString& FieldName, int64& OutNumber) const;
 
 	/** Get the field named FieldName as a number, and makes sure it's within uint8 range. Returns false if it doesn't exist or cannot be converted. */
-	bool TryGetNumberField(const FString& FieldName, uint8& OutNumber) const;
+	JSON_API bool TryGetNumberField(const FString& FieldName, uint8& OutNumber) const;
 
 	/** Get the field named FieldName as a number, and makes sure it's within uint16 range. Returns false if it doesn't exist or cannot be converted. */
-	bool TryGetNumberField(const FString& FieldName, uint16& OutNumber) const;
+	JSON_API bool TryGetNumberField(const FString& FieldName, uint16& OutNumber) const;
 
 	/** Get the field named FieldName as a number, and makes sure it's within uint32 range. Returns false if it doesn't exist or cannot be converted.  */
-	bool TryGetNumberField(const FString& FieldName, uint32& OutNumber) const;
+	JSON_API bool TryGetNumberField(const FString& FieldName, uint32& OutNumber) const;
 
 	/** Get the field named FieldName as a number, and makes sure it's within uint64 range. Returns false if it doesn't exist or cannot be converted.  */
-	bool TryGetNumberField(const FString& FieldName, uint64& OutNumber) const;
+	JSON_API bool TryGetNumberField(const FString& FieldName, uint64& OutNumber) const;
 
 	/** Add a field named FieldName with Number as value */
-	void SetNumberField( const FString& FieldName, double Number );
+	JSON_API void SetNumberField( const FString& FieldName, double Number );
 
 	/** Get the field named FieldName as a string. */
-	FString GetStringField(const FString& FieldName) const;
+	JSON_API FString GetStringField(const FString& FieldName) const;
 
 	/** Get the field named FieldName as a string. Returns false if it doesn't exist or cannot be converted. */
-	bool TryGetStringField(const FString& FieldName, FString& OutString) const;
+	JSON_API bool TryGetStringField(const FString& FieldName, FString& OutString) const;
 
 	/** Get the field named FieldName as an array of strings. Returns false if it doesn't exist or any member cannot be converted. */
-	bool TryGetStringArrayField(const FString& FieldName, TArray<FString>& OutArray) const;
+	JSON_API bool TryGetStringArrayField(const FString& FieldName, TArray<FString>& OutArray) const;
 
 	/** Get the field named FieldName as an array of enums. Returns false if it doesn't exist or any member is not a string. */
 	template<typename TEnum>
@@ -194,7 +194,7 @@ public:
 	}
 
 	/** Add a field named FieldName with value of StringValue */
-	void SetStringField( const FString& FieldName, const FString& StringValue );
+	JSON_API void SetStringField( const FString& FieldName, const FString& StringValue );
 
 	/**
 	 * Gets the field with the specified name as a boolean.
@@ -204,22 +204,22 @@ public:
 	 * @param FieldName The name of the field to get.
 	 * @return The field's value as a boolean.
 	 */
-	bool GetBoolField(const FString& FieldName) const;
+	JSON_API bool GetBoolField(const FString& FieldName) const;
 
 	/** Get the field named FieldName as a string. Returns false if it doesn't exist or cannot be converted. */
-	bool TryGetBoolField(const FString& FieldName, bool& OutBool) const;
+	JSON_API bool TryGetBoolField(const FString& FieldName, bool& OutBool) const;
 
 	/** Set a boolean field named FieldName and value of InValue */
-	void SetBoolField( const FString& FieldName, bool InValue );
+	JSON_API void SetBoolField( const FString& FieldName, bool InValue );
 
 	/** Get the field named FieldName as an array. */
-	const TArray< TSharedPtr<FJsonValue> >& GetArrayField(const FString& FieldName) const;
+	JSON_API const TArray< TSharedPtr<FJsonValue> >& GetArrayField(const FString& FieldName) const;
 
 	/** Try to get the field named FieldName as an array, or return false if it's another type */
-	bool TryGetArrayField(const FString& FieldName, const TArray< TSharedPtr<FJsonValue> >*& OutArray) const;
+	JSON_API bool TryGetArrayField(const FString& FieldName, const TArray< TSharedPtr<FJsonValue> >*& OutArray) const;
 
 	/** Set an array field named FieldName and value of Array */
-	void SetArrayField( const FString& FieldName, const TArray< TSharedPtr<FJsonValue> >& Array );
+	JSON_API void SetArrayField( const FString& FieldName, const TArray< TSharedPtr<FJsonValue> >& Array );
 
 	/**
 	 * Gets the field with the specified name as a Json object.
@@ -229,14 +229,14 @@ public:
 	 * @param FieldName The name of the field to get.
 	 * @return The field's value as a Json object.
 	 */
-	const TSharedPtr<FJsonObject>& GetObjectField(const FString& FieldName) const;
+	JSON_API const TSharedPtr<FJsonObject>& GetObjectField(const FString& FieldName) const;
 
 	/** Try to get the field named FieldName as an object, or return false if it's another type */
-	bool TryGetObjectField(const FString& FieldName, const TSharedPtr<FJsonObject>*& OutObject) const;
+	JSON_API bool TryGetObjectField(const FString& FieldName, const TSharedPtr<FJsonObject>*& OutObject) const;
 
 	/** Set an ObjectField named FieldName and value of JsonObject */
-	void SetObjectField( const FString& FieldName, const TSharedPtr<FJsonObject>& JsonObject );
+	JSON_API void SetObjectField( const FString& FieldName, const TSharedPtr<FJsonObject>& JsonObject );
 
-	static void Duplicate(const TSharedPtr<const FJsonObject>& Source, const TSharedPtr<FJsonObject>& Dest);
-	static void Duplicate(const TSharedPtr<FJsonObject>& Source, TSharedPtr<FJsonObject>& Dest);
+	static JSON_API void Duplicate(const TSharedPtr<const FJsonObject>& Source, const TSharedPtr<FJsonObject>& Dest);
+	static JSON_API void Duplicate(const TSharedPtr<FJsonObject>& Source, TSharedPtr<FJsonObject>& Dest);
 };
diff --git a/Engine/Source/Runtime/Json/Public/Dom/JsonValue.h b/Engine/Source/Runtime/Json/Public/Dom/JsonValue.h
--- a/Engine/Source/Runtime/Json/Public/Dom/JsonValue.h
+++ b/Engine/Source/Runtime/Json/Public/Dom/JsonValue.h
@@ -16,54 +16,54 @@ class FJsonObject;
  * A Json Value is a structure that can be any of the Json Types.
  * It should never be used on its, only its derived types should be used.
  */
-class JSON_API FJsonValue
+class FJsonValue
 {
 public:
 
 	/** Returns this value as a double, logging an error and returning zero if this is not an Json Number */
-	double AsNumber() const;
+	JSON_API double AsNumber() const;
 
 	/** Returns this value as a string, logging an error and returning an empty string if not possible */
-	FString AsString() const;
+	JSON_API FString AsString() const;
 
 	/** Returns this value as a boolean, logging an error and returning false if not possible */
-	bool AsBool() const;
+	JSON_API bool AsBool() const;
 
 	/** Returns this value as an array, logging an error and returning an empty array reference if not possible */
-	const TArray< TSharedPtr<FJsonValue> >& AsArray() const;
+	JSON_API const TArray< TSharedPtr<FJsonValue> >& AsArray() const;
 
 	/** Returns this value as an object, throwing an error if this is not an Json Object */
-	virtual const TSharedPtr<FJsonObject>& AsObject() const;
+	JSON_API virtual const TSharedPtr<FJsonObject>& AsObject() const;
 
 	/** Tries to convert this value to a number, returning false if not possible */
 	virtual bool TryGetNumber(double& OutNumber) const { return false; }
 
 	/** Tries to convert this value to a number, returning false if not possible */
-	virtual bool TryGetNumber(float& OutNumber) const;
+	JSON_API virtual bool TryGetNumber(float& OutNumber) const;
 
 	/** Tries to convert this value to a number, returning false if not possible */
-	virtual bool TryGetNumber(int8& OutNumber) const;
+	JSON_API virtual bool TryGetNumber(int8& OutNumber) const;
 
 	/** Tries to convert this value to a number, returning false if not possible */
-	virtual bool TryGetNumber(int16& OutNumber) const;
+	JSON_API virtual bool TryGetNumber(int16& OutNumber) const;
 
 	/** Tries to convert this value to a number, returning false if not possible */
-	virtual bool TryGetNumber(int32& OutNumber) const;
+	JSON_API virtual bool TryGetNumber(int32& OutNumber) const;
 
 	/** Tries to convert this value to a number, returning false if not possible */
-	virtual bool TryGetNumber(int64& OutNumber) const;
+	JSON_API virtual bool TryGetNumber(int64& OutNumber) const;
 
 	/** Tries to convert this value to a number, returning false if not possible */
-	virtual bool TryGetNumber(uint8& OutNumber) const;
+	JSON_API virtual bool TryGetNumber(uint8& OutNumber) const;
 
 	/** Tries to convert this value to a number, returning false if not possible */
-	virtual bool TryGetNumber(uint16& OutNumber) const;
+	JSON_API virtual bool TryGetNumber(uint16& OutNumber) const;
 
 	/** Tries to convert this value to a number, returning false if not possible */
-	virtual bool TryGetNumber(uint32& OutNumber) const;
+	JSON_API virtual bool TryGetNumber(uint32& OutNumber) const;
 
 	/** Tries to convert this value to a number, returning false if not possible */
-	virtual bool TryGetNumber(uint64& OutNumber) const;
+	JSON_API virtual bool TryGetNumber(uint64& OutNumber) const;
 
 	/** Tries to convert this value to a string, returning false if not possible */
 	virtual bool TryGetString(FString& OutString) const { return false; }
@@ -98,9 +98,9 @@ public:
 
 	EJson Type;
 
-	static TSharedPtr<FJsonValue> Duplicate(const TSharedPtr<FJsonValue>& Src);
+	static JSON_API TSharedPtr<FJsonValue> Duplicate(const TSharedPtr<FJsonValue>& Src);
 
-	static bool CompareEqual(const FJsonValue& Lhs, const FJsonValue& Rhs);
+	static JSON_API bool CompareEqual(const FJsonValue& Lhs, const FJsonValue& Rhs);
 
 protected:
 
@@ -109,7 +109,7 @@ protected:
 
 	virtual FString GetType() const = 0;
 
-	void ErrorMessage(const FString& InType) const;
+	JSON_API void ErrorMessage(const FString& InType) const;
 
 	friend inline bool operator==(const FJsonValue& Lhs, const FJsonValue& Rhs)
 	{
@@ -124,7 +124,7 @@ protected:
 
 
 /** A Json String Value. */
-class JSON_API FJsonValueString : public FJsonValue
+class FJsonValueString : public FJsonValue
 {
 public:
 	FJsonValueString(const FString& InString) : Value(InString) {Type = EJson::String;}
@@ -149,7 +149,7 @@ protected:
 
 
 /** A Json Number Value. */
-class JSON_API FJsonValueNumber : public FJsonValue
+class FJsonValueNumber : public FJsonValue
 {
 public:
 	FJsonValueNumber(double InNumber) : Value(InNumber) {Type = EJson::Number;}
@@ -166,7 +166,7 @@ protected:
 
 
 /** A Json Number Value, stored internally as a string so as not to lose precision */
-class JSON_API FJsonValueNumberString : public FJsonValue
+class FJsonValueNumberString : public FJsonValue
 {
 public:
 	FJsonValueNumberString(const FString& InString) : Value(InString) { Type = EJson::Number; }
@@ -194,7 +194,7 @@ protected:
 
 
 /** A Json Boolean Value. */
-class JSON_API FJsonValueBoolean : public FJsonValue
+class FJsonValueBoolean : public FJsonValue
 {
 public:
 	FJsonValueBoolean(bool InBool) : Value(InBool) {Type = EJson::Boolean;}
@@ -210,7 +210,7 @@ protected:
 
 
 /** A Json Array Value. */
-class JSON_API FJsonValueArray : public FJsonValue
+class FJsonValueArray : public FJsonValue
 {
 public:
 	FJsonValueArray(const TArray< TSharedPtr<FJsonValue> >& InArray) : Value(InArray) {Type = EJson::Array;}
@@ -226,7 +226,7 @@ protected:
 
 
 /** A Json Object Value. */
-class JSON_API FJsonValueObject : public FJsonValue
+class FJsonValueObject : public FJsonValue
 {
 public:
 	FJsonValueObject(TSharedPtr<FJsonObject> InObject) : Value(MoveTemp(InObject)) {Type = EJson::Object;}
@@ -241,7 +241,7 @@ protected:
 
 
 /** A Json Null Value. */
-class JSON_API FJsonValueNull : public FJsonValue
+class FJsonValueNull : public FJsonValue
 {
 public:
 	FJsonValueNull() {Type = EJson::Null;}
diff --git a/Engine/Source/Runtime/Json/Public/JsonUtils/JsonPointer.h b/Engine/Source/Runtime/Json/Public/JsonUtils/JsonPointer.h
--- a/Engine/Source/Runtime/Json/Public/JsonUtils/JsonPointer.h
+++ b/Engine/Source/Runtime/Json/Public/JsonUtils/JsonPointer.h
@@ -1,4 +1,4 @@
-﻿// Copyright Epic Games, Inc. All Rights Reserved.
+// Copyright Epic Games, Inc. All Rights Reserved.
 
 #pragma once
 
@@ -21,22 +21,22 @@ namespace UE::Json
 	* @see https://datatracker.ietf.org/doc/html/rfc6901
 	* @note Array append ("-") is currently unsupported.
 	*/
-	class JSON_API FJsonPointer
+	class FJsonPointer
 	{
 	public:
 		FJsonPointer() = default;
-		FJsonPointer(FStringView Path);
-		FJsonPointer(const FJsonPointer& Other);
-		FJsonPointer(FJsonPointer&& Other) noexcept;
+		JSON_API FJsonPointer(FStringView Path);
+		JSON_API FJsonPointer(const FJsonPointer& Other);
+		JSON_API FJsonPointer(FJsonPointer&& Other) noexcept;
 
 		/** Attempt to retrieve the value specified by this Json Pointer from the provided Json Value. */
-		bool TryGet(const TSharedPtr<FJsonObject>& RootObject, TSharedPtr<FJsonValue>& OutValue) const;
+		JSON_API bool TryGet(const TSharedPtr<FJsonObject>& RootObject, TSharedPtr<FJsonValue>& OutValue) const;
 
 		/** Attempt to retrieve the value specified by this Json Pointer from the provided Json Object. */
-		bool TryGet(const TSharedPtr<FJsonValue>& RootValue, TSharedPtr<FJsonValue>& OutValue) const;
+		JSON_API bool TryGet(const TSharedPtr<FJsonValue>& RootValue, TSharedPtr<FJsonValue>& OutValue) const;
 
 		/** Attempt to retrieve the value specified by this Json Pointer from the provided Json Array. */
-		bool TryGet(const TArray<TSharedPtr<FJsonValue>>& RootValues, TSharedPtr<FJsonValue>& OutValue) const;
+		JSON_API bool TryGet(const TArray<TSharedPtr<FJsonValue>>& RootValues, TSharedPtr<FJsonValue>& OutValue) const;
 
 		/** Checks that the path has one or more parts. */
 		FORCEINLINE bool IsValid() const
@@ -51,16 +51,16 @@ namespace UE::Json
 		}
 
 		/** Appends the given part to the path */
-		void Append(FStringView Part);
+		JSON_API void Append(FStringView Part);
 
 		/** Sets the last path part to the given array index.  */
-		void SetArrayIndex(int32 Index);
+		JSON_API void SetArrayIndex(int32 Index);
 
 		/** Gets the sub-path up to the given part length. */
-		FJsonPointer GetSubPath(int32 NumParts) const;
+		JSON_API FJsonPointer GetSubPath(int32 NumParts) const;
 
 		/** Get the JsonPointer's string representation. */
-		FString ToString() const;
+		JSON_API FString ToString() const;
 
 		/** Get the path part at the given index. */
 		FORCEINLINE const FString& operator[](int32 Idx) const
@@ -75,14 +75,14 @@ namespace UE::Json
 		}
 
 		/** Substitutes characters in the provided Part with their escaped counterparts (as per the spec). */
-		static FString EscapePart(const FString& Part);
+		static JSON_API FString EscapePart(const FString& Part);
 		
 		/** Substitutes escaped characters in the Provided path with their original character (as per the spec). */
-		static FString UnescapePart(const FString& Part);
+		static JSON_API FString UnescapePart(const FString& Part);
 
 	private:
 		/** Parses the provided path, populating PathParts and validating the result. */
-		bool ParsePath(const FString& InPath);
+		JSON_API bool ParsePath(const FString& InPath);
 
 	private:
 		TArray<FString> PathParts;
diff --git a/Engine/Source/Runtime/JsonUtilities/Public/JsonObjectConverter.h b/Engine/Source/Runtime/JsonUtilities/Public/JsonObjectConverter.h
--- a/Engine/Source/Runtime/JsonUtilities/Public/JsonObjectConverter.h
+++ b/Engine/Source/Runtime/JsonUtilities/Public/JsonObjectConverter.h
@@ -27,18 +27,18 @@ class FProperty;
 class UStruct;
 
 /** Class that handles converting Json objects to and from UStructs */
-class JSONUTILITIES_API FJsonObjectConverter
+class FJsonObjectConverter
 {
 public:
 
 	/** FName case insensitivity can make the casing of UPROPERTIES unpredictable. Attempt to standardize output. */
-	static FString StandardizeCase(const FString &StringIn);
+	static JSONUTILITIES_API FString StandardizeCase(const FString &StringIn);
 
 	/** Parse an FText from a json object (assumed to be of the form where keys are culture codes and values are strings) */
-	static bool GetTextFromObject(const TSharedRef<FJsonObject>& Obj, FText& TextOut);
+	static JSONUTILITIES_API bool GetTextFromObject(const TSharedRef<FJsonObject>& Obj, FText& TextOut);
 
 	/** Convert a Json value to text (takes some hints from the value name) */
-	static bool GetTextFromField(const FString& FieldName, const TSharedPtr<FJsonValue>& FieldValue, FText& TextOut);
+	static JSONUTILITIES_API bool GetTextFromField(const FString& FieldName, const TSharedPtr<FJsonValue>& FieldValue, FText& TextOut);
 
 public: // UStruct -> JSON
 
@@ -49,13 +49,13 @@ public: // UStruct -> JSON
 	 */
 	DECLARE_DELEGATE_RetVal_TwoParams(TSharedPtr<FJsonValue>, CustomExportCallback, FProperty* /* Property */, const void* /* Value */);
 
-	static const CustomExportCallback ExportCallback_WriteISO8601Dates;
+	static JSONUTILITIES_API const CustomExportCallback ExportCallback_WriteISO8601Dates;
 
 	/**
 	 * Utility Export Callback for having object properties expanded to full Json.
 	 */
 	UE_DEPRECATED(4.25, "ObjectJsonCallback has been deprecated - please remove the usage of it from your project")
-	static TSharedPtr<FJsonValue> ObjectJsonCallback(FProperty* Property , const void* Value);
+	static JSONUTILITIES_API TSharedPtr<FJsonValue> ObjectJsonCallback(FProperty* Property , const void* Value);
 
 	/**
 	 * Templated version of UStructToJsonObject to try and make most of the params. Also serves as an example use case
@@ -89,7 +89,7 @@ public: // UStruct -> JSON
 	 *
 	 * @return False if any properties failed to write
 	 */
-	static bool UStructToJsonObject(const UStruct* StructDefinition, const void* Struct, TSharedRef<FJsonObject> OutJsonObject, int64 CheckFlags = 0, int64 SkipFlags = 0, const CustomExportCallback* ExportCb = nullptr);
+	static JSONUTILITIES_API bool UStructToJsonObject(const UStruct* StructDefinition, const void* Struct, TSharedRef<FJsonObject> OutJsonObject, int64 CheckFlags = 0, int64 SkipFlags = 0, const CustomExportCallback* ExportCb = nullptr);
 
 	/**
 	 * Converts from a UStruct to a json string containing an object, using exportText
@@ -105,7 +105,7 @@ public: // UStruct -> JSON
 	 *
 	 * @return False if any properties failed to write
 	 */
-	static bool UStructToJsonObjectString(const UStruct* StructDefinition, const void* Struct, FString& OutJsonString, int64 CheckFlags = 0, int64 SkipFlags = 0, int32 Indent = 0, const CustomExportCallback* ExportCb = nullptr, bool bPrettyPrint = true);
+	static JSONUTILITIES_API bool UStructToJsonObjectString(const UStruct* StructDefinition, const void* Struct, FString& OutJsonString, int64 CheckFlags = 0, int64 SkipFlags = 0, int32 Indent = 0, const CustomExportCallback* ExportCb = nullptr, bool bPrettyPrint = true);
 
 	/**
 	 * Templated version; Converts from a UStruct to a json string containing an object, using exportText
@@ -164,7 +164,7 @@ public: // UStruct -> JSON
 	 *
 	 * @return False if any properties failed to write
 	 */
-	static bool UStructToJsonAttributes(const UStruct* StructDefinition, const void* Struct, TMap< FString, TSharedPtr<FJsonValue> >& OutJsonAttributes, int64 CheckFlags = 0, int64 SkipFlags = 0, const CustomExportCallback* ExportCb = nullptr);
+	static JSONUTILITIES_API bool UStructToJsonAttributes(const UStruct* StructDefinition, const void* Struct, TMap< FString, TSharedPtr<FJsonValue> >& OutJsonAttributes, int64 CheckFlags = 0, int64 SkipFlags = 0, const CustomExportCallback* ExportCb = nullptr);
 
 	/* * Converts from a FProperty to a Json Value using exportText
 	 *
@@ -177,7 +177,7 @@ public: // UStruct -> JSON
 	 *
 	 * @return					The constructed JsonValue from the property
 	 */
-	static TSharedPtr<FJsonValue> UPropertyToJsonValue(FProperty* Property, const void* Value, int64 CheckFlags = 0, int64 SkipFlags = 0, const CustomExportCallback* ExportCb = nullptr, FProperty* OuterProperty = nullptr);
+	static JSONUTILITIES_API TSharedPtr<FJsonValue> UPropertyToJsonValue(FProperty* Property, const void* Value, int64 CheckFlags = 0, int64 SkipFlags = 0, const CustomExportCallback* ExportCb = nullptr, FProperty* OuterProperty = nullptr);
 
 public: // JSON -> UStruct
 
@@ -194,7 +194,7 @@ public: // JSON -> UStruct
 	 *
 	 * @return False if any properties matched but failed to deserialize
 	 */
-	static bool JsonObjectToUStruct(const TSharedRef<FJsonObject>& JsonObject, const UStruct* StructDefinition, void* OutStruct, int64 CheckFlags = 0, int64 SkipFlags = 0, const bool bStrictMode = false, FText* OutFailReason = nullptr);
+	static JSONUTILITIES_API bool JsonObjectToUStruct(const TSharedRef<FJsonObject>& JsonObject, const UStruct* StructDefinition, void* OutStruct, int64 CheckFlags = 0, int64 SkipFlags = 0, const bool bStrictMode = false, FText* OutFailReason = nullptr);
 
 	/**
 	 * Templated version of JsonObjectToUStruct
@@ -227,7 +227,7 @@ public: // JSON -> UStruct
 	 *
 	 * @return False if any properties matched but failed to deserialize
 	 */
-	static bool JsonAttributesToUStruct(const TMap< FString, TSharedPtr<FJsonValue> >& JsonAttributes, const UStruct* StructDefinition, void* OutStruct, int64 CheckFlags = 0, int64 SkipFlags = 0, const bool bStrictMode = false, FText* OutFailReason = nullptr);
+	static JSONUTILITIES_API bool JsonAttributesToUStruct(const TMap< FString, TSharedPtr<FJsonValue> >& JsonAttributes, const UStruct* StructDefinition, void* OutStruct, int64 CheckFlags = 0, int64 SkipFlags = 0, const bool bStrictMode = false, FText* OutFailReason = nullptr);
 
 	/**
 	 * Converts a single JsonValue to the corresponding FProperty (this may recurse if the property is a UStruct for instance).
@@ -242,7 +242,7 @@ public: // JSON -> UStruct
 	 *
 	 * @return False if the property failed to serialize
 	 */
-	static bool JsonValueToUProperty(const TSharedPtr<FJsonValue>& JsonValue, FProperty* Property, void* OutValue, int64 CheckFlags = 0, int64 SkipFlags = 0, const bool bStrictMode = false, FText* OutFailReason = nullptr);
+	static JSONUTILITIES_API bool JsonValueToUProperty(const TSharedPtr<FJsonValue>& JsonValue, FProperty* Property, void* OutValue, int64 CheckFlags = 0, int64 SkipFlags = 0, const bool bStrictMode = false, FText* OutFailReason = nullptr);
 
 	/**
 	 * Converts from a json string containing an object to a UStruct
@@ -338,5 +338,5 @@ public: // JSON -> UStruct
 	* Parses text arguments from Json into a map
 	* @param JsonObject Object to parse arguments from
 	*/
-	static FFormatNamedArguments ParseTextArgumentsFromJson(const TSharedPtr<const FJsonObject>& JsonObject);
+	static JSONUTILITIES_API FFormatNamedArguments ParseTextArgumentsFromJson(const TSharedPtr<const FJsonObject>& JsonObject);
 };
diff --git a/Engine/Source/Runtime/JsonUtilities/Public/JsonObjectWrapper.h b/Engine/Source/Runtime/JsonUtilities/Public/JsonObjectWrapper.h
--- a/Engine/Source/Runtime/JsonUtilities/Public/JsonObjectWrapper.h
+++ b/Engine/Source/Runtime/JsonUtilities/Public/JsonObjectWrapper.h
@@ -19,26 +19,26 @@ class FOutputDevice;
 
 /** UStruct that holds a JsonObject, can be used by structs passed to JsonObjectConverter to pass through JsonObjects directly */
 USTRUCT(BlueprintType, meta = (DisplayName = "JsonObject"))
-struct JSONUTILITIES_API FJsonObjectWrapper
+struct FJsonObjectWrapper
 {
 	GENERATED_USTRUCT_BODY()
 	
 public:
-	FJsonObjectWrapper();
+	JSONUTILITIES_API FJsonObjectWrapper();
 	
 	UPROPERTY(EditAnywhere, Category = "JSON")
 	FString JsonString;
 
 	TSharedPtr<FJsonObject> JsonObject;
 
-	bool ImportTextItem(const TCHAR*& Buffer, int32 PortFlags, UObject* Parent, FOutputDevice* ErrorText);
-	bool ExportTextItem(FString& ValueStr, FJsonObjectWrapper const& DefaultValue, UObject* Parent, int32 PortFlags, UObject* ExportRootScope) const;
-	void PostSerialize(const FArchive& Ar);
+	JSONUTILITIES_API bool ImportTextItem(const TCHAR*& Buffer, int32 PortFlags, UObject* Parent, FOutputDevice* ErrorText);
+	JSONUTILITIES_API bool ExportTextItem(FString& ValueStr, FJsonObjectWrapper const& DefaultValue, UObject* Parent, int32 PortFlags, UObject* ExportRootScope) const;
+	JSONUTILITIES_API void PostSerialize(const FArchive& Ar);
 
-	explicit operator bool() const;
+	JSONUTILITIES_API explicit operator bool() const;
 
-	bool JsonObjectToString(FString& Str) const;
-	bool JsonObjectFromString(const FString& Str);
+	JSONUTILITIES_API bool JsonObjectToString(FString& Str) const;
+	JSONUTILITIES_API bool JsonObjectFromString(const FString& Str);
 };
 
 template<>
diff --git a/Engine/Source/Runtime/Landscape/Classes/ControlPointMeshActor.h b/Engine/Source/Runtime/Landscape/Classes/ControlPointMeshActor.h
--- a/Engine/Source/Runtime/Landscape/Classes/ControlPointMeshActor.h
+++ b/Engine/Source/Runtime/Landscape/Classes/ControlPointMeshActor.h
@@ -12,8 +12,8 @@
  *
  * @see UControlPointMeshComponent
  */
-UCLASS(NotPlaceable, hideCategories = (Input), showCategories = ("Input|MouseInput", "Input|TouchInput"), ConversionRoot, ComponentWrapperClass, meta = (ChildCanTick))
-class LANDSCAPE_API AControlPointMeshActor : public AActor
+UCLASS(NotPlaceable, hideCategories = (Input), showCategories = ("Input|MouseInput", "Input|TouchInput"), ConversionRoot, ComponentWrapperClass, meta = (ChildCanTick), MinimalAPI)
+class AControlPointMeshActor : public AActor
 {
 	GENERATED_UCLASS_BODY()
 
@@ -24,23 +24,23 @@ private:
 public:
 
 	/** Function to change mobility type */
-	void SetMobility(EComponentMobility::Type InMobility);
+	LANDSCAPE_API void SetMobility(EComponentMobility::Type InMobility);
 
 #if WITH_EDITOR
 	//~ Begin AActor Interface
-	virtual void CheckForErrors() override;
-	virtual bool GetReferencedContentObjects(TArray<UObject*>& Objects) const override;
+	LANDSCAPE_API virtual void CheckForErrors() override;
+	LANDSCAPE_API virtual bool GetReferencedContentObjects(TArray<UObject*>& Objects) const override;
 	//~ End AActor Interface
 #endif // WITH_EDITOR	
 
 protected:
 	//~ Begin UObject Interface.
-	virtual FString GetDetailedInfoInternal() const override;
+	LANDSCAPE_API virtual FString GetDetailedInfoInternal() const override;
 	//~ End UObject Interface.
 
 public:
 	/** Returns ControlPointMeshComponent subobject **/
-	class UControlPointMeshComponent* GetControlPointMeshComponent() const;
+	LANDSCAPE_API class UControlPointMeshComponent* GetControlPointMeshComponent() const;
 };
 
 
diff --git a/Engine/Source/Runtime/Landscape/Classes/ILandscapeSplineInterface.h b/Engine/Source/Runtime/Landscape/Classes/ILandscapeSplineInterface.h
--- a/Engine/Source/Runtime/Landscape/Classes/ILandscapeSplineInterface.h
+++ b/Engine/Source/Runtime/Landscape/Classes/ILandscapeSplineInterface.h
@@ -7,8 +7,8 @@
 #include "UObject/Interface.h"
 #include "ILandscapeSplineInterface.generated.h"
 
-UINTERFACE()
-class LANDSCAPE_API ULandscapeSplineInterface : public UInterface
+UINTERFACE(MinimalAPI)
+class ULandscapeSplineInterface : public UInterface
 {
 	GENERATED_UINTERFACE_BODY()
 };
@@ -16,19 +16,19 @@ class LANDSCAPE_API ULandscapeSplineInterface : public UInterface
 class ULandscapeSplinesComponent;
 class ULandscapeInfo;
 
-class LANDSCAPE_API ILandscapeSplineInterface
+class ILandscapeSplineInterface
 {
 	GENERATED_IINTERFACE_BODY()
 public:
-	virtual FGuid GetLandscapeGuid() const PURE_VIRTUAL(ILandscapeSplineInterface::GetLandscapeGuid, return FGuid(););
-	virtual ULandscapeInfo* GetLandscapeInfo() const PURE_VIRTUAL(ILandscapeSplineInterface::GetLandscapeInfo, return nullptr;);
-	virtual FTransform LandscapeActorToWorld() const PURE_VIRTUAL(ILandscapeSplineInterface::LandscapeActorToWorld, return FTransform::Identity;);
-	virtual ULandscapeSplinesComponent* GetSplinesComponent() const PURE_VIRTUAL(ILandscapeSplineInterface::GetSplineComponent, return nullptr;);
-	virtual void UpdateSharedProperties(ULandscapeInfo* InLandscapeInfo) PURE_VIRTUAL(ILandscapeSplineInterface::UpdateSharedProperties);
+	LANDSCAPE_API virtual FGuid GetLandscapeGuid() const PURE_VIRTUAL(ILandscapeSplineInterface::GetLandscapeGuid, return FGuid(););
+	LANDSCAPE_API virtual ULandscapeInfo* GetLandscapeInfo() const PURE_VIRTUAL(ILandscapeSplineInterface::GetLandscapeInfo, return nullptr;);
+	LANDSCAPE_API virtual FTransform LandscapeActorToWorld() const PURE_VIRTUAL(ILandscapeSplineInterface::LandscapeActorToWorld, return FTransform::Identity;);
+	LANDSCAPE_API virtual ULandscapeSplinesComponent* GetSplinesComponent() const PURE_VIRTUAL(ILandscapeSplineInterface::GetSplineComponent, return nullptr;);
+	LANDSCAPE_API virtual void UpdateSharedProperties(ULandscapeInfo* InLandscapeInfo) PURE_VIRTUAL(ILandscapeSplineInterface::UpdateSharedProperties);
 
 #if WITH_EDITOR
-	virtual bool SupportsForeignSplineMesh() const PURE_VIRTUAL(ILandscapeSplineInterface::SupportsForeignSplineMesh, return false;);
-	virtual void CreateSplineComponent() PURE_VIRTUAL(ILandscapeSplineInterface::CreateSplineComponent);
-	virtual void CreateSplineComponent(const FVector& Scale3D) PURE_VIRTUAL(ILandscapeSplineInterface::CreateSplineComponent);
+	LANDSCAPE_API virtual bool SupportsForeignSplineMesh() const PURE_VIRTUAL(ILandscapeSplineInterface::SupportsForeignSplineMesh, return false;);
+	LANDSCAPE_API virtual void CreateSplineComponent() PURE_VIRTUAL(ILandscapeSplineInterface::CreateSplineComponent);
+	LANDSCAPE_API virtual void CreateSplineComponent(const FVector& Scale3D) PURE_VIRTUAL(ILandscapeSplineInterface::CreateSplineComponent);
 #endif
-};
\ No newline at end of file
+};
