[Interchange] Warnings when cooking a game with Interchange plugins disabled fixes.
Serializing NodeContainer and Pipellines for UInterchangeImportAssetData.

diff --git a/Engine/Plugins/Interchange/Runtime/Source/Import/Private/Animation/InterchangeAnimSequenceFactory.cpp b/Engine/Plugins/Interchange/Runtime/Source/Import/Private/Animation/InterchangeAnimSequenceFactory.cpp
--- a/Engine/Plugins/Interchange/Runtime/Source/Import/Private/Animation/InterchangeAnimSequenceFactory.cpp
+++ b/Engine/Plugins/Interchange/Runtime/Source/Import/Private/Animation/InterchangeAnimSequenceFactory.cpp
@@ -998,7 +998,7 @@ UObject* UInterchangeAnimSequenceFactory::ImportObjectSourceData(const FImportAs
 		UInterchangeFactoryBaseNode* PreviousNode = nullptr;
 		if (InterchangeAssetImportData)
 		{
-			PreviousNode = InterchangeAssetImportData->NodeContainer->GetFactoryNode(InterchangeAssetImportData->NodeUniqueID);
+			PreviousNode = InterchangeAssetImportData->GetStoredFactoryNode(InterchangeAssetImportData->NodeUniqueID);
 		}
 		UInterchangeFactoryBaseNode* CurrentNode = NewObject<UInterchangeAnimSequenceFactoryNode>(GetTransientPackage());
 		UInterchangeBaseNode::CopyStorage(AnimSequenceFactoryNode, CurrentNode);
diff --git a/Engine/Plugins/Interchange/Runtime/Source/Import/Private/InterchangeImportModule.cpp b/Engine/Plugins/Interchange/Runtime/Source/Import/Private/InterchangeImportModule.cpp
--- a/Engine/Plugins/Interchange/Runtime/Source/Import/Private/InterchangeImportModule.cpp
+++ b/Engine/Plugins/Interchange/Runtime/Source/Import/Private/InterchangeImportModule.cpp
@@ -100,11 +100,15 @@ void FInterchangeImportModule::StartupModule()
 	{
 		FCoreDelegates::OnPostEngineInit.AddLambda(RegisterItems);
 	}
+
+	UInterchangeManager::SetInterchangeImportEnabled(true);
 }
 
 
 void FInterchangeImportModule::ShutdownModule()
-{}
+{
+	UInterchangeManager::SetInterchangeImportEnabled(false);
+}
 
 
 
diff --git a/Engine/Plugins/Interchange/Runtime/Source/Import/Private/Material/InterchangeMaterialFactory.cpp b/Engine/Plugins/Interchange/Runtime/Source/Import/Private/Material/InterchangeMaterialFactory.cpp
--- a/Engine/Plugins/Interchange/Runtime/Source/Import/Private/Material/InterchangeMaterialFactory.cpp
+++ b/Engine/Plugins/Interchange/Runtime/Source/Import/Private/Material/InterchangeMaterialFactory.cpp
@@ -578,7 +578,7 @@ UInterchangeFactoryBase::FImportAssetResult UInterchangeMaterialFactory::BeginIm
 				{
 					if (UInterchangeAssetImportData* AssetImportData = Cast<UInterchangeAssetImportData>(MaterialInstance->AssetImportData))
 					{
-						UInterchangeMaterialInstanceFactoryNode* PreviousNode = Cast<UInterchangeMaterialInstanceFactoryNode>(AssetImportData->NodeContainer->GetFactoryNode(AssetImportData->NodeUniqueID));
+						UInterchangeMaterialInstanceFactoryNode* PreviousNode = Cast<UInterchangeMaterialInstanceFactoryNode>(AssetImportData->GetStoredFactoryNode(AssetImportData->NodeUniqueID));
 
 						if (PreviousNode)
 						{
@@ -684,7 +684,7 @@ UInterchangeFactoryBase::FImportAssetResult UInterchangeMaterialFactory::ImportA
 		UInterchangeFactoryBaseNode* PreviousNode = nullptr;
 		if(InterchangeAssetImportData)
 		{
-			PreviousNode = InterchangeAssetImportData->NodeContainer->GetFactoryNode(InterchangeAssetImportData->NodeUniqueID);
+			PreviousNode = InterchangeAssetImportData->GetStoredFactoryNode(InterchangeAssetImportData->NodeUniqueID);
 		}
 		UInterchangeFactoryBaseNode* CurrentNode = NewObject<UInterchangeMaterialFactoryNode>(GetTransientPackage());
 		UInterchangeBaseNode::CopyStorage(MaterialFactoryNode, CurrentNode);
diff --git a/Engine/Plugins/Interchange/Runtime/Source/Import/Private/Mesh/InterchangeSkeletalMeshFactory.cpp b/Engine/Plugins/Interchange/Runtime/Source/Import/Private/Mesh/InterchangeSkeletalMeshFactory.cpp
--- a/Engine/Plugins/Interchange/Runtime/Source/Import/Private/Mesh/InterchangeSkeletalMeshFactory.cpp
+++ b/Engine/Plugins/Interchange/Runtime/Source/Import/Private/Mesh/InterchangeSkeletalMeshFactory.cpp
@@ -1661,7 +1661,7 @@ UInterchangeFactoryBase::FImportAssetResult UInterchangeSkeletalMeshFactory::End
 		UInterchangeFactoryBaseNode* PreviousNode = nullptr;
 		if (InterchangeAssetImportData)
 		{
-			PreviousNode = InterchangeAssetImportData->NodeContainer->GetFactoryNode(InterchangeAssetImportData->NodeUniqueID);
+			PreviousNode = InterchangeAssetImportData->GetStoredFactoryNode(InterchangeAssetImportData->NodeUniqueID);
 		}
 		UInterchangeFactoryBaseNode* CurrentNode = NewObject<UInterchangeSkeletalMeshFactoryNode>(GetTransientPackage());
 		UInterchangeBaseNode::CopyStorage(SkeletalMeshFactoryNode, CurrentNode);
diff --git a/Engine/Plugins/Interchange/Runtime/Source/Import/Private/Mesh/InterchangeStaticMeshFactory.cpp b/Engine/Plugins/Interchange/Runtime/Source/Import/Private/Mesh/InterchangeStaticMeshFactory.cpp
--- a/Engine/Plugins/Interchange/Runtime/Source/Import/Private/Mesh/InterchangeStaticMeshFactory.cpp
+++ b/Engine/Plugins/Interchange/Runtime/Source/Import/Private/Mesh/InterchangeStaticMeshFactory.cpp
@@ -622,7 +622,7 @@ UInterchangeFactoryBase::FImportAssetResult UInterchangeStaticMeshFactory::EndIm
 		UInterchangeFactoryBaseNode* PreviousNode = nullptr;
 		if (InterchangeAssetImportData)
 		{
-			PreviousNode = InterchangeAssetImportData->NodeContainer->GetFactoryNode(InterchangeAssetImportData->NodeUniqueID);
+			PreviousNode = InterchangeAssetImportData->GetStoredFactoryNode(InterchangeAssetImportData->NodeUniqueID);
 		}
 		UInterchangeFactoryBaseNode* CurrentNode = NewObject<UInterchangeStaticMeshFactoryNode>(GetTransientPackage());
 		UInterchangeBaseNode::CopyStorage(StaticMeshFactoryNode, CurrentNode);
diff --git a/Engine/Plugins/Interchange/Runtime/Source/Import/Private/Texture/InterchangeTextureFactory.cpp b/Engine/Plugins/Interchange/Runtime/Source/Import/Private/Texture/InterchangeTextureFactory.cpp
--- a/Engine/Plugins/Interchange/Runtime/Source/Import/Private/Texture/InterchangeTextureFactory.cpp
+++ b/Engine/Plugins/Interchange/Runtime/Source/Import/Private/Texture/InterchangeTextureFactory.cpp
@@ -1821,7 +1821,7 @@ void UInterchangeTextureFactory::SetupObject_GameThread(const FSetupObjectParams
 			UInterchangeFactoryBaseNode* PreviousNode = nullptr;
 			if (InterchangeAssetImportData)
 			{
-				PreviousNode = InterchangeAssetImportData->NodeContainer->GetFactoryNode(InterchangeAssetImportData->NodeUniqueID);
+				PreviousNode = InterchangeAssetImportData->GetStoredFactoryNode(InterchangeAssetImportData->NodeUniqueID);
 			}
 
 			UInterchangeFactoryBaseNode* CurrentNode = NewObject<UInterchangeFactoryBaseNode>(GetTransientPackage(), GetSupportedFactoryNodeClass(TextureFactoryNode));
diff --git a/Engine/Plugins/MeshPainting/Source/MeshPaintEditorMode/Private/MeshPaintModeHelpers.cpp b/Engine/Plugins/MeshPainting/Source/MeshPaintEditorMode/Private/MeshPaintModeHelpers.cpp
--- a/Engine/Plugins/MeshPainting/Source/MeshPaintEditorMode/Private/MeshPaintModeHelpers.cpp
+++ b/Engine/Plugins/MeshPainting/Source/MeshPaintEditorMode/Private/MeshPaintModeHelpers.cpp
@@ -339,13 +339,14 @@ void UMeshPaintModeSubsystem::ImportVertexColorsToSkeletalMesh(USkeletalMesh* Sk
 		UInterchangeAssetImportData* InterchangeAssetImportData = Cast<UInterchangeAssetImportData>(SkeletalMesh->GetAssetImportData());
 		if (InterchangeAssetImportData)
 		{
-			for (TObjectPtr<UObject> PipelineBase : InterchangeAssetImportData->Pipelines)
+			TArray<UObject*> Pipelines = InterchangeAssetImportData->GetPipelines();
+			for (UObject* PipelineBase : Pipelines)
 			{
-				UInterchangeGenericAssetsPipeline* GenericAssetPipeline = Cast<UInterchangeGenericAssetsPipeline>(PipelineBase.Get());
+				UInterchangeGenericAssetsPipeline* GenericAssetPipeline = Cast<UInterchangeGenericAssetsPipeline>(PipelineBase);
 
 				if (!GenericAssetPipeline)
 				{
-					if (UInterchangePythonPipelineAsset* PythonPipelineAsset = Cast<UInterchangePythonPipelineAsset>(PipelineBase.Get()))
+					if (UInterchangePythonPipelineAsset* PythonPipelineAsset = Cast<UInterchangePythonPipelineAsset>(PipelineBase))
 					{
 						GenericAssetPipeline = Cast<UInterchangeGenericAssetsPipeline>(PythonPipelineAsset->GeneratedPipeline);
 					}
diff --git a/Engine/Plugins/MeshPainting/Source/MeshPaintingToolset/Private/MeshPaintSkeletalMeshAdapter.cpp b/Engine/Plugins/MeshPainting/Source/MeshPaintingToolset/Private/MeshPaintSkeletalMeshAdapter.cpp
--- a/Engine/Plugins/MeshPainting/Source/MeshPaintingToolset/Private/MeshPaintSkeletalMeshAdapter.cpp
+++ b/Engine/Plugins/MeshPainting/Source/MeshPaintingToolset/Private/MeshPaintSkeletalMeshAdapter.cpp
@@ -398,13 +398,14 @@ void FMeshPaintSkeletalMeshComponentAdapter::PreEdit()
 		UInterchangeAssetImportData* InterchangeAssetImportData = Cast<UInterchangeAssetImportData>(ReferencedSkeletalMesh->GetAssetImportData());
 		if (InterchangeAssetImportData)
 		{
-			for (TObjectPtr<UObject> PipelineBase : InterchangeAssetImportData->Pipelines)
+			TArray<UObject*> Pipelines = InterchangeAssetImportData->GetPipelines();
+			for (UObject* PipelineBase : Pipelines)
 			{
-				UInterchangeGenericAssetsPipeline* GenericAssetPipeline = Cast<UInterchangeGenericAssetsPipeline>(PipelineBase.Get());
+				UInterchangeGenericAssetsPipeline* GenericAssetPipeline = Cast<UInterchangeGenericAssetsPipeline>(PipelineBase);
 
 				if (!GenericAssetPipeline)
 				{
-					if (UInterchangePythonPipelineAsset* PythonPipelineAsset = Cast<UInterchangePythonPipelineAsset>(PipelineBase.Get()))
+					if (UInterchangePythonPipelineAsset* PythonPipelineAsset = Cast<UInterchangePythonPipelineAsset>(PipelineBase))
 					{
 						GenericAssetPipeline = Cast<UInterchangeGenericAssetsPipeline>(PythonPipelineAsset->GeneratedPipeline);
 					}
diff --git a/Engine/Source/Runtime/Interchange/Core/Private/InterchangeCustomVersion.cpp b/Engine/Source/Runtime/Interchange/Core/Private/InterchangeCustomVersion.cpp
new file mode 100644
--- /dev/null
+++ b/Engine/Source/Runtime/Interchange/Core/Private/InterchangeCustomVersion.cpp
@@ -0,0 +1,7 @@
+// Copyright Epic Games, Inc. All Rights Reserved.
+
+#include "InterchangeCustomVersion.h"
+#include "Serialization/CustomVersion.h"
+
+const FGuid FInterchangeCustomVersion::GUID("92738C43-2988-4D9C-9A3D-9BBE6EFF9FC0");
+FCustomVersionRegistration GRegisterInterchangeCustomVersion(FInterchangeCustomVersion::GUID, FInterchangeCustomVersion::LatestVersion, TEXT("InterchangeAssetImportDataVer"));
\ No newline at end of file
diff --git a/Engine/Source/Runtime/Interchange/Core/Private/InterchangePipelineBase.cpp b/Engine/Source/Runtime/Interchange/Core/Private/InterchangePipelineBase.cpp
--- a/Engine/Source/Runtime/Interchange/Core/Private/InterchangePipelineBase.cpp
+++ b/Engine/Source/Runtime/Interchange/Core/Private/InterchangePipelineBase.cpp
@@ -356,3 +356,58 @@ void UInterchangePipelineBase::HidePropertiesOfSubCategory(UInterchangePipelineB
 	constexpr bool bVisibilityState = false;
 	InternalToggleVisibilityPropertiesOfMetaDataValue(OuterMostPipeline, Pipeline, bDoTransientSubPipeline, TEXT("SubCategory"), HideSubCategoryName, bVisibilityState);
 }
+
+struct FWeakObjectPtrData
+{
+	FString PropertyName;
+	FWeakObjectProperty* WeakObjectProperty;
+	void* ValuePtr;
+};
+
+void GatherObjectAndWeakObjectPtrs(UClass* Class, UObject* Object/*Value*/, TMap<FString, UObject*>& ObjectPtrs, TArray<FWeakObjectPtrData>& WeakObjectPtrs)
+{
+	for (TFieldIterator<FProperty> It(Class); It; ++It)
+	{
+		FProperty* Property = *It;
+
+		FString VariableName = Property->GetName();
+
+		if (FWeakObjectProperty* WeakObjectProperty = CastField<FWeakObjectProperty>(Property))
+		{
+			FWeakObjectPtrData& Data = WeakObjectPtrs.AddDefaulted_GetRef();
+			Data.PropertyName = Property->GetName();
+			Data.WeakObjectProperty = WeakObjectProperty;
+			Data.ValuePtr = Property->ContainerPtrToValuePtr<uint8>(Object);
+		}
+		else if (FObjectProperty* ObjectProperty = CastField<FObjectProperty>(Property))
+		{
+			const void* SubValue = Property->ContainerPtrToValuePtr<uint8>(Object);
+			UObject* SubObject = ObjectProperty->GetObjectPropertyValue(SubValue);
+			if (SubObject)
+			{
+				FString PropertyName = Property->GetName();
+				ObjectPtrs.Add(PropertyName, SubObject);
+
+				UClass* SubObjectPropertyClass = SubObject->GetClass();
+
+				GatherObjectAndWeakObjectPtrs(SubObjectPropertyClass, SubObject, ObjectPtrs, WeakObjectPtrs);
+			}
+		}
+	}
+}
+
+void UInterchangePipelineBase::UpdateWeakObjectPtrs()
+{
+	//Fix WeakObjectPtr connections:
+	TMap<FString, UObject*> ObjectPtrs;
+	TArray<FWeakObjectPtrData> WeakObjectPtrs;
+	GatherObjectAndWeakObjectPtrs(GetClass(), this, ObjectPtrs, WeakObjectPtrs);
+
+	for (FWeakObjectPtrData& WeakObjectPtrData : WeakObjectPtrs)
+	{
+		if (ObjectPtrs.Contains(WeakObjectPtrData.PropertyName))
+		{
+			WeakObjectPtrData.WeakObjectProperty->SetPropertyValue(WeakObjectPtrData.ValuePtr, ObjectPtrs[WeakObjectPtrData.PropertyName]);
+		}
+	}
+}
\ No newline at end of file
diff --git a/Engine/Source/Runtime/Interchange/Core/Private/Nodes/InterchangeBaseNodeContainer.cpp b/Engine/Source/Runtime/Interchange/Core/Private/Nodes/InterchangeBaseNodeContainer.cpp
--- a/Engine/Source/Runtime/Interchange/Core/Private/Nodes/InterchangeBaseNodeContainer.cpp
+++ b/Engine/Source/Runtime/Interchange/Core/Private/Nodes/InterchangeBaseNodeContainer.cpp
@@ -11,6 +11,7 @@
 #include "UObject/Object.h"
 #include "UObject/ObjectMacros.h"
 #include "UObject/UObjectIterator.h"
+#include "UObject/CoreRedirects.h"
 
 #include UE_INLINE_GENERATED_CPP_BY_NAME(InterchangeBaseNodeContainer)
 
@@ -237,6 +238,13 @@ void UInterchangeBaseNodeContainer::SerializeNodeContainerData(FArchive& Ar)
 		{
 			FString ClassFullName;
 			Ar << ClassFullName;
+
+			FCoreRedirectObjectName RedirectedObjectName = FCoreRedirects::GetRedirectedName(ECoreRedirectFlags::Type_Class, FCoreRedirectObjectName(ClassFullName));
+			if (RedirectedObjectName.IsValid())
+			{
+				ClassFullName = RedirectedObjectName.ToString();
+			}
+
 			//This cannot fail to make sure we have a healty serialization
 			if (!ensure(ClassPerName.Contains(ClassFullName)))
 			{
diff --git a/Engine/Source/Runtime/Interchange/Core/Public/InterchangeCustomVersion.h b/Engine/Source/Runtime/Interchange/Core/Public/InterchangeCustomVersion.h
new file mode 100644
--- /dev/null
+++ b/Engine/Source/Runtime/Interchange/Core/Public/InterchangeCustomVersion.h
@@ -0,0 +1,27 @@
+// Copyright Epic Games, Inc. All Rights Reserved.
+
+#pragma once
+
+#include "CoreMinimal.h"
+
+// Custom serialization version for changes to variant manager objects
+struct FInterchangeCustomVersion
+{
+	enum Type
+	{
+		// Roughly corresponds to 5.2
+		BeforeCustomVersionWasAdded = 0,
+
+		SerializedInterchangeObjectStoring,
+
+		// -----<new versions can be added above this line>-------------------------------------------------
+		VersionPlusOne,
+		LatestVersion = VersionPlusOne - 1
+	};
+
+	// The GUID for this custom version number
+	INTERCHANGECORE_API const static FGuid GUID;
+
+private:
+	FInterchangeCustomVersion() {}
+};
\ No newline at end of file
diff --git a/Engine/Source/Runtime/Interchange/Core/Public/InterchangePipelineBase.h b/Engine/Source/Runtime/Interchange/Core/Public/InterchangePipelineBase.h
--- a/Engine/Source/Runtime/Interchange/Core/Public/InterchangePipelineBase.h
+++ b/Engine/Source/Runtime/Interchange/Core/Public/InterchangePipelineBase.h
@@ -323,6 +323,12 @@ public:
 		return this == GetMostPipelineOuter();
 	};
 
+	/*
+	* Update/Re-set the TWeakObjectPtrs to point to their original TObjectPtr
+	* Works based on Variable naming.
+	*/
+	INTERCHANGECORE_API void UpdateWeakObjectPtrs();
+
 protected:
 
 	UE_DEPRECATED(5.2, "This function is replace by ExecutePipeline.")
diff --git a/Engine/Source/Runtime/Interchange/Engine/Private/InterchangeAssetImportData.cpp b/Engine/Source/Runtime/Interchange/Engine/Private/InterchangeAssetImportData.cpp
new file mode 100644
--- /dev/null
+++ b/Engine/Source/Runtime/Interchange/Engine/Private/InterchangeAssetImportData.cpp
@@ -0,0 +1,250 @@
+// Copyright Epic Games, Inc. All Rights Reserved.
+
+#include "InterchangeAssetImportData.h"
+#include "InterchangeManager.h"
+#include "InterchangePipelineBase.h"
+
+#include "InterchangeCustomVersion.h"
+
+#include "JsonObjectConverter.h"
+#include "Serialization/LargeMemoryWriter.h"
+#include "Serialization/LargeMemoryReader.h"
+#include "UObject/CoreRedirects.h"
+
+#include UE_INLINE_GENERATED_CPP_BY_NAME(InterchangeAssetImportData)
+
+void UInterchangeAssetImportData::PostLoad()
+{
+	Super::PostLoad();
+
+	if (NodeContainer_DEPRECATED)
+	{
+		SetNodeContainer(NodeContainer_DEPRECATED.Get());
+
+		NodeContainer_DEPRECATED = nullptr;
+	}
+
+	if (Pipelines_DEPRECATED.Num() > 0)
+	{
+		TransientPipelines.Empty();
+
+		for (TObjectPtr<UObject>& PipelineObject : Pipelines_DEPRECATED)
+		{
+			if (PipelineObject)
+			{
+				TransientPipelines.Add(PipelineObject.Get());
+			}
+		}
+
+		Pipelines_DEPRECATED.Empty();
+	}
+}
+
+UObject* DeSerializePipeline(const FString& PipelineStr, UClass* PipelineClass)
+{
+	UInterchangePipelineBase* GeneratedPipeline = NewObject<UInterchangePipelineBase>(GetTransientPackage(), PipelineClass);
+
+	TSharedPtr<FJsonObject> RootObject;
+	TSharedRef< TJsonReader<> > JsonReader = TJsonReaderFactory<>::Create(PipelineStr);
+	if (FJsonSerializer::Deserialize(JsonReader, RootObject))
+	{
+		const TSharedPtr<FJsonObject> JsonPipelineProperties = RootObject->GetObjectField(TEXT("GeneratedPipeline"));
+		FJsonObjectConverter::JsonObjectToUStruct(JsonPipelineProperties.ToSharedRef(), PipelineClass, GeneratedPipeline, 0, 0);
+	}
+
+	GeneratedPipeline->UpdateWeakObjectPtrs();
+
+	return GeneratedPipeline;
+}
+
+FString SerializePipeline(UObject* Pipeline)
+{
+	if (UClass* PipelineClass = Pipeline->GetClass())
+	{
+		TSharedRef<FJsonObject> RootObject = MakeShareable(new FJsonObject);
+		TSharedRef<FJsonObject> PipelinePropertiesObject = MakeShareable(new FJsonObject);
+		if (FJsonObjectConverter::UStructToJsonObject(PipelineClass, Pipeline, PipelinePropertiesObject, 0, 0))
+		{
+			RootObject->SetField(TEXT("GeneratedPipeline"), MakeShareable(new FJsonValueObject(PipelinePropertiesObject)));
+		}
+		//Write the json file
+		FString Json;
+		TSharedRef<TJsonWriter<> > JsonWriter = TJsonWriterFactory<>::Create(&Json, 0);
+		if (FJsonSerializer::Serialize(RootObject, JsonWriter))
+		{
+			return Json;
+		}
+	}
+	
+	return TEXT("");
+}
+
+void UInterchangeAssetImportData::Serialize(FArchive& Ar)
+{
+	Super::Serialize(Ar);
+
+	Ar.UsingCustomVersion(FInterchangeCustomVersion::GUID);
+	int32 CustomVersion = Ar.CustomVer(FInterchangeCustomVersion::GUID);
+
+	if (Ar.IsSaving())
+	{
+		if (UInterchangeManager::IsInterchangeImportEnabled())
+		{
+			//serialize
+			if (TransientNodeContainer)
+			{
+				FLargeMemoryWriter NodeContainerAr;
+				TransientNodeContainer->SerializeNodeContainerData(NodeContainerAr);
+				CachedNodeContainer = TArray64<uint8>(NodeContainerAr.GetData(), NodeContainerAr.TotalSize());
+			}
+			else
+			{
+				CachedNodeContainer.Reset();
+			}
+
+			CachedPipelines.Reset(TransientPipelines.Num());
+			for (TObjectPtr<UObject>& PipelineObjectPtr : TransientPipelines)
+			{
+				UObject* PipelineObject = PipelineObjectPtr.Get();
+				if (PipelineObject)
+				{
+					FString PipelineJSON = SerializePipeline(PipelineObject);
+
+					FString PipelineFullName = PipelineObject->GetClass()->GetFullName();
+					CachedPipelines.Emplace(PipelineFullName, PipelineJSON);
+				}
+			}
+		}
+	}
+
+	if (CustomVersion >= FInterchangeCustomVersion::SerializedInterchangeObjectStoring)
+	{
+		Ar << CachedNodeContainer;
+		Ar << CachedPipelines;
+	}
+}
+
+
+UInterchangeBaseNodeContainer* UInterchangeAssetImportData::GetNodeContainer() const
+{
+	ProcessContainerCache();
+
+	return Cast<UInterchangeBaseNodeContainer>(TransientNodeContainer.Get());
+}
+
+void UInterchangeAssetImportData::SetNodeContainer(UInterchangeBaseNodeContainer* InNodeContainer) 
+{
+	TransientNodeContainer = InNodeContainer; 
+}
+
+
+void UInterchangeAssetImportData::SetPipelines(const TArray<UObject*>& InPipelines) 
+{
+	TransientPipelines.Reset();
+	for (UObject* Pipeline : InPipelines)
+	{
+		TransientPipelines.Add(Pipeline);
+	}
+}
+
+TArray<UObject*> UInterchangeAssetImportData::GetPipelines() const
+{
+	ProcessPipelinesCache();
+
+	TArray<UObject*> OutPipelines;
+	for (TObjectPtr<UObject>& Pipeline : TransientPipelines)
+	{
+		if (Pipeline.Get())
+		{
+			OutPipelines.Add(Pipeline.Get());
+		}
+	}
+	return OutPipelines;
+}
+
+int32 UInterchangeAssetImportData::GetNumberOfPipelines() const
+{
+	ProcessPipelinesCache();
+	return TransientPipelines.Num();
+}
+
+
+const UInterchangeBaseNode* UInterchangeAssetImportData::GetStoredNode(const FString& InNodeUniqueId) const
+{
+	UInterchangeBaseNodeContainer* NodeContainerResolved = TransientNodeContainer.Get();
+	if (NodeContainerResolved)
+	{
+		return NodeContainerResolved->GetNode(InNodeUniqueId);
+	}
+
+	return nullptr;
+}
+
+UInterchangeFactoryBaseNode* UInterchangeAssetImportData::GetStoredFactoryNode(const FString& InNodeUniqueId) const
+{
+	UInterchangeBaseNodeContainer* NodeContainerResolved = TransientNodeContainer.Get();
+	if (NodeContainerResolved)
+	{
+		return NodeContainerResolved->GetFactoryNode(InNodeUniqueId);
+	}
+
+	return nullptr;
+}
+
+
+void UInterchangeAssetImportData::ProcessContainerCache() const
+{
+	if (UInterchangeManager::IsInterchangeImportEnabled())
+	{
+		//de-serialize
+		if (!TransientNodeContainer && CachedNodeContainer.Num() > 0)
+		{
+			FLargeMemoryReader NodeContainerAr(CachedNodeContainer.GetData(), CachedNodeContainer.Num());
+			TransientNodeContainer = NewObject<UInterchangeBaseNodeContainer>();
+			TransientNodeContainer->SerializeNodeContainerData(NodeContainerAr);
+		}
+	}
+}
+
+void UInterchangeAssetImportData::ProcessPipelinesCache() const
+{
+	if (UInterchangeManager::IsInterchangeImportEnabled())
+	{
+		if ((TransientPipelines.Num() == 0) && (CachedPipelines.Num() > 0))
+		{
+			TransientPipelines.Reset(CachedPipelines.Num());
+
+			TMap<FString, UClass*> ClassPerName;
+			for (FThreadSafeObjectIterator It(UClass::StaticClass()); It; ++It)
+			{
+				UClass* Class = Cast<UClass>(*It);
+				if (Class->IsChildOf(UInterchangePipelineBase::StaticClass()))
+				{
+					ClassPerName.Add(Class->GetFullName(), Class);
+				}
+			}
+
+			for (const TPair<FString, FString>& CachedPipeline : CachedPipelines)
+			{
+				FString ClassFullName = CachedPipeline.Key;
+				FCoreRedirectObjectName RedirectedObjectName = FCoreRedirects::GetRedirectedName(ECoreRedirectFlags::Type_Class, FCoreRedirectObjectName(ClassFullName));
+				if (RedirectedObjectName.IsValid())
+				{
+					ClassFullName = RedirectedObjectName.ToString();
+				}
+				//This cannot fail to make sure we have a healty serialization
+				if (!ensure(ClassPerName.Contains(ClassFullName)))
+				{
+					//We did not successfully serialize the content of the file into the node container
+					return;
+				}
+
+				UClass* ToCreateClass = ClassPerName.FindChecked(ClassFullName);
+				
+				UObject* Pipeline = DeSerializePipeline(CachedPipeline.Value, ToCreateClass);
+
+				TransientPipelines.Add(Pipeline);
+			}
+		}
+	}
+}
\ No newline at end of file
diff --git a/Engine/Source/Runtime/Interchange/Engine/Private/InterchangeImportCommon.cpp b/Engine/Source/Runtime/Interchange/Engine/Private/InterchangeImportCommon.cpp
--- a/Engine/Source/Runtime/Interchange/Engine/Private/InterchangeImportCommon.cpp
+++ b/Engine/Source/Runtime/Interchange/Engine/Private/InterchangeImportCommon.cpp
@@ -45,6 +45,11 @@ namespace UE::Interchange
 
 		void RecursivelyDuplicateFactoryNodeDependencies(UInterchangeBaseNodeContainer* SourceContainer, UInterchangeBaseNodeContainer* DestinationContainer, TArray<FString>& FactoryDependencies)
 		{
+			if (!DestinationContainer)
+			{
+				return;
+			}
+
 			for (const FString& DependencyUid : FactoryDependencies)
 			{
 				if (UInterchangeFactoryBaseNode* FactoryNode = Cast<UInterchangeFactoryBaseNode>(SourceContainer->GetFactoryNode(DependencyUid)))
@@ -62,7 +67,7 @@ namespace UE::Interchange
 		void EndSetupAssetData(FFactoryCommon::FUpdateImportAssetDataParameters& Parameters, UInterchangeAssetImportData* AssetImportData)
 		{
 			TRACE_CPUPROFILER_EVENT_SCOPE("UE::Interchange::Private::ImportCommon::EndSetupAssetData")
-			AssetImportData->NodeContainer = NewObject<UInterchangeBaseNodeContainer>(AssetImportData);
+			AssetImportData->SetNodeContainer(NewObject<UInterchangeBaseNodeContainer>(AssetImportData));
 			//We copy only the factory node dependencies, we use this only 
 			if (UInterchangeFactoryBaseNode* FactoryNode = Cast<UInterchangeFactoryBaseNode>(Parameters.NodeContainer->GetFactoryNode(Parameters.NodeUniqueID)))
 			{
@@ -70,17 +75,19 @@ namespace UE::Interchange
 				AssetImportData->NodeUniqueID = Parameters.NodeUniqueID;
 				TArray<FString> FactoryDependencies;
 				FactoryDependencies.Add(AssetImportData->NodeUniqueID);
-				RecursivelyDuplicateFactoryNodeDependencies(Parameters.NodeContainer, AssetImportData->NodeContainer, FactoryDependencies);
+				RecursivelyDuplicateFactoryNodeDependencies(Parameters.NodeContainer, AssetImportData->GetNodeContainer(), FactoryDependencies);
 			}
-			AssetImportData->Pipelines.Reset();
+			
+			TArray<UObject*> NewPipelines;
 			for (const UObject* Pipeline : Parameters.Pipelines)
 			{
 				UObject* DupPipeline = Cast<UObject>(StaticDuplicateObject(Pipeline, AssetImportData));
 				if (DupPipeline)
 				{
-					AssetImportData->Pipelines.Add(DupPipeline);
+					NewPipelines.Add(DupPipeline);
 				}
 			}
+			AssetImportData->SetPipelines(NewPipelines);
 		}
 	}
 
@@ -250,7 +257,8 @@ namespace UE::Interchange
 			return false;
 		}
 
-		for (UObject* PipelineObject : InterchangeAssetImportData->Pipelines)
+		TArray<UObject*> Pipelines = InterchangeAssetImportData->GetPipelines();
+		for (UObject* PipelineObject : Pipelines)
 		{
 			if (UInterchangePythonPipelineAsset* PythonPipelineAsset = Cast<UInterchangePythonPipelineAsset>(PipelineObject))
 			{
diff --git a/Engine/Source/Runtime/Interchange/Engine/Private/InterchangeManager.cpp b/Engine/Source/Runtime/Interchange/Engine/Private/InterchangeManager.cpp
--- a/Engine/Source/Runtime/Interchange/Engine/Private/InterchangeManager.cpp
+++ b/Engine/Source/Runtime/Interchange/Engine/Private/InterchangeManager.cpp
@@ -49,7 +49,7 @@
 
 #include UE_INLINE_GENERATED_CPP_BY_NAME(InterchangeManager)
 
-static bool GInterchangeImportEnable = true;
+static bool GInterchangeImportEnable = false;
 static FAutoConsoleVariableRef CCvarInterchangeImportEnable(
 	TEXT("Interchange.FeatureFlags.Import.Enable"),
 	GInterchangeImportEnable,
@@ -859,7 +859,10 @@ UInterchangeManager& UInterchangeManager::GetInterchangeManager()
 
 		InterchangeManager->GCEndDelegate = FCoreUObjectDelegates::GetPostGarbageCollect().AddLambda([]()
 			{
-				InterchangeManager->StartQueuedTasks(InterchangeManager->bGCEndDelegateCancellAllTask);
+				if (IsInterchangeImportEnabled())
+				{
+					InterchangeManager->StartQueuedTasks(InterchangeManager->bGCEndDelegateCancellAllTask);
+				}
 			});
 
 		//We cancel any running task when we pre exit the engine
@@ -1416,7 +1419,7 @@ UInterchangeManager::ImportInternal(const FString& ContentPath, const UInterchan
 	}
 #endif
 	
-	const bool bIsReimport = OriginalAssetImportData && OriginalAssetImportData->Pipelines.Num() > 0;
+	const bool bIsReimport = OriginalAssetImportData && OriginalAssetImportData->GetNumberOfPipelines() > 0;
 	auto AdjustPipelineSettingForContext = [bImportScene, bIsReimport, ImportAssetParameters](UInterchangePipelineBase* Pipeline)
 	{
 		EInterchangePipelineContext Context;
@@ -1472,12 +1475,14 @@ UInterchangeManager::ImportInternal(const FString& ContentPath, const UInterchan
 		{
 			FInterchangeStackInfo& StackInfo = PipelineStacks.AddDefaulted_GetRef();
 			StackInfo.StackName = ReimportPipelineName;
-			for (int32 PipelineIndex = 0; PipelineIndex < OriginalAssetImportData->Pipelines.Num(); ++PipelineIndex)
+
+			TArray<UObject*> Pipelines = OriginalAssetImportData->GetPipelines();
+			for (UObject* CurrentPipeline : Pipelines)
 			{
-				UInterchangePipelineBase* SourcePipeline = Cast<UInterchangePipelineBase>(OriginalAssetImportData->Pipelines[PipelineIndex]);
+				UInterchangePipelineBase* SourcePipeline = Cast<UInterchangePipelineBase>(CurrentPipeline);
 				if (!SourcePipeline)
 				{
-					if (UInterchangePythonPipelineAsset* PythonPipelineAsset = Cast<UInterchangePythonPipelineAsset>(OriginalAssetImportData->Pipelines[PipelineIndex]))
+					if (UInterchangePythonPipelineAsset* PythonPipelineAsset = Cast<UInterchangePythonPipelineAsset>(CurrentPipeline))
 					{
 						SourcePipeline = PythonPipelineAsset->GeneratedPipeline;
 					}
@@ -1486,6 +1491,7 @@ UInterchangeManager::ImportInternal(const FString& ContentPath, const UInterchan
 				{
 					//Duplicate the pipeline saved in the asset import data
 					UInterchangePipelineBase* GeneratedPipeline = Cast<UInterchangePipelineBase>(StaticDuplicateObject(SourcePipeline, PipelineInstancesPackage));
+					GeneratedPipeline->UpdateWeakObjectPtrs();
 					// Make sure that the instance does not carry over standalone and public flags as they are not actual assets to be persisted
 					GeneratedPipeline->ClearFlags(EObjectFlags::RF_Standalone | EObjectFlags::RF_Public);
 					AdjustPipelineSettingForContext(GeneratedPipeline);
@@ -1708,6 +1714,11 @@ bool UInterchangeManager::IsInterchangeImportEnabled()
 	return CCvarInterchangeImportEnable->GetBool();
 }
 
+void UInterchangeManager::SetInterchangeImportEnabled(bool bEnabled)
+{
+	CCvarInterchangeImportEnable->Set(bEnabled);
+}
+
 bool UInterchangeManager::ExportAsset(const UObject* Asset, bool bIsAutomated)
 {
 	return false;
diff --git a/Engine/Source/Runtime/Interchange/Engine/Private/InterchangeMeshUtilities.cpp b/Engine/Source/Runtime/Interchange/Engine/Private/InterchangeMeshUtilities.cpp
--- a/Engine/Source/Runtime/Interchange/Engine/Private/InterchangeMeshUtilities.cpp
+++ b/Engine/Source/Runtime/Interchange/Engine/Private/InterchangeMeshUtilities.cpp
@@ -121,7 +121,8 @@ TFuture<bool> UInterchangeMeshUtilities::InternalImportCustomLodAsync(TSharedPtr
 	ImportAssetParameters.bIsAutomated = true;
 	if (InterchangeAssetImportData)
 	{
-		for (TObjectPtr<UObject> SelectedPipeline : InterchangeAssetImportData->Pipelines)
+		TArray<UObject*> Pipelines = InterchangeAssetImportData->GetPipelines();
+		for (UObject* SelectedPipeline : Pipelines)
 		{
 			UInterchangePipelineBase* GeneratedPipeline = nullptr;
 			if (UInterchangePythonPipelineAsset* PythonPipelineAsset = Cast<UInterchangePythonPipelineAsset>(SelectedPipeline))
@@ -134,6 +135,7 @@ TFuture<bool> UInterchangeMeshUtilities::InternalImportCustomLodAsync(TSharedPtr
 			}
 			if (ensure(GeneratedPipeline))
 			{
+				GeneratedPipeline->UpdateWeakObjectPtrs();
 				GeneratedPipeline->AdjustSettingsForContext(ImportType, nullptr);
 				ImportAssetParameters.OverridePipelines.Add(GeneratedPipeline);
 			}
diff --git a/Engine/Source/Runtime/Interchange/Engine/Private/InterchangeSceneImportAsset.cpp b/Engine/Source/Runtime/Interchange/Engine/Private/InterchangeSceneImportAsset.cpp
--- a/Engine/Source/Runtime/Interchange/Engine/Private/InterchangeSceneImportAsset.cpp
+++ b/Engine/Source/Runtime/Interchange/Engine/Private/InterchangeSceneImportAsset.cpp
@@ -94,7 +94,7 @@ void UInterchangeSceneImportAsset::OnPostWorldRename(UWorld* World)
 		FString UniqueID;
 		SceneObjects.RemoveAndCopyValue(EntryToRemove, UniqueID);
 
-		UInterchangeFactoryBaseNode* FactoryNode = AssetImportData->NodeContainer->GetFactoryNode(UniqueID);
+		UInterchangeFactoryBaseNode* FactoryNode = AssetImportData->GetStoredFactoryNode(UniqueID);
 		if(ensure(FactoryNode))
 		{
 			const FString DisplayName = FactoryNode->GetDisplayLabel();
@@ -176,7 +176,7 @@ void UInterchangeSceneImportAsset::UpdateSceneObjects()
 
 	SceneObjects.Reset();
 
-	AssetImportData->NodeContainer->IterateNodesOfType<UInterchangeFactoryBaseNode>(
+	AssetImportData->GetNodeContainer()->IterateNodesOfType<UInterchangeFactoryBaseNode>(
 		[this](const FString& NodeUid, UInterchangeFactoryBaseNode* FactoryNode)
 		{
 			if (FactoryNode)
@@ -244,7 +244,7 @@ const UInterchangeFactoryBaseNode* UInterchangeSceneImportAsset::GetFactoryNode(
 
 	if (const FString* UniqueIDPtr = SceneObjects.Find(ObjectPath))
 	{
-		return Cast<UInterchangeFactoryBaseNode>(AssetImportData->NodeContainer->GetNode(*UniqueIDPtr));
+		return Cast<UInterchangeFactoryBaseNode>(AssetImportData->GetStoredNode(*UniqueIDPtr));
 	}
 #endif
 
diff --git a/Engine/Source/Runtime/Interchange/Engine/Private/Tasks/InterchangeTaskImportObject.cpp b/Engine/Source/Runtime/Interchange/Engine/Private/Tasks/InterchangeTaskImportObject.cpp
--- a/Engine/Source/Runtime/Interchange/Engine/Private/Tasks/InterchangeTaskImportObject.cpp
+++ b/Engine/Source/Runtime/Interchange/Engine/Private/Tasks/InterchangeTaskImportObject.cpp
@@ -87,7 +87,7 @@ namespace UE::Interchange::Private
 		{
 			if(UInterchangeAssetImportData* OriginalAssetImportData = Cast<UInterchangeAssetImportData>(SubObject))
 			{
-				if (UInterchangeBaseNodeContainer* OriginalNodeContainer = OriginalAssetImportData->NodeContainer)
+				if (UInterchangeBaseNodeContainer* OriginalNodeContainer = OriginalAssetImportData->GetNodeContainer())
 				{
 					//Find the original factory node used by the last ReimportObject import
 					if (UInterchangeFactoryBaseNode* OriginalFactoryNode = OriginalNodeContainer->GetFactoryNode(OriginalAssetImportData->NodeUniqueID))
diff --git a/Engine/Source/Runtime/Interchange/Engine/Public/InterchangeAssetImportData.h b/Engine/Source/Runtime/Interchange/Engine/Public/InterchangeAssetImportData.h
--- a/Engine/Source/Runtime/Interchange/Engine/Public/InterchangeAssetImportData.h
+++ b/Engine/Source/Runtime/Interchange/Engine/Public/InterchangeAssetImportData.h
@@ -19,6 +19,11 @@ class UInterchangeAssetImportData : public UAssetImportData
 {
 	GENERATED_BODY()
 public:
+	// Begin UObject interface
+	virtual void PostLoad() override;
+	virtual void Serialize(FArchive& Ar) override;
+	// End UObject interface
+
 
 	/**
 	 * Return the first filename stored in this data. The resulting filename will be absolute (ie, not relative to the asset).
@@ -69,9 +74,9 @@ public:
 	 */
 	virtual void AppendAssetRegistryTags(TArray<FAssetRegistryTag>& OutTags) override
 	{
-		if (NodeContainer->IsNodeUidValid(NodeUniqueID))
+		if (GetNodeContainer()->IsNodeUidValid(NodeUniqueID))
 		{
-			if (const UInterchangeBaseNode* Node = NodeContainer->GetNode(NodeUniqueID))
+			if (const UInterchangeBaseNode* Node = GetStoredNode(NodeUniqueID))
 			{
 				Node->AppendAssetRegistryTags(OutTags);
 			}
@@ -81,15 +86,48 @@ public:
 #endif
 #endif
 
-	/** The graph that was use to create this asset */
-	UPROPERTY(VisibleAnywhere, Category = "Interchange | AssetImportData")
-	TObjectPtr<UInterchangeBaseNodeContainer> NodeContainer;
-
 	/** The Node UID pass to the factory that exist in the graph that was use to create this asset */
 	UPROPERTY(VisibleAnywhere, Category = "Interchange | AssetImportData")
 	FString NodeUniqueID;
 
-	/** List of pipelines use to import an asset */
-	UPROPERTY(EditAnywhere, Category = "Interchange | AssetImportData")
-	TArray<TObjectPtr<UObject>> Pipelines;
+
+	UFUNCTION(BlueprintCallable, Category = "Interchange | AssetImportData")
+	INTERCHANGEENGINE_API UInterchangeBaseNodeContainer* GetNodeContainer() const;
+
+	UFUNCTION(BlueprintCallable, Category = "Interchange | AssetImportData")
+	INTERCHANGEENGINE_API void SetNodeContainer(UInterchangeBaseNodeContainer* InNodeContainer);
+
+
+	/**
+	* Returns Array of non-null pipelines
+	*/
+	UFUNCTION(BlueprintCallable, Category = "Interchange | AssetImportData")
+	INTERCHANGEENGINE_API TArray<UObject*> GetPipelines() const;
+
+	UFUNCTION(BlueprintCallable, Category = "Interchange | AssetImportData")
+	INTERCHANGEENGINE_API int32 GetNumberOfPipelines() const;
+
+	UFUNCTION(BlueprintCallable, Category = "Interchange | AssetImportData")
+	INTERCHANGEENGINE_API void SetPipelines(const TArray<UObject*>& InPipelines);
+
+
+	UFUNCTION(BlueprintCallable, Category = "Interchange | AssetImportData")
+	INTERCHANGEENGINE_API const UInterchangeBaseNode* GetStoredNode(const FString& InNodeUniqueId) const;
+
+	UFUNCTION(BlueprintCallable, Category = "Interchange | AssetImportData")
+	INTERCHANGEENGINE_API UInterchangeFactoryBaseNode* GetStoredFactoryNode(const FString& InNodeUniqueId) const;
+
+private:
+	UPROPERTY(meta = (DeprecatedProperty, DeprecationMessage = "Use GetNodeContainer/SetNodeContainer instead."))
+	TObjectPtr<UInterchangeBaseNodeContainer> NodeContainer_DEPRECATED;
+	UPROPERTY(meta = (DeprecatedProperty, DeprecationMessage = "Use GetPipelines/SetPipelines instead."))
+	TArray<TObjectPtr<UObject>> Pipelines_DEPRECATED;
+
+	mutable TObjectPtr<UInterchangeBaseNodeContainer> TransientNodeContainer;
+	mutable TArray<TObjectPtr<UObject>> TransientPipelines;
+
+	void ProcessContainerCache() const;
+	void ProcessPipelinesCache() const;
+	TArray64<uint8> CachedNodeContainer;
+	TArray<TPair<FString, FString>> CachedPipelines; //Class, Data(serialized JSON) pair
 };
diff --git a/Engine/Source/Runtime/Interchange/Engine/Public/InterchangeManager.h b/Engine/Source/Runtime/Interchange/Engine/Public/InterchangeManager.h
--- a/Engine/Source/Runtime/Interchange/Engine/Public/InterchangeManager.h
+++ b/Engine/Source/Runtime/Interchange/Engine/Public/InterchangeManager.h
@@ -395,6 +395,8 @@ public:
 
 	/** Return the CVar which make interchange enable or not.*/
 	static INTERCHANGEENGINE_API bool IsInterchangeImportEnabled();
+	/** Set the CVar which make interchange enable or not.*/
+	static INTERCHANGEENGINE_API void SetInterchangeImportEnabled(bool bEnabled);
 
 	/** delegate type fired when new assets have been imported. Note: InCreatedObject can be NULL if import failed. Params: UFactory* InFactory, UObject* InCreatedObject */
 	DECLARE_MULTICAST_DELEGATE_OneParam(FInterchangeOnAssetPostImport, UObject*);
