Fix crash when re-importing texture

#rnx

diff --git a/Engine/Source/Runtime/Interchange/Engine/Public/InterchangeAssetImportData.h b/Engine/Source/Runtime/Interchange/Engine/Public/InterchangeAssetImportData.h
--- a/Engine/Source/Runtime/Interchange/Engine/Public/InterchangeAssetImportData.h
+++ b/Engine/Source/Runtime/Interchange/Engine/Public/InterchangeAssetImportData.h
@@ -123,7 +123,10 @@ private:
 	UPROPERTY(meta = (DeprecatedProperty, DeprecationMessage = "Use GetPipelines/SetPipelines instead."))
 	TArray<TObjectPtr<UObject>> Pipelines_DEPRECATED;
 
+	UPROPERTY(Transient)
 	mutable TObjectPtr<UInterchangeBaseNodeContainer> TransientNodeContainer;
+
+	UPROPERTY(Transient)
 	mutable TArray<TObjectPtr<UObject>> TransientPipelines;
 
 	void ProcessContainerCache() const;
