Interchange: Fix missing call to ProcessContainerCache when accessing reimport data.

#rnx

diff --git a/Engine/Source/Runtime/Interchange/Engine/Private/InterchangeAssetImportData.cpp b/Engine/Source/Runtime/Interchange/Engine/Private/InterchangeAssetImportData.cpp
--- a/Engine/Source/Runtime/Interchange/Engine/Private/InterchangeAssetImportData.cpp
+++ b/Engine/Source/Runtime/Interchange/Engine/Private/InterchangeAssetImportData.cpp
@@ -172,6 +172,7 @@ int32 UInterchangeAssetImportData::GetNumberOfPipelines() const
 
 const UInterchangeBaseNode* UInterchangeAssetImportData::GetStoredNode(const FString& InNodeUniqueId) const
 {
+	ProcessContainerCache();
 	UInterchangeBaseNodeContainer* NodeContainerResolved = TransientNodeContainer.Get();
 	if (NodeContainerResolved)
 	{
@@ -183,6 +184,7 @@ const UInterchangeBaseNode* UInterchangeAssetImportData::GetStoredNode(const FSt
 
 UInterchangeFactoryBaseNode* UInterchangeAssetImportData::GetStoredFactoryNode(const FString& InNodeUniqueId) const
 {
+	ProcessContainerCache();
 	UInterchangeBaseNodeContainer* NodeContainerResolved = TransientNodeContainer.Get();
 	if (NodeContainerResolved)
 	{
