Fix crash when Loading skeletalmesh with cloth asset. Conditionnal load is call on the asset import data which will not call the postload right after the serialization but later.

#rnx

diff --git a/Engine/Source/Runtime/Interchange/Engine/Private/InterchangeAssetImportData.cpp b/Engine/Source/Runtime/Interchange/Engine/Private/InterchangeAssetImportData.cpp
--- a/Engine/Source/Runtime/Interchange/Engine/Private/InterchangeAssetImportData.cpp
+++ b/Engine/Source/Runtime/Interchange/Engine/Private/InterchangeAssetImportData.cpp
@@ -133,13 +133,13 @@ UInterchangeBaseNodeContainer* UInterchangeAssetImportData::GetNodeContainer() c
 	return Cast<UInterchangeBaseNodeContainer>(TransientNodeContainer.Get());
 }
 
-void UInterchangeAssetImportData::SetNodeContainer(UInterchangeBaseNodeContainer* InNodeContainer) 
+void UInterchangeAssetImportData::SetNodeContainer(UInterchangeBaseNodeContainer* InNodeContainer) const
 {
 	TransientNodeContainer = InNodeContainer; 
 }
 
 
-void UInterchangeAssetImportData::SetPipelines(const TArray<UObject*>& InPipelines) 
+void UInterchangeAssetImportData::SetPipelines(const TArray<UObject*>& InPipelines) const
 {
 	TransientPipelines.Reset();
 	for (UObject* Pipeline : InPipelines)
@@ -206,6 +206,10 @@ void UInterchangeAssetImportData::ProcessContainerCache() const
 			TransientNodeContainer = NewObject<UInterchangeBaseNodeContainer>();
 			TransientNodeContainer->SerializeNodeContainerData(NodeContainerAr);
 		}
+		else if(!TransientNodeContainer)
+		{
+			ProcessDeprecatedData();
+		}
 	}
 }
 
@@ -249,5 +253,32 @@ void UInterchangeAssetImportData::ProcessPipelinesCache() const
 				TransientPipelines.Add(Pipeline);
 			}
 		}
+		else if (TransientPipelines.Num() == 0)
+		{
+			ProcessDeprecatedData();
+		}
+	}
+}
+
+void UInterchangeAssetImportData::ProcessDeprecatedData() const
+{
+	if (UInterchangeManager::IsInterchangeImportEnabled())
+	{
+		if (!TransientNodeContainer && NodeContainer_DEPRECATED)
+		{
+			SetNodeContainer(NodeContainer_DEPRECATED.Get());
+		}
+		if (TransientPipelines.Num() == 0)
+		{
+			TransientPipelines.Empty();
+
+			for (const TObjectPtr<UObject>& PipelineObject : Pipelines_DEPRECATED)
+			{
+				if (PipelineObject)
+				{
+					TransientPipelines.Add(PipelineObject.Get());
+				}
+			}
+		}
 	}
 }
\ No newline at end of file
diff --git a/Engine/Source/Runtime/Interchange/Engine/Private/InterchangeManager.cpp b/Engine/Source/Runtime/Interchange/Engine/Private/InterchangeManager.cpp
--- a/Engine/Source/Runtime/Interchange/Engine/Private/InterchangeManager.cpp
+++ b/Engine/Source/Runtime/Interchange/Engine/Private/InterchangeManager.cpp
@@ -801,6 +801,7 @@ UInterchangePipelineBase* UE::Interchange::GeneratePipelineInstance(const FSoftO
 	{
 		// Make sure that the instance does not carry over standalone and public flags as they are not actual assets to be persisted
 		GeneratedPipeline->ClearFlags(EObjectFlags::RF_Standalone|EObjectFlags::RF_Public);
+		GeneratedPipeline->UpdateWeakObjectPtrs();
 	}
 
 	return GeneratedPipeline;
diff --git a/Engine/Source/Runtime/Interchange/Engine/Public/InterchangeAssetImportData.h b/Engine/Source/Runtime/Interchange/Engine/Public/InterchangeAssetImportData.h
--- a/Engine/Source/Runtime/Interchange/Engine/Public/InterchangeAssetImportData.h
+++ b/Engine/Source/Runtime/Interchange/Engine/Public/InterchangeAssetImportData.h
@@ -74,11 +74,14 @@ public:
 	 */
 	virtual void AppendAssetRegistryTags(TArray<FAssetRegistryTag>& OutTags) override
 	{
-		if (GetNodeContainer()->IsNodeUidValid(NodeUniqueID))
+		if(const UInterchangeBaseNodeContainer* NodeContainerTmp = GetNodeContainer())
 		{
-			if (const UInterchangeBaseNode* Node = GetStoredNode(NodeUniqueID))
+			if (NodeContainerTmp->IsNodeUidValid(NodeUniqueID))
 			{
-				Node->AppendAssetRegistryTags(OutTags);
+				if (const UInterchangeBaseNode* Node = GetStoredNode(NodeUniqueID))
+				{
+					Node->AppendAssetRegistryTags(OutTags);
+				}
 			}
 		}
 		Super::AppendAssetRegistryTags(OutTags);
@@ -95,7 +98,7 @@ public:
 	INTERCHANGEENGINE_API UInterchangeBaseNodeContainer* GetNodeContainer() const;
 
 	UFUNCTION(BlueprintCallable, Category = "Interchange | AssetImportData")
-	INTERCHANGEENGINE_API void SetNodeContainer(UInterchangeBaseNodeContainer* InNodeContainer);
+	INTERCHANGEENGINE_API void SetNodeContainer(UInterchangeBaseNodeContainer* InNodeContainer) const;
 
 
 	/**
@@ -108,7 +111,7 @@ public:
 	INTERCHANGEENGINE_API int32 GetNumberOfPipelines() const;
 
 	UFUNCTION(BlueprintCallable, Category = "Interchange | AssetImportData")
-	INTERCHANGEENGINE_API void SetPipelines(const TArray<UObject*>& InPipelines);
+	INTERCHANGEENGINE_API void SetPipelines(const TArray<UObject*>& InPipelines) const;
 
 
 	UFUNCTION(BlueprintCallable, Category = "Interchange | AssetImportData")
@@ -131,6 +134,7 @@ private:
 
 	void ProcessContainerCache() const;
 	void ProcessPipelinesCache() const;
+	void ProcessDeprecatedData() const;
 	TArray64<uint8> CachedNodeContainer;
 	TArray<TPair<FString, FString>> CachedPipelines; //Class, Data(serialized JSON) pair
 };
