Improved workflow on a level import/reimport
  - Added logic to detect that an asset was part of a level import and not an asset import
  - Added properties on the level pipeline to allow the user to control the handling of deleted assets and actors
  - Added logic to delete actors and assets on a level re-import

#rnx

diff --git a/Engine/Plugins/Interchange/Editor/Source/Pipelines/Private/InterchangeEditorPipelineDetails.cpp b/Engine/Plugins/Interchange/Editor/Source/Pipelines/Private/InterchangeEditorPipelineDetails.cpp
--- a/Engine/Plugins/Interchange/Editor/Source/Pipelines/Private/InterchangeEditorPipelineDetails.cpp
+++ b/Engine/Plugins/Interchange/Editor/Source/Pipelines/Private/InterchangeEditorPipelineDetails.cpp
@@ -256,7 +256,11 @@ void FInterchangePipelineBaseDetailsCustomization::CustomizeDetails(IDetailLayou
 	CachedDetailBuilder = &DetailBuilder;
 	TArray<TWeakObjectPtr<UObject>> EditingObjects;
 	DetailBuilder.GetObjectsBeingCustomized(EditingObjects);
-	check(EditingObjects.Num() == 1);
+
+	if (!ensure(EditingObjects.Num() == 1))
+	{
+		return;
+	}
 
 	InterchangePipeline = Cast<UInterchangePipelineBase>(EditingObjects[0].Get());
 
diff --git a/Engine/Plugins/Interchange/Runtime/Source/Import/Private/Scene/InterchangeActorHelper.cpp b/Engine/Plugins/Interchange/Runtime/Source/Import/Private/Scene/InterchangeActorHelper.cpp
--- a/Engine/Plugins/Interchange/Runtime/Source/Import/Private/Scene/InterchangeActorHelper.cpp
+++ b/Engine/Plugins/Interchange/Runtime/Source/Import/Private/Scene/InterchangeActorHelper.cpp
@@ -98,14 +98,10 @@ AActor* UE::Interchange::ActorHelper::SpawnFactoryActor(const UInterchangeFactor
 		}
 	}
 	// The related actor has been deleted. Check on reimport policy
-	else if (CreateSceneObjectsParams.ReimportFactoryNode)
+	else if (CreateSceneObjectsParams.ReimportFactoryNode && !FactoryNode->ShouldForceNodeReimport())
 	{
 		// if reimport policy does not prioritize new content, do not recreate the actor
-		const EReimportStrategyFlags ReimportStrategyFlags = CreateSceneObjectsParams.FactoryNode->GetReimportStrategyFlags();
-		if (ReimportStrategyFlags != EReimportStrategyFlags::ApplyPipelineProperties)
-		{
-			return nullptr;
-		}
+		return nullptr;
 	}
 
 	if (!SpawnedActor)
diff --git a/Engine/Plugins/Interchange/Runtime/Source/Import/Private/Scene/InterchangeSceneImportAssetFactory.cpp b/Engine/Plugins/Interchange/Runtime/Source/Import/Private/Scene/InterchangeSceneImportAssetFactory.cpp
--- a/Engine/Plugins/Interchange/Runtime/Source/Import/Private/Scene/InterchangeSceneImportAssetFactory.cpp
+++ b/Engine/Plugins/Interchange/Runtime/Source/Import/Private/Scene/InterchangeSceneImportAssetFactory.cpp
@@ -127,12 +127,6 @@ void UInterchangeSceneImportAssetFactory::SetupObject_GameThread(const FSetupObj
 #if WITH_EDITORONLY_DATA
 		using namespace UE::Interchange;
 
-		if (Arguments.bIsReimport)
-		{
-			ensure(SceneImportAsset->AssetImportData);
-			SceneImportAsset->UpdateSceneObjects();
-		}
-
 		FFactoryCommon::FUpdateImportAssetDataParameters Parameters(SceneImportAsset, SceneImportAsset->AssetImportData, Arguments.SourceData, Arguments.NodeUniqueID, Arguments.NodeContainer, Arguments.OriginalPipelines);
 		SceneImportAsset->AssetImportData = Cast<UInterchangeAssetImportData>(FFactoryCommon::UpdateImportAssetData(Parameters));
 #endif
diff --git a/Engine/Plugins/Interchange/Runtime/Source/Pipelines/InterchangePipelines.Build.cs b/Engine/Plugins/Interchange/Runtime/Source/Pipelines/InterchangePipelines.Build.cs
--- a/Engine/Plugins/Interchange/Runtime/Source/Pipelines/InterchangePipelines.Build.cs
+++ b/Engine/Plugins/Interchange/Runtime/Source/Pipelines/InterchangePipelines.Build.cs
@@ -36,6 +36,7 @@ namespace UnrealBuildTool.Rules
 					{
 						"PhysicsUtilities",
 						"TextureUtilitiesCommon",
+						"UnrealEd",
 					}
 				);
 			}
diff --git a/Engine/Plugins/Interchange/Runtime/Source/Pipelines/Private/InterchangeGenericScenesPipeline.cpp b/Engine/Plugins/Interchange/Runtime/Source/Pipelines/Private/InterchangeGenericScenesPipeline.cpp
--- a/Engine/Plugins/Interchange/Runtime/Source/Pipelines/Private/InterchangeGenericScenesPipeline.cpp
+++ b/Engine/Plugins/Interchange/Runtime/Source/Pipelines/Private/InterchangeGenericScenesPipeline.cpp
@@ -3,6 +3,7 @@
 #include "InterchangeGenericScenesPipeline.h"
 
 #include "InterchangeActorFactoryNode.h"
+#include "InterchangeAssetImportData.h"
 #include "InterchangeCameraNode.h"
 #include "InterchangeCameraFactoryNode.h"
 #include "InterchangeCommonPipelineDataFactoryNode.h"
@@ -29,6 +30,11 @@
 #include "Engine/SpotLight.h"
 #include "Engine/StaticMeshActor.h"
 
+#if WITH_EDITOR
+#include "Engine/World.h"
+#include "ObjectTools.h"
+#endif
+
 #include UE_INLINE_GENERATED_CPP_BY_NAME(InterchangeGenericScenesPipeline)
 
 namespace UE::Interchange::Private
@@ -81,6 +87,99 @@ namespace UE::Interchange::Private
 			}
 		}
 	}
+
+#if WITH_EDITOR
+	void DeleteAssets(const TArray<UObject*>& AssetsToDelete)
+	{
+		if (AssetsToDelete.IsEmpty())
+		{
+			return;
+		}
+
+		TArray<UObject*> ObjectsToForceDelete;
+		ObjectsToForceDelete.Reserve(AssetsToDelete.Num());
+
+		for (UObject* Asset : AssetsToDelete)
+		{
+			if (Asset)
+			{
+				ObjectsToForceDelete.Add(Asset);
+			}
+		}
+
+		if (ObjectsToForceDelete.IsEmpty())
+		{
+			return;
+		}
+
+		constexpr bool bShowConfirmation = true;
+		constexpr ObjectTools::EAllowCancelDuringDelete AllowCancelDuringDelete = ObjectTools::EAllowCancelDuringDelete::CancelNotAllowed;
+		ObjectTools::DeleteObjects(ObjectsToForceDelete, bShowConfirmation, AllowCancelDuringDelete);
+	}
+#else
+	void DeleteAssets(const TArray<UObject*>& AssetsToDelete)
+	{
+		if (AssetsToDelete.IsEmpty())
+		{
+			return;
+		}
+
+		bool bForceGarbageCollection = false;
+		for (UObject* Asset : AssetsToDelete)
+		{
+			if (Asset)
+			{
+				Asset->Rename(nullptr, GetTransientPackage(), REN_NonTransactional | REN_DontCreateRedirectors);
+
+				if (Asset->IsRooted())
+				{
+					Asset->RemoveFromRoot();
+				}
+
+				Asset->ClearFlags(RF_Public | RF_Standalone);
+				Asset->MarkAsGarbage();
+
+				bForceGarbageCollection = true;
+			}
+		}
+
+		if (bForceGarbageCollection)
+		{
+			CollectGarbage(GARBAGE_COLLECTION_KEEPFLAGS);
+		}
+	}
+#endif // WITH_EDITOR
+
+	void DeleteActors(const TArray<AActor*>& ActorsToDelete)
+	{
+		if (ActorsToDelete.IsEmpty())
+		{
+			return;
+		}
+
+		for (AActor* Actor : ActorsToDelete)
+		{
+			if (!Actor)
+			{
+				continue;
+			}
+
+			if (UWorld* OwningWorld = Actor->GetWorld())
+			{
+				OwningWorld->EditorDestroyActor(Actor, true);
+				// Since deletion can be delayed, rename to avoid future name collision
+				// Call UObject::Rename directly on actor to avoid AActor::Rename which unnecessarily sunregister and re-register components
+				Actor->UObject::Rename(nullptr, GetTransientPackage(), REN_DontCreateRedirectors | REN_ForceNoResetLoaders);
+			}
+		}
+	}
+}
+
+void UInterchangeGenericLevelPipeline::AdjustSettingsForContext(EInterchangePipelineContext ImportType, TObjectPtr<UObject> ReimportAsset)
+{
+	Super::AdjustSettingsForContext(ImportType, ReimportAsset);
+
+	bIsReimportContext = ReimportAsset != nullptr;
 }
 
 void UInterchangeGenericLevelPipeline::ExecutePipeline(UInterchangeBaseNodeContainer* InBaseNodeContainer, const TArray<UInterchangeSourceData*>& InSourceDatas)
@@ -93,6 +192,15 @@ void UInterchangeGenericLevelPipeline::ExecutePipeline(UInterchangeBaseNodeConta
 
 	BaseNodeContainer = InBaseNodeContainer;
 
+	// Make sure all factory nodes created for assets have the chosen policy strategy
+	InBaseNodeContainer->IterateNodesOfType<UInterchangeFactoryBaseNode>([this](const FString& NodeUid, UInterchangeFactoryBaseNode* FactoryNode)
+		{
+			if (this->bForceReimportDeletedAssets)
+			{
+				FactoryNode->SetForceNodeReimport();
+			}
+		});
+
 	FTransform GlobalOffsetTransform = FTransform::Identity;
 	if (UInterchangeCommonPipelineDataFactoryNode* CommonPipelineDataFactoryNode = UInterchangeCommonPipelineDataFactoryNode::GetUniqueInstance(BaseNodeContainer))
 	{
@@ -303,7 +411,11 @@ void UInterchangeGenericLevelPipeline::ExecuteSceneNodePreImport(const FTransfor
 	}
 
 	//Make sure all actor factory nodes and dependencies have the specified strategy
-	Private::UpdateReimportStrategyFlags(*BaseNodeContainer, *ActorFactoryNode, ReimportStrategy);
+	Private::UpdateReimportStrategyFlags(*BaseNodeContainer, *ActorFactoryNode, ReimportPropertyStrategy);
+	if (bForceReimportDeletedActors)
+	{
+		ActorFactoryNode->SetForceNodeReimport();
+	}
 
 #if WITH_EDITORONLY_DATA
 	// Add dependency to newly created factory node
@@ -705,22 +817,92 @@ void UInterchangeGenericLevelPipeline::ExecutePostImportPipeline(const UIntercha
 	Super::ExecutePostImportPipeline(InBaseNodeContainer, NodeKey, CreatedAsset, bIsAReimport);
 
 #if WITH_EDITORONLY_DATA
+	using namespace UE::Interchange;
+
 	//We do not use the provided base container since ExecutePreImportPipeline cache it
 	//We just make sure the same one is pass in parameter
-	if (!InBaseNodeContainer || !ensure(BaseNodeContainer == InBaseNodeContainer) || !CreatedAsset)
+	if (!InBaseNodeContainer || !ensure(BaseNodeContainer == InBaseNodeContainer) || !CreatedAsset || !ensure(IsInGameThread()))
 	{
 		return;
 	}
 
 	if (UInterchangeSceneImportAsset* SceneImportAsset = Cast<UInterchangeSceneImportAsset>(CreatedAsset))
 	{
+		if (!bIsReimportContext || !bDeleteMissingActors)
+		{
+			SceneImportAsset->UpdateSceneObjects();
+			return;
+		}
+
 		const UInterchangeSceneImportAssetFactoryNode* FactoryNode = Cast<UInterchangeSceneImportAssetFactoryNode>(BaseNodeContainer->GetFactoryNode(NodeKey));
 		if (!ensure(FactoryNode))
 		{
+			SceneImportAsset->UpdateSceneObjects();
 			return;
 		}
 
+		// Cache list of objects previously imported in case of a re-import
+		TArray<FSoftObjectPath> PrevSoftObjectPaths;
+		SceneImportAsset->GetSceneSoftObjectPaths(PrevSoftObjectPaths);
+
 		SceneImportAsset->UpdateSceneObjects();
+
+		// Nothing to take care of
+		if (PrevSoftObjectPaths.IsEmpty())
+		{
+			return;
+		}
+
+		TArray<FSoftObjectPath> NewSoftObjectPaths;
+		SceneImportAsset->GetSceneSoftObjectPaths(NewSoftObjectPaths);
+
+		TSet<FSoftObjectPath> SoftObjectPathSet(NewSoftObjectPaths);
+		TArray<AActor*> ActorsToDelete;
+		TArray<UObject*> AssetsToForceDelete;
+
+		ActorsToDelete.Reserve(PrevSoftObjectPaths.Num());
+		AssetsToForceDelete.Reserve(PrevSoftObjectPaths.Num());
+
+		for (const FSoftObjectPath& ObjectPath : PrevSoftObjectPaths)
+		{
+			if (!SoftObjectPathSet.Contains(ObjectPath))
+			{
+				if (UObject* Object = ObjectPath.TryLoad())
+				{
+					if (IsValid(Object))
+					{
+						if (Object->GetClass()->IsChildOf<AActor>())
+						{
+							ActorsToDelete.Add(Cast<AActor>(Object));
+						}
+						else
+						{
+							AssetsToForceDelete.Add(Object);
+						}
+					}
+				}
+			}
+		}
+
+		Private::DeleteActors(ActorsToDelete);
+
+		if (!AssetsToForceDelete.IsEmpty())
+		{
+			Private::DeleteAssets(AssetsToForceDelete);
+		}
+
+		// Update newly imported objects with a soft reference to the UInterchangeSceneImportAsset
+		for (const FSoftObjectPath& ObjectPath : NewSoftObjectPaths)
+		{
+			UObject* Object = ObjectPath.TryLoad();
+			if (IsValid(Object) && Object != CreatedAsset)
+			{
+				if (UInterchangeAssetImportData* AssetImportData = UInterchangeAssetImportData::GetFromObject(Object))
+				{
+					AssetImportData->SceneImportAsset = CreatedAsset;
+				}
+			}
+		}
 	}
 #endif
 }
diff --git a/Engine/Plugins/Interchange/Runtime/Source/Pipelines/Public/InterchangeGenericScenesPipeline.h b/Engine/Plugins/Interchange/Runtime/Source/Pipelines/Public/InterchangeGenericScenesPipeline.h
--- a/Engine/Plugins/Interchange/Runtime/Source/Pipelines/Public/InterchangeGenericScenesPipeline.h
+++ b/Engine/Plugins/Interchange/Runtime/Source/Pipelines/Public/InterchangeGenericScenesPipeline.h
@@ -21,13 +21,31 @@ class INTERCHANGEPIPELINES_API UInterchangeGenericLevelPipeline : public UInterc
 {
 	GENERATED_BODY()
 public:
-	/* Allow user to choose the re-import strategy when importing into level. */
-	UPROPERTY(EditAnywhere, BlueprintReadWrite, Category = "Actors", meta = (AdjustPipelineAndRefreshDetailOnChange = "True"))
-	EReimportStrategyFlags ReimportStrategy = EReimportStrategyFlags::ApplyNoProperties;
+	/* Allow user to choose the re-import strategy when re-importing into level. */
+	UPROPERTY(EditAnywhere, BlueprintReadWrite, Category = "Level Reimport Policy", meta = (SubCategory = "Actors properties", AdjustPipelineAndRefreshDetailOnChange = "True"))
+	EReimportStrategyFlags ReimportPropertyStrategy = EReimportStrategyFlags::ApplyNoProperties;
 
-protected:
+	/* Enables or not the deletion of actors which were not part of the translation when re-importing into level. */
+	UPROPERTY(EditAnywhere, BlueprintReadWrite, Category = "Level Reimport Policy", meta = (SubCategory = "Reimport Actors"))
+	bool bDeleteMissingActors = false;
+
+	/* Enables or not spawning actors which were deleted in the editor prior to a reimport. */
+	UPROPERTY(EditAnywhere, BlueprintReadWrite, Category = "Level Reimport Policy", meta = (SubCategory = "Reimport Actors"))
+	bool bForceReimportDeletedActors = false;
+
+	/* Enables or not re-creating assets which were deleted in the editor prior to a reimport into level. */
+	UPROPERTY(EditAnywhere, BlueprintReadWrite, Category = "Level Reimport Policy", meta = (SubCategory = "Reimport Assets"))
+	bool bForceReimportDeletedAssets = false;
+
+	/* Enables or not the deletion of assets which were not part of the translation when re-importing into level. */
+	UPROPERTY(EditAnywhere, BlueprintReadWrite, Category = "Level Reimport Policy", meta = (SubCategory = "Reimport Assets"))
+	bool bDeleteMissingAssets = false;
 
 	/** BEGIN UInterchangePipelineBase overrides */
+	virtual void AdjustSettingsForContext(EInterchangePipelineContext ImportType, TObjectPtr<UObject> ReimportAsset) override;
+
+protected:
+
 	virtual void ExecutePipeline(UInterchangeBaseNodeContainer* InBaseNodeContainer, const TArray<UInterchangeSourceData*>& InSourceDatas) override;
 	virtual void ExecutePostImportPipeline(const UInterchangeBaseNodeContainer* BaseNodeContainer, const FString& NodeKey, UObject* CreatedAsset, bool bIsAReimport) override;
 
@@ -60,6 +78,7 @@ protected:
 	/** Disable this option to not convert Standard(Perspective) to Physical Cameras*/
 	UPROPERTY(EditAnywhere, BlueprintReadWrite, Category = "Common Skeletal Meshes and Animations")
 	bool bUsePhysicalInsteadOfStandardPerspectiveCamera = true;
+
 protected:
 	UInterchangeBaseNodeContainer* BaseNodeContainer = nullptr;
 #if WITH_EDITORONLY_DATA
diff --git a/Engine/Source/Editor/UnrealEd/Private/Editor.cpp b/Engine/Source/Editor/UnrealEd/Private/Editor.cpp
--- a/Engine/Source/Editor/UnrealEd/Private/Editor.cpp
+++ b/Engine/Source/Editor/UnrealEd/Private/Editor.cpp
@@ -455,7 +455,9 @@ UE::Interchange::FAssetImportResultRef FReimportManager::ReimportAsync(UObject*
 				CanReimportHandler->SetAutomatedReimport(bAutomated);
 				EReimportResult::Type Result = CanReimportHandler->Reimport( Obj, SourceFileIndex );
 				CanReimportHandler->SetAutomatedReimport(bOriginalAutomated);
-				if( Result == EReimportResult::Succeeded )
+				// Even if the reimport has been successful, check that the originating object is still valid
+				// The reimport might be a reimport to level which triggered the deletion of the object
+				if( Result == EReimportResult::Succeeded && IsValid(Obj))
 				{
 					Obj->PostEditChange();
 					GEditor->BroadcastObjectReimported(Obj);
diff --git a/Engine/Source/Runtime/Interchange/Core/Private/Nodes/InterchangeFactoryBaseNode.cpp b/Engine/Source/Runtime/Interchange/Core/Private/Nodes/InterchangeFactoryBaseNode.cpp
--- a/Engine/Source/Runtime/Interchange/Core/Private/Nodes/InterchangeFactoryBaseNode.cpp
+++ b/Engine/Source/Runtime/Interchange/Core/Private/Nodes/InterchangeFactoryBaseNode.cpp
@@ -21,17 +21,23 @@ namespace UE
 			return BaseNodeFactoryDependencies_BaseKey;
 		}
 
-const FAttributeKey& FFactoryBaseNodeStaticData::ReimportStrategyFlagsKey()
-{
-	static FAttributeKey AttributeKey(TEXT("__Reimport_Strategy_Flags_Key__"));
-	return AttributeKey;
-}
+		const FAttributeKey& FFactoryBaseNodeStaticData::ReimportStrategyFlagsKey()
+		{
+			static FAttributeKey AttributeKey(TEXT("__Reimport_Strategy_Flags_Key__"));
+			return AttributeKey;
+		}
 
-const FAttributeKey& FFactoryBaseNodeStaticData::SkipNodeImportKey()
-{
-	static FAttributeKey AttributeKey(TEXT("__Internal_Task_Skip_Node__"));
-	return AttributeKey;
-}
+		const FAttributeKey& FFactoryBaseNodeStaticData::SkipNodeImportKey()
+		{
+			static FAttributeKey AttributeKey(TEXT("__Internal_Task_Skip_Node__"));
+			return AttributeKey;
+		}
+
+		const FAttributeKey& FFactoryBaseNodeStaticData::ForceNodeReimportKey()
+		{
+			static FAttributeKey AttributeKey(TEXT("__Internal_Task_Force_Node_Reimport__"));
+			return AttributeKey;
+		}
 
 	} //ns Interchange
 } //ns UE
@@ -250,3 +256,40 @@ UInterchangeFactoryBaseNode* UInterchangeFactoryBaseNode::DuplicateWithObject(co
 
 	return CurrentNode;
 }
+
+bool UInterchangeFactoryBaseNode::ShouldForceNodeReimport() const
+{
+	if (!Attributes->ContainAttribute(UE::Interchange::FFactoryBaseNodeStaticData::ForceNodeReimportKey()))
+	{
+		return false;
+	}
+	bool bShouldForceNodeReimport = false;
+	Attributes->GetAttributeHandle<bool>(UE::Interchange::FFactoryBaseNodeStaticData::ForceNodeReimportKey()).Get(bShouldForceNodeReimport);
+	return bShouldForceNodeReimport;
+}
+
+bool UInterchangeFactoryBaseNode::SetForceNodeReimport()
+{
+	UE::Interchange::EAttributeStorageResult Result = Attributes->RegisterAttribute(UE::Interchange::FFactoryBaseNodeStaticData::ForceNodeReimportKey(), true);
+	if (IsAttributeStorageResultSuccess(Result))
+	{
+		UE::Interchange::FAttributeStorage::TAttributeHandle<bool> Handle = Attributes->GetAttributeHandle<bool>(UE::Interchange::FFactoryBaseNodeStaticData::ForceNodeReimportKey());
+		return Handle.IsValid();
+	}
+	return false;
+}
+
+bool UInterchangeFactoryBaseNode::UnsetForceNodeReimport()
+{
+	if (!Attributes->ContainAttribute(UE::Interchange::FFactoryBaseNodeStaticData::ForceNodeReimportKey()))
+	{
+		return true;
+	}
+	UE::Interchange::EAttributeStorageResult Result = Attributes->UnregisterAttribute(UE::Interchange::FFactoryBaseNodeStaticData::ForceNodeReimportKey());
+	if (IsAttributeStorageResultSuccess(Result))
+	{
+		UE::Interchange::FAttributeStorage::TAttributeHandle<bool> Handle = Attributes->GetAttributeHandle<bool>(UE::Interchange::FFactoryBaseNodeStaticData::ForceNodeReimportKey());
+		return !Handle.IsValid();
+	}
+	return false;
+}
diff --git a/Engine/Source/Runtime/Interchange/Core/Public/Nodes/InterchangeFactoryBaseNode.h b/Engine/Source/Runtime/Interchange/Core/Public/Nodes/InterchangeFactoryBaseNode.h
--- a/Engine/Source/Runtime/Interchange/Core/Public/Nodes/InterchangeFactoryBaseNode.h
+++ b/Engine/Source/Runtime/Interchange/Core/Public/Nodes/InterchangeFactoryBaseNode.h
@@ -135,6 +135,7 @@ namespace UE::Interchange
 		static INTERCHANGECORE_API const FString& FactoryDependenciesBaseKey();
 		static INTERCHANGECORE_API const FAttributeKey& ReimportStrategyFlagsKey();
 		static INTERCHANGECORE_API const FAttributeKey& SkipNodeImportKey();
+		static INTERCHANGECORE_API const FAttributeKey& ForceNodeReimportKey();
 	};
 
 } // namespace UE::Interchange
@@ -320,6 +321,25 @@ public:
 	 */
 	static INTERCHANGECORE_API UInterchangeFactoryBaseNode* DuplicateWithObject(const UInterchangeFactoryBaseNode* SourceNode, UObject* Object);
 
+	/**
+	 * Return whether or not an object should be created even if it has been deleted in the editor.
+	 * Return false by default
+	 */
+	UFUNCTION(BlueprintCallable, Category = "Interchange | Node")
+	INTERCHANGECORE_API bool ShouldForceNodeReimport() const;
+
+	/**
+	 * Allow the creation of the unreal object even if it has been previously deleted in the editor.
+	 */
+	UFUNCTION(BlueprintCallable, Category = "Interchange | Node")
+	INTERCHANGECORE_API bool SetForceNodeReimport();
+
+	/**
+	 * Disallow the creation of the unreal object if it has been previously deleted in the editor.
+	 */
+	UFUNCTION(BlueprintCallable, Category = "Interchange | Node")
+	INTERCHANGECORE_API bool UnsetForceNodeReimport();
+
 protected:
 	/**
 	 * Those dependencies are use by the interchange parsing task to make sure the asset are created in the correct order.
diff --git a/Engine/Source/Runtime/Interchange/Engine/Private/InterchangeManager.cpp b/Engine/Source/Runtime/Interchange/Engine/Private/InterchangeManager.cpp
--- a/Engine/Source/Runtime/Interchange/Engine/Private/InterchangeManager.cpp
+++ b/Engine/Source/Runtime/Interchange/Engine/Private/InterchangeManager.cpp
@@ -727,7 +727,8 @@ void UE::Interchange::FImportResult::AddReferencedObjects(FReferenceCollector& C
 
 void UE::Interchange::SanitizeObjectPath(FString& ObjectPath)
 {
-	const TCHAR* InvalidChar = INVALID_OBJECTPATH_CHARACTERS;
+	const TCHAR* InvalidChar = INVALID_OBJECTPATH_CHARACTERS INVALID_LONGPACKAGE_CHARACTERS;
+
 	while (*InvalidChar)
 	{
 		ObjectPath.ReplaceCharInline(*InvalidChar, TCHAR('_'), ESearchCase::CaseSensitive);
@@ -1335,7 +1336,7 @@ UInterchangeManager::ImportInternal(const FString& ContentPath, const UInterchan
 		Attribs.Add(FAnalyticsEventAttribute(TEXT("SourceExtension"), FPaths::GetExtension(SourceData->GetFilename())));
 	}
 
-	const bool bImportScene = ImportType == UE::Interchange::EImportType::ImportType_Scene;
+	bool bImportScene = ImportType == UE::Interchange::EImportType::ImportType_Scene;
 	const FInterchangeImportSettings& InterchangeImportSettings = FInterchangeProjectSettingsUtils::GetDefaultImportSettings(bImportScene);
 	
 	if (InterchangeImportSettings.PipelineStacks.Num() == 0)
@@ -1354,26 +1355,19 @@ UInterchangeManager::ImportInternal(const FString& ContentPath, const UInterchan
 		MutableInterchangeImportSettings.DefaultPipelineStack = Keys[0];
 	}
 
-	UInterchangeAssetImportData* OriginalAssetImportData = nullptr;
+	UInterchangeAssetImportData* OriginalAssetImportData = UInterchangeAssetImportData::GetFromObject(ImportAssetParameters.ReimportAsset);
 	FString PackageBasePath = ContentPath;
-	if(!ImportAssetParameters.ReimportAsset)
+	if (!ImportAssetParameters.ReimportAsset)
 	{
 		UE::Interchange::SanitizeObjectPath(PackageBasePath);
 	}
 	else
 	{
 		PackageBasePath = FPaths::GetPath(ImportAssetParameters.ReimportAsset->GetPathName());
-		TArray<UObject*> SubObjects;
-		GetObjectsWithOuter(ImportAssetParameters.ReimportAsset, SubObjects);
-		for (UObject* SubObject : SubObjects)
-		{
-			OriginalAssetImportData = Cast<UInterchangeAssetImportData>(SubObject);
-			if(OriginalAssetImportData)
-			{
-				break;
-			}
-		}
 	}
+
+	const bool bIsReimport = OriginalAssetImportData && OriginalAssetImportData->GetPipelines().Num() > 0;
+
 	bool bImportAborted = false; // True when we're unable to go through with the import process
 
 	//Create a task for every source data
@@ -1382,6 +1376,7 @@ UInterchangeManager::ImportInternal(const FString& ContentPath, const UInterchan
 	TaskData.bFollowRedirectors = ImportAssetParameters.bFollowRedirectors;
 	TaskData.ImportType = ImportType;
 	TaskData.ReimportObject = ImportAssetParameters.ReimportAsset;
+
 	TSharedRef<UE::Interchange::FImportAsyncHelper, ESPMode::ThreadSafe> AsyncHelper = CreateAsyncHelper(TaskData, ImportAssetParameters);
 	AsyncHelper->UniqueId = UniqueId;
 
@@ -1420,21 +1415,19 @@ UInterchangeManager::ImportInternal(const FString& ContentPath, const UInterchan
 	}
 #endif
 	
-	const bool bIsReimport = OriginalAssetImportData && OriginalAssetImportData->GetNumberOfPipelines() > 0;
-	auto AdjustPipelineSettingForContext = [bImportScene, bIsReimport, ImportAssetParameters](UInterchangePipelineBase* Pipeline)
+	auto AdjustPipelineSettingForContext = [bIsReimport, bImportScene, &TaskData](UInterchangePipelineBase* Pipeline)
 	{
 		EInterchangePipelineContext Context;
-		UObject* ReimportObject = nullptr;
 		if (bIsReimport)
 		{
 			Context = bImportScene ? EInterchangePipelineContext::SceneReimport : EInterchangePipelineContext::AssetReimport;
-			ReimportObject = ImportAssetParameters.ReimportAsset;
 		}
 		else
 		{
 			Context = bImportScene ? EInterchangePipelineContext::SceneImport : EInterchangePipelineContext::AssetImport;
 		}
-		Pipeline->AdjustSettingsForContext(Context, ReimportObject);
+
+		Pipeline->AdjustSettingsForContext(Context, TaskData.ReimportObject);
 	};
 
 	if ( ImportAssetParameters.OverridePipelines.Num() == 0 )
@@ -1787,7 +1780,7 @@ void UInterchangeManager::ReleaseAsyncHelper(TWeakPtr<UE::Interchange::FImportAs
 
 	check(AsyncHelper.IsValid());
 	
-	constexpr bool bLogWarningsAndErrors = true;
+	constexpr bool bLogWarningsAndErrors = false;
 
 	bool bSucceeded = false;
 	{
diff --git a/Engine/Source/Runtime/Interchange/Engine/Private/InterchangeSceneImportAsset.cpp b/Engine/Source/Runtime/Interchange/Engine/Private/InterchangeSceneImportAsset.cpp
--- a/Engine/Source/Runtime/Interchange/Engine/Private/InterchangeSceneImportAsset.cpp
+++ b/Engine/Source/Runtime/Interchange/Engine/Private/InterchangeSceneImportAsset.cpp
@@ -250,3 +250,15 @@ const UInterchangeFactoryBaseNode* UInterchangeSceneImportAsset::GetFactoryNode(
 
 	return nullptr;
 }
+
+void UInterchangeSceneImportAsset::GetSceneSoftObjectPaths(TArray<FSoftObjectPath>& SoftObjectPaths)
+{
+#if WITH_EDITORONLY_DATA
+	SoftObjectPaths.Reserve(SceneObjects.Num());
+
+	for (TPair< FSoftObjectPath, FString >& Entry : SceneObjects)
+	{
+		SoftObjectPaths.Add(Entry.Key);
+	}
+#endif
+}
\ No newline at end of file
diff --git a/Engine/Source/Runtime/Interchange/Engine/Private/Tasks/InterchangeTaskImportObject.cpp b/Engine/Source/Runtime/Interchange/Engine/Private/Tasks/InterchangeTaskImportObject.cpp
--- a/Engine/Source/Runtime/Interchange/Engine/Private/Tasks/InterchangeTaskImportObject.cpp
+++ b/Engine/Source/Runtime/Interchange/Engine/Private/Tasks/InterchangeTaskImportObject.cpp
@@ -81,43 +81,39 @@ namespace UE::Interchange::Private
 		}
 
 		//See if the FactoryNode match the ReimportObject.
-		TArray<UObject*> SubObjects;
-		GetObjectsWithOuter(ReimportObject, SubObjects);
-		for (UObject* SubObject : SubObjects)
+		if (UInterchangeAssetImportData* OriginalAssetImportData = UInterchangeAssetImportData::GetFromObject(ReimportObject))
 		{
-			if(UInterchangeAssetImportData* OriginalAssetImportData = Cast<UInterchangeAssetImportData>(SubObject))
+			if (UInterchangeBaseNodeContainer* OriginalNodeContainer = OriginalAssetImportData->GetNodeContainer())
 			{
-				if (UInterchangeBaseNodeContainer* OriginalNodeContainer = OriginalAssetImportData->GetNodeContainer())
+				//Find the original factory node used by the last ReimportObject import
+				if (UInterchangeFactoryBaseNode* OriginalFactoryNode = OriginalNodeContainer->GetFactoryNode(OriginalAssetImportData->NodeUniqueID))
 				{
-					//Find the original factory node used by the last ReimportObject import
-					if (UInterchangeFactoryBaseNode* OriginalFactoryNode = OriginalNodeContainer->GetFactoryNode(OriginalAssetImportData->NodeUniqueID))
+					//Compare the original factory node UObject class to the factory node UObject class
+					if (OriginalFactoryNode->GetObjectClass()->IsChildOf(FactoryClass))
 					{
-						//Compare the original factory node UObject class to the factory node UObject class
-						if (OriginalFactoryNode->GetObjectClass()->IsChildOf(FactoryClass))
+						//Compare the original factory node, unique id to the factory node unique id
+						if (OriginalFactoryNode->GetUniqueID().Equals(FactoryNode->GetUniqueID()))
 						{
-							//Compare the original factory node, unique id to the factory node unique id
-							if (OriginalFactoryNode->GetUniqueID().Equals(FactoryNode->GetUniqueID()))
+							return true;
+						}
+						//Compare the original factory node, name to the factory node name
+						if (OriginalFactoryNode->GetDisplayLabel().Equals(FactoryNode->GetDisplayLabel()))
+						{
+							FString PackageSubPath;
+							FactoryNode->GetCustomSubPath(PackageSubPath);
+							FString OriginalPackageSubPath;
+							OriginalFactoryNode->GetCustomSubPath(PackageSubPath);
+							//Make sure both sub path are equal
+							if (PackageSubPath.Equals(OriginalPackageSubPath))
 							{
 								return true;
 							}
-							//Compare the original factory node, name to the factory node name
-							if (OriginalFactoryNode->GetDisplayLabel().Equals(FactoryNode->GetDisplayLabel()))
-							{
-								FString PackageSubPath;
-								FactoryNode->GetCustomSubPath(PackageSubPath);
-								FString OriginalPackageSubPath;
-								OriginalFactoryNode->GetCustomSubPath(PackageSubPath);
-								//Make sure both sub path are equal
-								if (PackageSubPath.Equals(OriginalPackageSubPath))
-								{
-									return true;
-								}
-							}
 						}
 					}
 				}
 			}
 		}
+
 		return false;
 	}
 
@@ -197,8 +193,8 @@ namespace UE::Interchange::Private
 		FString AssetName;
 		Private::InternalGetPackageName(*AsyncHelper, SourceIndex, PackageBasePath, FactoryNode, PackageName, AssetName);
 		bool bSkipAsset = false;
-		UObject* ReimportObject = UE::Interchange::FFactoryCommon::GetObjectToReimport(AsyncHelper->TaskData.ReimportObject, *FactoryNode, PackageName, AssetName);
-		if (ReimportObject)
+		UObject* ObjectToReimport = UE::Interchange::FFactoryCommon::GetObjectToReimport(AsyncHelper->TaskData.ReimportObject, *FactoryNode, PackageName, AssetName);
+		if (ObjectToReimport)
 		{
 			UInterchangeBaseNodeContainer* NodeContainer = nullptr;
 			if (AsyncHelper->BaseNodeContainers.IsValidIndex(SourceIndex))
@@ -206,11 +202,11 @@ namespace UE::Interchange::Private
 				NodeContainer = AsyncHelper->BaseNodeContainers[SourceIndex].Get();
 			}
 
-			if (Private::ShouldReimportFactoryNode(FactoryNode, NodeContainer, ReimportObject))
+			if (Private::ShouldReimportFactoryNode(FactoryNode, NodeContainer, ObjectToReimport))
 			{
-				Pkg = ReimportObject->GetPackage();
+				Pkg = ObjectToReimport->GetPackage();
 				PackageName = Pkg->GetPathName();
-				AssetName = ReimportObject->GetName();
+				AssetName = ObjectToReimport->GetName();
 			}
 			else
 			{
@@ -219,25 +215,35 @@ namespace UE::Interchange::Private
 		}
 		else
 		{
-			Pkg = AsyncHelper->GetCreatedPackage(PackageName);
-
-			if (!Pkg)
+			// Check if the object has not been deleted after an import
+			if (const UInterchangeFactoryBaseNode* ReimportFactoryNode = FFactoryCommon::GetFactoryNode(AsyncHelper->TaskData.ReimportObject, PackageName, AssetName))
 			{
-				UInterchangeResultError_Generic* Message = Factory->AddMessage<UInterchangeResultError_Generic>();
-				Message->SourceAssetName = AsyncHelper->SourceDatas[SourceIndex]->GetFilename();
-				Message->DestinationAssetName = AssetName;
-				Message->AssetType = FactoryNode->GetObjectClass();
-				Message->Text = NSLOCTEXT("InternalImportObjectStartup", "BadPackage", "It was not possible to create the asset as its package was not created correctly.");
-				return ErrorImportAssetResult;
+				// if the reimport of the node has not been set, do not recreate the asset
+				bSkipAsset = !FactoryNode->ShouldForceNodeReimport();
 			}
 
-			if (!AsyncHelper->SourceDatas.IsValidIndex(SourceIndex) || !AsyncHelper->Translators.IsValidIndex(SourceIndex))
+			if (!bSkipAsset)
 			{
-				UInterchangeResultError_Generic* Message = Factory->AddMessage<UInterchangeResultError_Generic>();
-				Message->DestinationAssetName = AssetName;
-				Message->AssetType = FactoryNode->GetObjectClass();
-				Message->Text = NSLOCTEXT("InternalImportObjectStartup", "SourceDataOrTranslatorInvalid", "It was not possible to create the asset as its translator was not created correctly.");
-				return ErrorImportAssetResult;
+				Pkg = AsyncHelper->GetCreatedPackage(PackageName);
+
+				if (!Pkg)
+				{
+					UInterchangeResultError_Generic* Message = Factory->AddMessage<UInterchangeResultError_Generic>();
+					Message->SourceAssetName = AsyncHelper->SourceDatas[SourceIndex]->GetFilename();
+					Message->DestinationAssetName = AssetName;
+					Message->AssetType = FactoryNode->GetObjectClass();
+					Message->Text = NSLOCTEXT("InternalImportObjectStartup", "BadPackage", "It was not possible to create the asset as its package was not created correctly.");
+					return ErrorImportAssetResult;
+				}
+
+				if (!AsyncHelper->SourceDatas.IsValidIndex(SourceIndex) || !AsyncHelper->Translators.IsValidIndex(SourceIndex))
+				{
+					UInterchangeResultError_Generic* Message = Factory->AddMessage<UInterchangeResultError_Generic>();
+					Message->DestinationAssetName = AssetName;
+					Message->AssetType = FactoryNode->GetObjectClass();
+					Message->Text = NSLOCTEXT("InternalImportObjectStartup", "SourceDataOrTranslatorInvalid", "It was not possible to create the asset as its translator was not created correctly.");
+					return ErrorImportAssetResult;
+				}
 			}
 		}
 
@@ -264,7 +270,7 @@ namespace UE::Interchange::Private
 			{
 				CreateAssetParams.NodeContainer = AsyncHelper->BaseNodeContainers[SourceIndex].Get();
 			}
-			CreateAssetParams.ReimportObject = ReimportObject;
+			CreateAssetParams.ReimportObject = ObjectToReimport;
 
 			return FactoryOperationLambda(CreateAssetParams);
 		}
@@ -313,13 +319,13 @@ void UE::Interchange::FTaskImportObject_GameThread::DoTask(ENamedThreads::Type C
 	FString AssetName;
 	Private::InternalGetPackageName(*AsyncHelper, SourceIndex, PackageBasePath, FactoryNode, PackageName, AssetName);
 
-	UObject* ReimportObject = FFactoryCommon::GetObjectToReimport(AsyncHelper->TaskData.ReimportObject, *FactoryNode, PackageName, AssetName);
+	UObject* ObjectToReimport = FFactoryCommon::GetObjectToReimport(AsyncHelper->TaskData.ReimportObject, *FactoryNode, PackageName, AssetName);
 	if (!ensure(!IsGarbageCollecting()))
 	{
 		//Skip this asset
 		UInterchangeResultError_Generic* Message = Factory->AddMessage<UInterchangeResultError_Generic>();
 		Message->SourceAssetName = AsyncHelper->SourceDatas[SourceIndex]->GetFilename();
-		Message->DestinationAssetName = ReimportObject ? ReimportObject->GetName() : FactoryNode->GetAssetName();
+		Message->DestinationAssetName = ObjectToReimport ? ObjectToReimport->GetName() : FactoryNode->GetAssetName();
 		Message->AssetType = FactoryNode->GetObjectClass();
 		Message->Text = FText::Format(NSLOCTEXT("FTaskImportObject_GameThread", "GarbageCollectIsRunning_Error", "Cannot import asset '{0}'; the garbage collect is running during the import process.")
 			, FText::FromString(AssetName));
@@ -331,7 +337,7 @@ void UE::Interchange::FTaskImportObject_GameThread::DoTask(ENamedThreads::Type C
 		//Skip this asset
 		UInterchangeResultError_Generic* Message = Factory->AddMessage<UInterchangeResultError_Generic>();
 		Message->SourceAssetName = AsyncHelper->SourceDatas[SourceIndex]->GetFilename();
-		Message->DestinationAssetName = ReimportObject ? ReimportObject->GetName() : FactoryNode->GetAssetName();
+		Message->DestinationAssetName = ObjectToReimport ? ObjectToReimport->GetName() : FactoryNode->GetAssetName();
 		Message->AssetType = FactoryNode->GetObjectClass();
 		Message->Text = FText::Format(NSLOCTEXT("FTaskImportObject_GameThread", "InvalidBaseNodeContainer", "Cannot import asset '{0}'; the node container is invalid for this source index.")
 			, FText::FromString(AssetName));
@@ -341,25 +347,34 @@ void UE::Interchange::FTaskImportObject_GameThread::DoTask(ENamedThreads::Type C
 
 	UObject* ExistingAsset = nullptr;
 	//If we do a reimport no need to create a package
-	if (ReimportObject)
+	if (ObjectToReimport)
 	{
-		if (Private::ShouldReimportFactoryNode(FactoryNode, NodeContainer, ReimportObject))
+		if (Private::ShouldReimportFactoryNode(FactoryNode, NodeContainer, ObjectToReimport))
 		{
-			FactoryNode->SetDisplayLabel(ReimportObject->GetName());
-			FactoryNode->SetAssetName(ReimportObject->GetName());
-			Pkg = ReimportObject->GetPackage();
+			FactoryNode->SetDisplayLabel(ObjectToReimport->GetName());
+			FactoryNode->SetAssetName(ObjectToReimport->GetName());
+			Pkg = ObjectToReimport->GetPackage();
 			PackageName = Pkg->GetPathName();
-			AssetName = ReimportObject->GetName();
+			AssetName = ObjectToReimport->GetName();
 		}
 		else
 		{
 			//Skip this asset, reimport object original factory node does not match this factory node
 			return;
 		}
-		ExistingAsset = ReimportObject;
+		ExistingAsset = ObjectToReimport;
 	}
 	else
 	{
+		// Check if the object has not been deleted after an import
+		if (FFactoryCommon::GetFactoryNode(AsyncHelper->TaskData.ReimportObject, PackageName, AssetName))
+		{
+			if (!FactoryNode->ShouldForceNodeReimport())
+			{
+				return;
+			}
+		}
+
 		Private::InternalGetPackageName(*AsyncHelper, SourceIndex, PackageBasePath, FactoryNode, PackageName, AssetName);
 		// We can not create assets that share the name of a map file in the same location
 		if (FPackageUtils::IsMapPackageAsset(PackageName))
@@ -486,7 +501,7 @@ void UE::Interchange::FTaskImportObject_GameThread::DoTask(ENamedThreads::Type C
 	CreateAssetParams.SourceData = AsyncHelper->SourceDatas[SourceIndex];
 	CreateAssetParams.Translator = AsyncHelper->Translators[SourceIndex];
 	CreateAssetParams.NodeContainer = NodeContainer;
-	CreateAssetParams.ReimportObject = ReimportObject;
+	CreateAssetParams.ReimportObject = ObjectToReimport;
 	//Make sure the asset UObject is created with the correct type on the main thread
 	UInterchangeFactoryBase::FImportAssetResult ImportAssetResult = Factory->BeginImportAsset_GameThread(CreateAssetParams);
 	if (ImportAssetResult.ImportedObject)
@@ -505,11 +520,11 @@ void UE::Interchange::FTaskImportObject_GameThread::DoTask(ENamedThreads::Type C
 			AssetInfo.ImportedObject = ImportAssetResult.ImportedObject;
 			AssetInfo.Factory = Factory;
 			AssetInfo.FactoryNode = FactoryNode;
-			AssetInfo.bIsReimport = bool(ReimportObject != nullptr);
+			AssetInfo.bIsReimport = bool(ObjectToReimport != nullptr);
 		}
 		// Make sure the destination package is loaded
 		Pkg->FullyLoad();
-		if (!ReimportObject)
+		if (!ObjectToReimport)
 		{
 			AsyncHelper->AddCreatedPackage(PackageName, Pkg);
 		}
@@ -606,8 +621,8 @@ void UE::Interchange::FTaskImportObjectFinalize_GameThread::DoTask(ENamedThreads
 
 				const FString PackagePathName = ImportAssetResult.ImportedObject->GetPathName();
 				const FString AssetName = ImportAssetResult.ImportedObject->GetName();
-				const UObject* ReimportObject = FFactoryCommon::GetObjectToReimport(AsyncHelper->TaskData.ReimportObject, *FactoryNode, PackagePathName, AssetName);
-				AssetInfo.bIsReimport = bool(ReimportObject != nullptr);
+				const UObject* ObjectToReimport = FFactoryCommon::GetObjectToReimport(AsyncHelper->TaskData.ReimportObject, *FactoryNode, PackagePathName, AssetName);
+				AssetInfo.bIsReimport = bool(ObjectToReimport != nullptr);
 			}
 
 			// Fill in destination asset and type in any results which have been added previously by a translator or pipeline, now that we have a corresponding factory.
diff --git a/Engine/Source/Runtime/Interchange/Engine/Public/InterchangeAssetImportData.h b/Engine/Source/Runtime/Interchange/Engine/Public/InterchangeAssetImportData.h
--- a/Engine/Source/Runtime/Interchange/Engine/Public/InterchangeAssetImportData.h
+++ b/Engine/Source/Runtime/Interchange/Engine/Public/InterchangeAssetImportData.h
@@ -89,6 +89,29 @@ public:
 #endif
 #endif
 
+	/** On a level import, set to the UInterchangeSceneImportAsset created during the import */
+	UPROPERTY(EditAnywhere, Category = "Interchange | AssetImportData")
+	FSoftObjectPath SceneImportAsset;
+
+	/** Returns a pointer to the UInterchangeAssetImportData referred by the input object if applicable */
+	static UInterchangeAssetImportData* GetFromObject(UObject* Object)
+	{
+		if (Object)
+		{
+			TArray<UObject*> SubObjects;
+			GetObjectsWithOuter(Object, SubObjects);
+			for (UObject* SubObject : SubObjects)
+			{
+				if (UInterchangeAssetImportData* AssetImportData = Cast<UInterchangeAssetImportData>(SubObject))
+				{
+					return AssetImportData;
+				}
+			}
+		}
+
+		return nullptr;
+	}
+
 	/** The Node UID pass to the factory that exist in the graph that was use to create this asset */
 	UPROPERTY(VisibleAnywhere, Category = "Interchange | AssetImportData")
 	FString NodeUniqueID;
@@ -100,7 +123,6 @@ public:
 	UFUNCTION(BlueprintCallable, Category = "Interchange | AssetImportData")
 	INTERCHANGEENGINE_API void SetNodeContainer(UInterchangeBaseNodeContainer* InNodeContainer) const;
 
-
 	/**
 	* Returns Array of non-null pipelines
 	*/
diff --git a/Engine/Source/Runtime/Interchange/Engine/Public/InterchangeSceneImportAsset.h b/Engine/Source/Runtime/Interchange/Engine/Public/InterchangeSceneImportAsset.h
--- a/Engine/Source/Runtime/Interchange/Engine/Public/InterchangeSceneImportAsset.h
+++ b/Engine/Source/Runtime/Interchange/Engine/Public/InterchangeSceneImportAsset.h
@@ -71,6 +71,8 @@ public:
 	 */
 	INTERCHANGEENGINE_API const UInterchangeFactoryBaseNode* GetFactoryNode(const FString& PackageName, const FString& AssetName, const FString& SubPathString = FString()) const;
 
+	INTERCHANGEENGINE_API void GetSceneSoftObjectPaths(TArray<FSoftObjectPath>& SoftObjectPaths);
+
 private:
 
 #if WITH_EDITOR
