[Interchange] Fixing UpdateWeakPtrs call for InterchangeBasePipelines' Duplicate calls

diff --git a/Engine/Source/Runtime/Interchange/Core/Private/InterchangePipelineBase.cpp b/Engine/Source/Runtime/Interchange/Core/Private/InterchangePipelineBase.cpp
--- a/Engine/Source/Runtime/Interchange/Core/Private/InterchangePipelineBase.cpp
+++ b/Engine/Source/Runtime/Interchange/Core/Private/InterchangePipelineBase.cpp
@@ -410,4 +410,11 @@ void UInterchangePipelineBase::UpdateWeakObjectPtrs()
 			WeakObjectPtrData.WeakObjectProperty->SetPropertyValue(WeakObjectPtrData.ValuePtr, ObjectPtrs[WeakObjectPtrData.PropertyName]);
 		}
 	}
+}
+
+void UInterchangePipelineBase::PostDuplicate(bool bDuplicateForPIE)
+{
+	Super::PostDuplicate(bDuplicateForPIE);
+
+	UpdateWeakObjectPtrs();
 }
\ No newline at end of file
diff --git a/Engine/Source/Runtime/Interchange/Core/Public/InterchangePipelineBase.h b/Engine/Source/Runtime/Interchange/Core/Public/InterchangePipelineBase.h
--- a/Engine/Source/Runtime/Interchange/Core/Public/InterchangePipelineBase.h
+++ b/Engine/Source/Runtime/Interchange/Core/Public/InterchangePipelineBase.h
@@ -329,6 +329,10 @@ public:
 	*/
 	INTERCHANGECORE_API void UpdateWeakObjectPtrs();
 
+	/** Begin UObject overrides*/
+	INTERCHANGECORE_API virtual void PostDuplicate(bool bDuplicateForPIE) override;
+	/** End UObject overrides */
+
 protected:
 
 	UE_DEPRECATED(5.2, "This function is replace by ExecutePipeline.")
diff --git a/Engine/Source/Runtime/Interchange/Engine/Private/InterchangeManager.cpp b/Engine/Source/Runtime/Interchange/Engine/Private/InterchangeManager.cpp
--- a/Engine/Source/Runtime/Interchange/Engine/Private/InterchangeManager.cpp
+++ b/Engine/Source/Runtime/Interchange/Engine/Private/InterchangeManager.cpp
@@ -802,7 +802,6 @@ UInterchangePipelineBase* UE::Interchange::GeneratePipelineInstance(const FSoftO
 	{
 		// Make sure that the instance does not carry over standalone and public flags as they are not actual assets to be persisted
 		GeneratedPipeline->ClearFlags(EObjectFlags::RF_Standalone|EObjectFlags::RF_Public);
-		GeneratedPipeline->UpdateWeakObjectPtrs();
 	}
 
 	return GeneratedPipeline;
@@ -1485,7 +1484,6 @@ UInterchangeManager::ImportInternal(const FString& ContentPath, const UInterchan
 				{
 					//Duplicate the pipeline saved in the asset import data
 					UInterchangePipelineBase* GeneratedPipeline = Cast<UInterchangePipelineBase>(StaticDuplicateObject(SourcePipeline, PipelineInstancesPackage));
-					GeneratedPipeline->UpdateWeakObjectPtrs();
 					// Make sure that the instance does not carry over standalone and public flags as they are not actual assets to be persisted
 					GeneratedPipeline->ClearFlags(EObjectFlags::RF_Standalone | EObjectFlags::RF_Public);
 					AdjustPipelineSettingForContext(GeneratedPipeline);
diff --git a/Engine/Source/Runtime/Interchange/Engine/Private/InterchangeMeshUtilities.cpp b/Engine/Source/Runtime/Interchange/Engine/Private/InterchangeMeshUtilities.cpp
--- a/Engine/Source/Runtime/Interchange/Engine/Private/InterchangeMeshUtilities.cpp
+++ b/Engine/Source/Runtime/Interchange/Engine/Private/InterchangeMeshUtilities.cpp
@@ -135,7 +135,6 @@ TFuture<bool> UInterchangeMeshUtilities::InternalImportCustomLodAsync(TSharedPtr
 			}
 			if (ensure(GeneratedPipeline))
 			{
-				GeneratedPipeline->UpdateWeakObjectPtrs();
 				GeneratedPipeline->AdjustSettingsForContext(ImportType, nullptr);
 				ImportAssetParameters.OverridePipelines.Add(GeneratedPipeline);
 			}
