[Interchange] Changing EngineTests' expected warnings to successful imports

diff --git a/Engine/Plugins/Interchange/Runtime/Source/Import/Private/Gltf/InterchangeGltfTranslator.cpp b/Engine/Plugins/Interchange/Runtime/Source/Import/Private/Gltf/InterchangeGltfTranslator.cpp
--- a/Engine/Plugins/Interchange/Runtime/Source/Import/Private/Gltf/InterchangeGltfTranslator.cpp
+++ b/Engine/Plugins/Interchange/Runtime/Source/Import/Private/Gltf/InterchangeGltfTranslator.cpp
@@ -719,13 +719,20 @@ bool UInterchangeGLTFTranslator::Translate( UInterchangeBaseNodeContainer& NodeC
 			}
 			break;
 		case GLTF::EMessageSeverity::Warning:
-		default:
 			{
 				UInterchangeResultWarning_Generic* WarningResult = AddMessage< UInterchangeResultWarning_Generic >();
 				WarningResult->Text = FText::FromString(LogMessage.Get<1>());
 				Result = WarningResult;
 			}
-
+			break;
+		case GLTF::EMessageSeverity::Display:
+			{
+				UInterchangeResultDisplay_Generic* DisplayResult = AddMessage< UInterchangeResultDisplay_Generic >();
+				DisplayResult->Text = FText::FromString(LogMessage.Get<1>());
+				Result = DisplayResult;
+			}
+			break;
+		default:
 			break;
 		}
 
diff --git a/Engine/Plugins/Interchange/Runtime/Source/Parsers/GLTFCore/Private/GLTF/ExtensionsHandler.cpp b/Engine/Plugins/Interchange/Runtime/Source/Parsers/GLTFCore/Private/GLTF/ExtensionsHandler.cpp
--- a/Engine/Plugins/Interchange/Runtime/Source/Parsers/GLTFCore/Private/GLTF/ExtensionsHandler.cpp
+++ b/Engine/Plugins/Interchange/Runtime/Source/Parsers/GLTFCore/Private/GLTF/ExtensionsHandler.cpp
@@ -271,7 +271,7 @@ namespace GLTF
 				default:
 					if (!ensure(false))
 					{
-						Messages.Emplace(EMessageSeverity::Warning, FString::Printf(TEXT("Material.Extension not supported: %s"), *ToString(Extension)));
+						Messages.Emplace(RuntimeWarningSeverity(), FString::Printf(TEXT("Material.Extension not supported: %s"), *ToString(Extension)));
 					}
 					break;
 			}
@@ -352,7 +352,7 @@ namespace GLTF
 				default:
 					if (!ensure(false))
 					{
-						Messages.Emplace(EMessageSeverity::Warning, FString::Printf(TEXT("Primitive.Extension not supported: %s"), *ToString(Extension)));
+						Messages.Emplace(RuntimeWarningSeverity(), FString::Printf(TEXT("Primitive.Extension not supported: %s"), *ToString(Extension)));
 					}
 					break;
 			}
@@ -433,7 +433,7 @@ namespace GLTF
 		{
 			if (ExtensionsSupported.Find(StrValuePair.Get<0>()) == INDEX_NONE)
 			{
-				Messages.Emplace(EMessageSeverity::Warning, FString::Printf(TEXT("Extension is not supported: %s"), *StrValuePair.Get<0>()));
+				Messages.Emplace(RuntimeWarningSeverity(), FString::Printf(TEXT("Extension is not supported: %s"), *StrValuePair.Get<0>()));
 			}
 		}
 	}
@@ -466,7 +466,7 @@ namespace GLTF
 		else if (Type == TEXT("directional"))
 			Light.Type = FLight::EType::Directional;
 		else
-			Messages.Emplace(EMessageSeverity::Warning, FString::Printf(TEXT("Light has no type specified: %s"), *Light.Name));
+			Messages.Emplace(RuntimeWarningSeverity(), FString::Printf(TEXT("Light has no type specified: %s"), *Light.Name));
 	}
 
 }  // namespace GLTF
diff --git a/Engine/Plugins/Interchange/Runtime/Source/Parsers/GLTFCore/Private/GLTF/MaterialUtilities.cpp b/Engine/Plugins/Interchange/Runtime/Source/Parsers/GLTFCore/Private/GLTF/MaterialUtilities.cpp
--- a/Engine/Plugins/Interchange/Runtime/Source/Parsers/GLTFCore/Private/GLTF/MaterialUtilities.cpp
+++ b/Engine/Plugins/Interchange/Runtime/Source/Parsers/GLTFCore/Private/GLTF/MaterialUtilities.cpp
@@ -45,7 +45,7 @@ namespace GLTF
 						const FString ExtensionName = SupportedTextureInfoExtensionsStringified[Index];
 						if (!ExtensionsObj.HasTypedField<EJson::Object>(ExtensionName))
 						{
-							OutMessages.Emplace(EMessageSeverity::Warning, FString::Printf(TEXT("Extension is not supported: %s"), *ExtensionName));
+							OutMessages.Emplace(RuntimeWarningSeverity(), FString::Printf(TEXT("Extension is not supported: %s"), *ExtensionName));
 							continue;
 						}
 
diff --git a/Engine/Plugins/Interchange/Runtime/Source/Parsers/GLTFCore/Private/GLTFLogger.cpp b/Engine/Plugins/Interchange/Runtime/Source/Parsers/GLTFCore/Private/GLTFLogger.cpp
new file mode 100644
--- /dev/null
+++ b/Engine/Plugins/Interchange/Runtime/Source/Parsers/GLTFCore/Private/GLTFLogger.cpp
@@ -0,0 +1,20 @@
+// Copyright Epic Games, Inc. All Rights Reserved.
+
+#include "GLTFLogger.h"
+
+#include "CoreGlobals.h"
+
+namespace GLTF
+{
+	EMessageSeverity RuntimeWarningSeverity()
+	{
+		if (GIsAutomationTesting)
+		{
+			return EMessageSeverity::Display;
+		}
+		else
+		{
+			return EMessageSeverity::Warning;
+		}
+	}
+}  // namespace GLTF
diff --git a/Engine/Plugins/Interchange/Runtime/Source/Parsers/GLTFCore/Public/GLTFLogger.h b/Engine/Plugins/Interchange/Runtime/Source/Parsers/GLTFCore/Public/GLTFLogger.h
--- a/Engine/Plugins/Interchange/Runtime/Source/Parsers/GLTFCore/Public/GLTFLogger.h
+++ b/Engine/Plugins/Interchange/Runtime/Source/Parsers/GLTFCore/Public/GLTFLogger.h
@@ -9,9 +9,19 @@ namespace GLTF
 {
 	enum class EMessageSeverity
 	{
+		Display,
 		Warning,
 		Error,
 	};
+
+	/**
+	* There are certain warnings that we acknowledge and accept when running AutomationTests.
+	* For that reason we want Display severity for these specific warnings when Automation is running.
+	* For example not supported extension warnings are acknowledged and accepted when Automation is running,
+	* but we want those shown still as warnings when running otherwise.
+	*/
+	EMessageSeverity RuntimeWarningSeverity();
+
 	using FLogMessage = TTuple<EMessageSeverity, FString>;
 
 	class GLTFCORE_API FBaseLogger
diff --git a/Engine/Source/Runtime/Interchange/Core/Public/InterchangeResult.h b/Engine/Source/Runtime/Interchange/Core/Public/InterchangeResult.h
--- a/Engine/Source/Runtime/Interchange/Core/Public/InterchangeResult.h
+++ b/Engine/Source/Runtime/Interchange/Core/Public/InterchangeResult.h
@@ -149,3 +149,19 @@ public:
 
 	virtual FText GetText() const override { return NSLOCTEXT("InterchangeResultNS", "UInterchangeResultError_ReimportFail_GetText", "Re-import Fail, see log for more detail."); }
 };
+
+/**
+* Used for Successful imports with messages with lower than Warning priorities (for ep Display)
+*/
+UCLASS(MinimalAPI)
+class UInterchangeResultDisplay_Generic : public UInterchangeResultSuccess
+{
+	GENERATED_BODY()
+
+public:
+
+	virtual FText GetText() const override { return Text; }
+
+	UPROPERTY()
+	FText Text;
+};
\ No newline at end of file
diff --git a/Engine/Source/Runtime/Interchange/Engine/Private/Tasks/InterchangeTaskParsing.cpp b/Engine/Source/Runtime/Interchange/Engine/Private/Tasks/InterchangeTaskParsing.cpp
--- a/Engine/Source/Runtime/Interchange/Engine/Private/Tasks/InterchangeTaskParsing.cpp
+++ b/Engine/Source/Runtime/Interchange/Engine/Private/Tasks/InterchangeTaskParsing.cpp
@@ -385,11 +385,23 @@ void UE::Interchange::FTaskParsing::DoTask(ENamedThreads::Type CurrentThread, co
 		}
 		for (const TPair<TWeakObjectPtr<UInterchangeTranslatorBase>, FString>& MessagePerTranslator : TranslatorMessageMap)
 		{
-			UInterchangeResultWarning_Generic* WarningResult = NewObject<UInterchangeResultWarning_Generic>(GetTransientPackage(), UInterchangeResultWarning_Generic::StaticClass());
-			FString Message = NSLOCTEXT("InterchangeTaskParsingDoTask", "RenamedAssetsMessageHeader", "Renamed Assets:").ToString();
-			Message += MessagePerTranslator.Value;
-			WarningResult->Text = FText::FromString(Message);
-			MessagePerTranslator.Key->AddMessage(WarningResult);
+			if (GIsAutomationTesting)
+			{
+				
+				UInterchangeResultDisplay_Generic* DisplayResult = NewObject<UInterchangeResultDisplay_Generic>(GetTransientPackage(), UInterchangeResultDisplay_Generic::StaticClass());
+				FString Message = NSLOCTEXT("InterchangeTaskParsingDoTask", "RenamedAssetsMessageHeader", "Renamed Assets:").ToString();
+				Message += MessagePerTranslator.Value;
+				DisplayResult->Text = FText::FromString(Message);
+				MessagePerTranslator.Key->AddMessage(DisplayResult);
+			}
+			else
+			{
+				UInterchangeResultWarning_Generic* WarningResult = NewObject<UInterchangeResultWarning_Generic>(GetTransientPackage(), UInterchangeResultWarning_Generic::StaticClass());
+				FString Message = NSLOCTEXT("InterchangeTaskParsingDoTask", "RenamedAssetsMessageHeader", "Renamed Assets:").ToString();
+				Message += MessagePerTranslator.Value;
+				WarningResult->Text = FText::FromString(Message);
+				MessagePerTranslator.Key->AddMessage(WarningResult);
+			}
 		}
 	}
 
