Discard invalid data loaded while serializing UInterchangeAssetImportData

diff --git a/Engine/Source/Runtime/Interchange/Engine/Private/InterchangeAssetImportData.cpp b/Engine/Source/Runtime/Interchange/Engine/Private/InterchangeAssetImportData.cpp
--- a/Engine/Source/Runtime/Interchange/Engine/Private/InterchangeAssetImportData.cpp
+++ b/Engine/Source/Runtime/Interchange/Engine/Private/InterchangeAssetImportData.cpp
@@ -1,6 +1,7 @@
 // Copyright Epic Games, Inc. All Rights Reserved.
 
 #include "InterchangeAssetImportData.h"
+#include "InterchangeEngineLogPrivate.h"
 #include "InterchangeManager.h"
 #include "InterchangePipelineBase.h"
 
@@ -20,9 +21,22 @@ void UInterchangeAssetImportData::PostLoad()
 
 	if (NodeContainer_DEPRECATED)
 	{
-		SetNodeContainer(NodeContainer_DEPRECATED.Get());
+		bool HasNulls = false;
+		NodeContainer_DEPRECATED->IterateNodes([&HasNulls](const FString&UID, UInterchangeBaseNode*Node)
+		{
+			HasNulls |= !Node;
+		});
+
+		if (!HasNulls)
+		{
+			SetNodeContainer(NodeContainer_DEPRECATED.Get());
 
-		NodeContainer_DEPRECATED = nullptr;
+			NodeContainer_DEPRECATED = nullptr;
+		}
+		else
+		{
+			UE_LOG(LogInterchangeEngine, Display, TEXT("Missing Interchange reimport data for %s"), *this->GetOuter()->GetFullName());
+		}
 	}
 
 	if (Pipelines_DEPRECATED.Num() > 0)
