Interchange: Prevent saving debug pipeline in asset import data. The graph inspector pipeline will not be save anymore in any asset import data.

#rnx

diff --git a/Engine/Plugins/Interchange/Editor/Source/Pipelines/Public/InterchangeGraphInspectorPipeline.h b/Engine/Plugins/Interchange/Editor/Source/Pipelines/Public/InterchangeGraphInspectorPipeline.h
--- a/Engine/Plugins/Interchange/Editor/Source/Pipelines/Public/InterchangeGraphInspectorPipeline.h
+++ b/Engine/Plugins/Interchange/Editor/Source/Pipelines/Public/InterchangeGraphInspectorPipeline.h
@@ -28,6 +28,11 @@ public:
 	// END Pre import pipeline properties
 	//////////////////////////////////////////////////////////////////////////
 
+	/*
+	 * This pipeline must never be save into any asset import data
+	 */
+	virtual bool SupportReimport() const override { return false; }
+
 protected:
 
 	virtual void ExecutePipeline(UInterchangeBaseNodeContainer* BaseNodeContainer, const TArray<UInterchangeSourceData*>& SourceDatas) override;
diff --git a/Engine/Source/Runtime/Interchange/Core/Public/InterchangePipelineBase.h b/Engine/Source/Runtime/Interchange/Core/Public/InterchangePipelineBase.h
--- a/Engine/Source/Runtime/Interchange/Core/Public/InterchangePipelineBase.h
+++ b/Engine/Source/Runtime/Interchange/Core/Public/InterchangePipelineBase.h
@@ -333,6 +333,13 @@ public:
 	INTERCHANGECORE_API virtual void PostDuplicate(bool bDuplicateForPIE) override;
 	/** End UObject overrides */
 
+	/*
+	 * If it return true it mean we want to save this pipeline in the asset import data.
+	 * We will reuse this pipeline when re-importing the asset.
+	 * If false its probably a debug helper pipeline we do not want to save into assets
+	 */
+	INTERCHANGECORE_API virtual bool SupportReimport() const { return true; }
+
 protected:
 
 	UE_DEPRECATED(5.2, "This function is replace by ExecutePipeline.")
diff --git a/Engine/Source/Runtime/Interchange/Engine/Private/InterchangeImportCommon.cpp b/Engine/Source/Runtime/Interchange/Engine/Private/InterchangeImportCommon.cpp
--- a/Engine/Source/Runtime/Interchange/Engine/Private/InterchangeImportCommon.cpp
+++ b/Engine/Source/Runtime/Interchange/Engine/Private/InterchangeImportCommon.cpp
@@ -81,10 +81,27 @@ namespace UE::Interchange
 			TArray<UObject*> NewPipelines;
 			for (const UObject* Pipeline : Parameters.Pipelines)
 			{
-				UObject* DupPipeline = Cast<UObject>(StaticDuplicateObject(Pipeline, AssetImportData));
-				if (DupPipeline)
+				//Do not copy pipeline that cannot be use at reimport
+				bool bPipelineSupportReimport = false;
+				if (const UInterchangePythonPipelineAsset* PythonPipelineAsset = Cast<UInterchangePythonPipelineAsset>(Pipeline))
 				{
-					NewPipelines.Add(DupPipeline);
+					if (PythonPipelineAsset->GeneratedPipeline)
+					{
+						bPipelineSupportReimport = PythonPipelineAsset->GeneratedPipeline->SupportReimport();
+					}
+				}
+				else if (const UInterchangePipelineBase* PipelineBase = Cast<UInterchangePipelineBase>(Pipeline))
+				{
+					bPipelineSupportReimport = PipelineBase->SupportReimport();
+				}
+
+				if(bPipelineSupportReimport)
+				{
+					UObject* DupPipeline = Cast<UObject>(StaticDuplicateObject(Pipeline, AssetImportData));
+					if (DupPipeline)
+					{
+						NewPipelines.Add(DupPipeline);
+					}
 				}
 			}
 			AssetImportData->SetPipelines(NewPipelines);
diff --git a/Engine/Source/Runtime/Interchange/Engine/Private/InterchangeManager.cpp b/Engine/Source/Runtime/Interchange/Engine/Private/InterchangeManager.cpp
--- a/Engine/Source/Runtime/Interchange/Engine/Private/InterchangeManager.cpp
+++ b/Engine/Source/Runtime/Interchange/Engine/Private/InterchangeManager.cpp
@@ -1482,7 +1482,7 @@ UInterchangeManager::ImportInternal(const FString& ContentPath, const UInterchan
 						SourcePipeline = PythonPipelineAsset->GeneratedPipeline;
 					}
 				}
-				if (SourcePipeline) //Its possible a pipeline doesnt exist anymore so it wont load into memory when we loading the outer asset
+				if (SourcePipeline && SourcePipeline->SupportReimport()) //Its possible a pipeline doesnt exist anymore so it wont load into memory when we loading the outer asset
 				{
 					//Duplicate the pipeline saved in the asset import data
 					UInterchangePipelineBase* GeneratedPipeline = Cast<UInterchangePipelineBase>(StaticDuplicateObject(SourcePipeline, PipelineInstancesPackage));
@@ -1491,7 +1491,7 @@ UInterchangeManager::ImportInternal(const FString& ContentPath, const UInterchan
 					AdjustPipelineSettingForContext(GeneratedPipeline);
 					StackInfo.Pipelines.Add(GeneratedPipeline);
 				}
-				else
+				else if(!SourcePipeline)
 				{
 					//A pipeline was not loaded
 					UE_LOG(LogInterchangeEngine, Warning, TEXT("Interchange Reimport: Missing import pipeline from the reimpoting asset. The reimport might fail."));
