Support for "KeyValueProperties" through User Defined Attributes
- UserDefined Attributes can now be used to transfer the Properties which might be used along with add and fill delgates

diff --git a/Engine/Source/Runtime/Interchange/Core/Private/Nodes/InterchangeUserDefinedAttribute.cpp b/Engine/Source/Runtime/Interchange/Core/Private/Nodes/InterchangeUserDefinedAttribute.cpp
--- a/Engine/Source/Runtime/Interchange/Core/Private/Nodes/InterchangeUserDefinedAttribute.cpp
+++ b/Engine/Source/Runtime/Interchange/Core/Private/Nodes/InterchangeUserDefinedAttribute.cpp
@@ -1,88 +1,89 @@
 // Copyright Epic Games, Inc. All Rights Reserved.
 
 #include "Nodes/InterchangeUserDefinedAttribute.h"
-
+#include "Nodes/InterchangeFactoryBaseNode.h"
 #include "Nodes/InterchangeBaseNode.h"
 
 #include UE_INLINE_GENERATED_CPP_BY_NAME(InterchangeUserDefinedAttribute)
 
 const FString UInterchangeUserDefinedAttributesAPI::UserDefinedAttributeBaseKey = TEXT("UserDefined_");
+const FString UInterchangeUserDefinedAttributesAPI::UserDefinedAttributeDelegateKey = TEXT("AddDelegate_");
 const FString UInterchangeUserDefinedAttributesAPI::UserDefinedAttributeValuePostKey = TEXT("_Value");
 const FString UInterchangeUserDefinedAttributesAPI::UserDefinedAttributePayLoadPostKey = TEXT("_Payload");
 
-bool UInterchangeUserDefinedAttributesAPI::CreateUserDefinedAttribute_Boolean(UInterchangeBaseNode* InterchangeNode, const FString& UserDefinedAttributeName, const bool& Value, const FString& PayloadKey)
+bool UInterchangeUserDefinedAttributesAPI::CreateUserDefinedAttribute_Boolean(UInterchangeBaseNode* InterchangeNode, const FString& UserDefinedAttributeName, const bool& Value, const FString& PayloadKey, bool RequiresDelegate /*= false*/)
 {
 	TOptional<FString> OptionalPayload;
 	if (!PayloadKey.IsEmpty())
 	{
 		OptionalPayload = PayloadKey;
 	}
-	return UInterchangeUserDefinedAttributesAPI::CreateUserDefinedAttribute(InterchangeNode, UserDefinedAttributeName, Value, OptionalPayload);
+	return UInterchangeUserDefinedAttributesAPI::CreateUserDefinedAttribute(InterchangeNode, UserDefinedAttributeName, Value, OptionalPayload, RequiresDelegate);
 }
 
-bool UInterchangeUserDefinedAttributesAPI::CreateUserDefinedAttribute_Float(UInterchangeBaseNode* InterchangeNode, const FString& UserDefinedAttributeName, const float& Value, const FString& PayloadKey)
+bool UInterchangeUserDefinedAttributesAPI::CreateUserDefinedAttribute_Float(UInterchangeBaseNode* InterchangeNode, const FString& UserDefinedAttributeName, const float& Value, const FString& PayloadKey, bool RequiresDelegate /*= false*/)
 {
 	TOptional<FString> OptionalPayload;
 	if (!PayloadKey.IsEmpty())
 	{
 		OptionalPayload = PayloadKey;
 	}
-	return UInterchangeUserDefinedAttributesAPI::CreateUserDefinedAttribute(InterchangeNode, UserDefinedAttributeName, Value, OptionalPayload);
+	return UInterchangeUserDefinedAttributesAPI::CreateUserDefinedAttribute(InterchangeNode, UserDefinedAttributeName, Value, OptionalPayload, RequiresDelegate);
 }
 
-bool UInterchangeUserDefinedAttributesAPI::CreateUserDefinedAttribute_Double(UInterchangeBaseNode* InterchangeNode, const FString& UserDefinedAttributeName, const double& Value, const FString& PayloadKey)
+bool UInterchangeUserDefinedAttributesAPI::CreateUserDefinedAttribute_Double(UInterchangeBaseNode* InterchangeNode, const FString& UserDefinedAttributeName, const double& Value, const FString& PayloadKey, bool RequiresDelegate /*= false*/)
 {
 	TOptional<FString> OptionalPayload;
 	if (!PayloadKey.IsEmpty())
 	{
 		OptionalPayload = PayloadKey;
 	}
-	return UInterchangeUserDefinedAttributesAPI::CreateUserDefinedAttribute(InterchangeNode, UserDefinedAttributeName, Value, OptionalPayload);
+	return UInterchangeUserDefinedAttributesAPI::CreateUserDefinedAttribute(InterchangeNode, UserDefinedAttributeName, Value, OptionalPayload, RequiresDelegate);
 }
 
-bool UInterchangeUserDefinedAttributesAPI::CreateUserDefinedAttribute_Int32(UInterchangeBaseNode* InterchangeNode, const FString& UserDefinedAttributeName, const int32& Value, const FString& PayloadKey)
+bool UInterchangeUserDefinedAttributesAPI::CreateUserDefinedAttribute_Int32(UInterchangeBaseNode* InterchangeNode, const FString& UserDefinedAttributeName, const int32& Value, const FString& PayloadKey, bool RequiresDelegate /*= false*/)
 {
 	TOptional<FString> OptionalPayload;
 	if (!PayloadKey.IsEmpty())
 	{
 		OptionalPayload = PayloadKey;
 	}
-	return UInterchangeUserDefinedAttributesAPI::CreateUserDefinedAttribute(InterchangeNode, UserDefinedAttributeName, Value, OptionalPayload);
+	return UInterchangeUserDefinedAttributesAPI::CreateUserDefinedAttribute(InterchangeNode, UserDefinedAttributeName, Value, OptionalPayload, RequiresDelegate);
 }
 
-bool UInterchangeUserDefinedAttributesAPI::CreateUserDefinedAttribute_FString(UInterchangeBaseNode* InterchangeNode, const FString& UserDefinedAttributeName, const FString& Value, const FString& PayloadKey)
+bool UInterchangeUserDefinedAttributesAPI::CreateUserDefinedAttribute_FString(UInterchangeBaseNode* InterchangeNode, const FString& UserDefinedAttributeName, const FString& Value, const FString& PayloadKey, bool RequiresDelegate /*= false*/)
 {
 	TOptional<FString> OptionalPayload;
 	if (!PayloadKey.IsEmpty())
 	{
 		OptionalPayload = PayloadKey;
 	}
-	return UInterchangeUserDefinedAttributesAPI::CreateUserDefinedAttribute(InterchangeNode, UserDefinedAttributeName, Value, OptionalPayload);
+	return UInterchangeUserDefinedAttributesAPI::CreateUserDefinedAttribute(InterchangeNode, UserDefinedAttributeName, Value, OptionalPayload, RequiresDelegate);
 }
 
 bool UInterchangeUserDefinedAttributesAPI::RemoveUserDefinedAttribute(UInterchangeBaseNode* InterchangeNode, const FString& UserDefinedAttributeName)
 {
-
-	FString StorageBaseKey = UserDefinedAttributeBaseKey + UserDefinedAttributeName;
-	UE::Interchange::FAttributeKey UserDefinedValueKey = UE::Interchange::FAttributeKey(StorageBaseKey + UserDefinedAttributeValuePostKey);
-	if (InterchangeNode->HasAttribute(UserDefinedValueKey))
+	bool RequiresDelegates;
+	const bool bGeneratePayloadKey = false;
+	if (HasAttribute(InterchangeNode, UserDefinedAttributeName, bGeneratePayloadKey, RequiresDelegates))
 	{
+		UE::Interchange::FAttributeKey UserDefinedValueKey = MakeUserDefinedPropertyValueKey(UserDefinedAttributeName, RequiresDelegates);
 		if (!InterchangeNode->RemoveAttribute(UserDefinedValueKey.Key))
 		{
 			return false;
 		}
-	}
 
-	UE::Interchange::FAttributeKey UserDefinedPayloadKey = UE::Interchange::FAttributeKey(StorageBaseKey + UserDefinedAttributePayLoadPostKey);
-	if (InterchangeNode->HasAttribute(UserDefinedPayloadKey))
-	{
-		if (!InterchangeNode->RemoveAttribute(UserDefinedPayloadKey.Key))
+		UE::Interchange::FAttributeKey UserDefinedPayloadKey = MakeUserDefinedPropertyPayloadKey(UserDefinedAttributeName, RequiresDelegates);
+		if (InterchangeNode->HasAttribute(UserDefinedPayloadKey))
 		{
-			return false;
+			if (!InterchangeNode->RemoveAttribute(UserDefinedPayloadKey.Key))
+			{
+				return false;
+			}
 		}
 	}
 
-	//Attribute was successfuly removed
+	//Attribute was successfully removed/ never existed
 	return true;
 }
 
@@ -149,19 +150,28 @@ TArray<FInterchangeUserDefinedAttributeInfo> UInterchangeUserDefinedAttributesAP
 	TArray<FInterchangeUserDefinedAttributeInfo> UserDefinedAttributeInfos;
 	int32 RightChopIndex = UserDefinedAttributeBaseKey.Len();
 	int32 LeftChopIndex = UserDefinedAttributeValuePostKey.Len();
+	int32 AddDelegateRightChopIndex = UserDefinedAttributeDelegateKey.Len();
+
 	for (UE::Interchange::FAttributeKey& AttributeKey : AttributeKeys)
 	{
 		const FString AttributeKeyString = AttributeKey.ToString();
 		if (AttributeKeyString.StartsWith(UserDefinedAttributeBaseKey) && AttributeKeyString.EndsWith(UserDefinedAttributeValuePostKey))
 		{
+			bool RequiresDelegate = false;
 			FString UserDefinedAttributeName = AttributeKeyString.RightChop(RightChopIndex).LeftChop(LeftChopIndex);
+			if (UserDefinedAttributeName.StartsWith(UserDefinedAttributeDelegateKey))
+			{
+				UserDefinedAttributeName = UserDefinedAttributeName.RightChop(AddDelegateRightChopIndex);
+				RequiresDelegate = true;
+			}
 			FInterchangeUserDefinedAttributeInfo& UserDefinedAttributeInfo = UserDefinedAttributeInfos.AddDefaulted_GetRef();
 			UserDefinedAttributeInfo.Type = InterchangeNode->GetAttributeType(AttributeKey);
 			UserDefinedAttributeInfo.Name = UserDefinedAttributeName;
+			UserDefinedAttributeInfo.RequiresDelegate = RequiresDelegate;
 
 			//Get the optional payload key
 			const FString StorageBaseKey = UserDefinedAttributeBaseKey + UserDefinedAttributeName;
-			const UE::Interchange::FAttributeKey UserDefinedPayloadKey = UE::Interchange::FAttributeKey(StorageBaseKey + UserDefinedAttributePayLoadPostKey);
+			const UE::Interchange::FAttributeKey UserDefinedPayloadKey = MakeUserDefinedPropertyPayloadKey(UserDefinedAttributeName, RequiresDelegate);
 			if (InterchangeNode->HasAttribute(UserDefinedPayloadKey))
 			{
 				FString PayloadKey;
@@ -187,15 +197,21 @@ void UInterchangeUserDefinedAttributesAPI::DuplicateAllUserDefinedAttribute(cons
 	InterchangeSourceNode->GetAttributeKeys(AttributeKeys);
 	int32 RightChopIndex = UserDefinedAttributeBaseKey.Len();
 	int32 LeftChopIndex = UserDefinedAttributeValuePostKey.Len();
+	int32 AddDelegateRightChopIndex = UserDefinedAttributeDelegateKey.Len();
+
 	for (UE::Interchange::FAttributeKey& AttributeKey : AttributeKeys)
 	{
 		const FString AttributeKeyString = AttributeKey.ToString();
 		if (AttributeKeyString.StartsWith(UserDefinedAttributeBaseKey) && AttributeKeyString.EndsWith(UserDefinedAttributeValuePostKey))
 		{
+			bool RequiresDelegate = false;
 			FString UserDefinedAttributeName = AttributeKeyString.RightChop(RightChopIndex).LeftChop(LeftChopIndex);
-			//Get the optional payload key
-			const FString StorageBaseKey = UserDefinedAttributeBaseKey + UserDefinedAttributeName;
-
+			if (UserDefinedAttributeName.StartsWith(UserDefinedAttributeDelegateKey))
+			{
+				UserDefinedAttributeName = UserDefinedAttributeName.RightChop(AddDelegateRightChopIndex);
+				RequiresDelegate = true;
+			}
+			
 			UE::Interchange::EAttributeTypes Type = InterchangeSourceNode->GetAttributeType(AttributeKey);
 			FString DuplicateName = UserDefinedAttributeName;
 			if (bAddSourceNodeName)
@@ -203,7 +219,8 @@ void UInterchangeUserDefinedAttributesAPI::DuplicateAllUserDefinedAttribute(cons
 				DuplicateName = InterchangeSourceNode->GetDisplayLabel() + TEXT(".") + UserDefinedAttributeName;
 			}
 
-			const UE::Interchange::FAttributeKey UserDefinedPayloadKey = UE::Interchange::FAttributeKey(StorageBaseKey + UserDefinedAttributePayLoadPostKey);
+			//Get the optional payload key
+			const UE::Interchange::FAttributeKey UserDefinedPayloadKey = MakeUserDefinedPropertyPayloadKey(UserDefinedAttributeName, RequiresDelegate);
 			TOptional<FString> PayloadKey;
 			if (InterchangeSourceNode->HasAttribute(UserDefinedPayloadKey))
 			{
@@ -214,64 +231,194 @@ void UInterchangeUserDefinedAttributesAPI::DuplicateAllUserDefinedAttribute(cons
 			switch (Type)
 			{
 				case UE::Interchange::EAttributeTypes::Bool:
-					CreateUserDefinedAttribute(InterchangeDestinationNode, DuplicateName, InterchangeSourceNode->GetAttributeChecked<bool>(AttributeKey.Key), PayloadKey);
+					CreateUserDefinedAttribute(InterchangeDestinationNode, DuplicateName, InterchangeSourceNode->GetAttributeChecked<bool>(AttributeKey.Key), PayloadKey, RequiresDelegate);
 					break;
 				case UE::Interchange::EAttributeTypes::Int8:
-					CreateUserDefinedAttribute(InterchangeDestinationNode, DuplicateName, InterchangeSourceNode->GetAttributeChecked<int8>(AttributeKey.Key), PayloadKey);
+					CreateUserDefinedAttribute(InterchangeDestinationNode, DuplicateName, InterchangeSourceNode->GetAttributeChecked<int8>(AttributeKey.Key), PayloadKey, RequiresDelegate);
 					break;
 				case UE::Interchange::EAttributeTypes::Int16:
-					CreateUserDefinedAttribute(InterchangeDestinationNode, DuplicateName, InterchangeSourceNode->GetAttributeChecked<int16>(AttributeKey.Key), PayloadKey);
+					CreateUserDefinedAttribute(InterchangeDestinationNode, DuplicateName, InterchangeSourceNode->GetAttributeChecked<int16>(AttributeKey.Key), PayloadKey, RequiresDelegate);
 					break;
 				case UE::Interchange::EAttributeTypes::Int32:
-					CreateUserDefinedAttribute(InterchangeDestinationNode, DuplicateName, InterchangeSourceNode->GetAttributeChecked<int32>(AttributeKey.Key), PayloadKey);
+					CreateUserDefinedAttribute(InterchangeDestinationNode, DuplicateName, InterchangeSourceNode->GetAttributeChecked<int32>(AttributeKey.Key), PayloadKey, RequiresDelegate);
 					break;
 				case UE::Interchange::EAttributeTypes::Int64:
-					CreateUserDefinedAttribute(InterchangeDestinationNode, DuplicateName, InterchangeSourceNode->GetAttributeChecked<int64>(AttributeKey.Key), PayloadKey);
+					CreateUserDefinedAttribute(InterchangeDestinationNode, DuplicateName, InterchangeSourceNode->GetAttributeChecked<int64>(AttributeKey.Key), PayloadKey, RequiresDelegate);
 					break;
 				case UE::Interchange::EAttributeTypes::UInt8:
-					CreateUserDefinedAttribute(InterchangeDestinationNode, DuplicateName, InterchangeSourceNode->GetAttributeChecked<uint8>(AttributeKey.Key), PayloadKey);
+					CreateUserDefinedAttribute(InterchangeDestinationNode, DuplicateName, InterchangeSourceNode->GetAttributeChecked<uint8>(AttributeKey.Key), PayloadKey, RequiresDelegate);
 					break;
 				case UE::Interchange::EAttributeTypes::UInt16:
-					CreateUserDefinedAttribute(InterchangeDestinationNode, DuplicateName, InterchangeSourceNode->GetAttributeChecked<uint16>(AttributeKey.Key), PayloadKey);
+					CreateUserDefinedAttribute(InterchangeDestinationNode, DuplicateName, InterchangeSourceNode->GetAttributeChecked<uint16>(AttributeKey.Key), PayloadKey, RequiresDelegate);
 					break;
 				case UE::Interchange::EAttributeTypes::UInt32:
-					CreateUserDefinedAttribute(InterchangeDestinationNode, DuplicateName, InterchangeSourceNode->GetAttributeChecked<uint32>(AttributeKey.Key), PayloadKey);
+					CreateUserDefinedAttribute(InterchangeDestinationNode, DuplicateName, InterchangeSourceNode->GetAttributeChecked<uint32>(AttributeKey.Key), PayloadKey, RequiresDelegate);
 					break;
 				case UE::Interchange::EAttributeTypes::UInt64:
-					CreateUserDefinedAttribute(InterchangeDestinationNode, DuplicateName, InterchangeSourceNode->GetAttributeChecked<uint64>(AttributeKey.Key), PayloadKey);
+					CreateUserDefinedAttribute(InterchangeDestinationNode, DuplicateName, InterchangeSourceNode->GetAttributeChecked<uint64>(AttributeKey.Key), PayloadKey, RequiresDelegate);
 					break;
 				case UE::Interchange::EAttributeTypes::Float:
-					CreateUserDefinedAttribute(InterchangeDestinationNode, DuplicateName, InterchangeSourceNode->GetAttributeChecked<float>(AttributeKey.Key), PayloadKey);
+					CreateUserDefinedAttribute(InterchangeDestinationNode, DuplicateName, InterchangeSourceNode->GetAttributeChecked<float>(AttributeKey.Key), PayloadKey, RequiresDelegate);
 					break;
 				case UE::Interchange::EAttributeTypes::Float16:
-					CreateUserDefinedAttribute(InterchangeDestinationNode, DuplicateName, InterchangeSourceNode->GetAttributeChecked<FFloat16>(AttributeKey.Key), PayloadKey);
+					CreateUserDefinedAttribute(InterchangeDestinationNode, DuplicateName, InterchangeSourceNode->GetAttributeChecked<FFloat16>(AttributeKey.Key), PayloadKey, RequiresDelegate);
 					break;
 				case UE::Interchange::EAttributeTypes::Vector2f:
-					CreateUserDefinedAttribute(InterchangeDestinationNode, DuplicateName, InterchangeSourceNode->GetAttributeChecked<FVector2f>(AttributeKey.Key), PayloadKey);
+					CreateUserDefinedAttribute(InterchangeDestinationNode, DuplicateName, InterchangeSourceNode->GetAttributeChecked<FVector2f>(AttributeKey.Key), PayloadKey, RequiresDelegate);
 					break;
 				case UE::Interchange::EAttributeTypes::Vector3f:
-					CreateUserDefinedAttribute(InterchangeDestinationNode, DuplicateName, InterchangeSourceNode->GetAttributeChecked<FVector3f>(AttributeKey.Key), PayloadKey);
+					CreateUserDefinedAttribute(InterchangeDestinationNode, DuplicateName, InterchangeSourceNode->GetAttributeChecked<FVector3f>(AttributeKey.Key), PayloadKey, RequiresDelegate);
 					break;
 				case UE::Interchange::EAttributeTypes::Vector4f:
-					CreateUserDefinedAttribute(InterchangeDestinationNode, DuplicateName, InterchangeSourceNode->GetAttributeChecked<FVector4f>(AttributeKey.Key), PayloadKey);
+					CreateUserDefinedAttribute(InterchangeDestinationNode, DuplicateName, InterchangeSourceNode->GetAttributeChecked<FVector4f>(AttributeKey.Key), PayloadKey, RequiresDelegate);
 					break;
 				case UE::Interchange::EAttributeTypes::Double:
-					CreateUserDefinedAttribute(InterchangeDestinationNode, DuplicateName, InterchangeSourceNode->GetAttributeChecked<double>(AttributeKey.Key), PayloadKey);
+					CreateUserDefinedAttribute(InterchangeDestinationNode, DuplicateName, InterchangeSourceNode->GetAttributeChecked<double>(AttributeKey.Key), PayloadKey, RequiresDelegate);
 					break;
 				case UE::Interchange::EAttributeTypes::Vector2d:
-					CreateUserDefinedAttribute(InterchangeDestinationNode, DuplicateName, InterchangeSourceNode->GetAttributeChecked<FVector2D>(AttributeKey.Key), PayloadKey);
+					CreateUserDefinedAttribute(InterchangeDestinationNode, DuplicateName, InterchangeSourceNode->GetAttributeChecked<FVector2D>(AttributeKey.Key), PayloadKey, RequiresDelegate);
 					break;
 				case UE::Interchange::EAttributeTypes::Vector3d:
-					CreateUserDefinedAttribute(InterchangeDestinationNode, DuplicateName, InterchangeSourceNode->GetAttributeChecked<FVector3d>(AttributeKey.Key), PayloadKey);
+					CreateUserDefinedAttribute(InterchangeDestinationNode, DuplicateName, InterchangeSourceNode->GetAttributeChecked<FVector3d>(AttributeKey.Key), PayloadKey, RequiresDelegate);
 					break;
 				case UE::Interchange::EAttributeTypes::Vector4d:
-					CreateUserDefinedAttribute(InterchangeDestinationNode, DuplicateName, InterchangeSourceNode->GetAttributeChecked<FVector4d>(AttributeKey.Key), PayloadKey);
+					CreateUserDefinedAttribute(InterchangeDestinationNode, DuplicateName, InterchangeSourceNode->GetAttributeChecked<FVector4d>(AttributeKey.Key), PayloadKey, RequiresDelegate);
 					break;
 				case UE::Interchange::EAttributeTypes::String:
-					CreateUserDefinedAttribute(InterchangeDestinationNode, DuplicateName, InterchangeSourceNode->GetAttributeChecked<FString>(AttributeKey.Key), PayloadKey);
+					CreateUserDefinedAttribute(InterchangeDestinationNode, DuplicateName, InterchangeSourceNode->GetAttributeChecked<FString>(AttributeKey.Key), PayloadKey, RequiresDelegate);
 					break;
 			}
 		}
 	}
 }
 
+void UInterchangeUserDefinedAttributesAPI::AddApplyAndFillDelegatesToFactory(UInterchangeFactoryBaseNode* InterchangeFactoryNode, UClass* ParentClass)
+{
+	TArray<UE::Interchange::FAttributeKey> AttributeKeys;
+	InterchangeFactoryNode->GetAttributeKeys(AttributeKeys);
+	int32 RightChopIndex = UserDefinedAttributeBaseKey.Len();
+	int32 LeftChopIndex = UserDefinedAttributeValuePostKey.Len();
+	int32 AddDelegateRightChopIndex = UserDefinedAttributeDelegateKey.Len();
+
+	for (UE::Interchange::FAttributeKey& AttributeKey : AttributeKeys)
+	{
+		const FString AttributeKeyString = AttributeKey.ToString();
+		if (AttributeKeyString.StartsWith(UserDefinedAttributeBaseKey) && AttributeKeyString.EndsWith(UserDefinedAttributeValuePostKey))
+		{
+			bool RequiresDelegate = false;
+			FString UserDefinedAttributeName = AttributeKeyString.RightChop(RightChopIndex).LeftChop(LeftChopIndex);
+			if (!UserDefinedAttributeName.StartsWith(UserDefinedAttributeDelegateKey))
+			{
+				continue;				
+			}
+
+			UserDefinedAttributeName = UserDefinedAttributeName.RightChop(AddDelegateRightChopIndex);
+			RequiresDelegate = true;
+
+			UE::Interchange::EAttributeTypes Type = InterchangeFactoryNode->GetAttributeType(AttributeKey);
+			switch (Type)
+			{
+			case UE::Interchange::EAttributeTypes::Bool:
+				InterchangeFactoryNode->AddApplyAndFillDelegates<bool>(AttributeKey.Key, ParentClass, *UserDefinedAttributeName);
+				break;
+			case UE::Interchange::EAttributeTypes::Int8:
+				InterchangeFactoryNode->AddApplyAndFillDelegates<int8>(AttributeKey.Key, ParentClass, *UserDefinedAttributeName);
+				break;
+			case UE::Interchange::EAttributeTypes::Int16:
+				InterchangeFactoryNode->AddApplyAndFillDelegates<int16>(AttributeKey.Key, ParentClass, *UserDefinedAttributeName);
+				break;
+			case UE::Interchange::EAttributeTypes::Int32:
+				InterchangeFactoryNode->AddApplyAndFillDelegates<int32>(AttributeKey.Key, ParentClass, *UserDefinedAttributeName);
+				break;
+			case UE::Interchange::EAttributeTypes::Int64:
+				InterchangeFactoryNode->AddApplyAndFillDelegates<int64>(AttributeKey.Key, ParentClass, *UserDefinedAttributeName);
+				break;
+			case UE::Interchange::EAttributeTypes::UInt8:
+				InterchangeFactoryNode->AddApplyAndFillDelegates<uint8>(AttributeKey.Key, ParentClass, *UserDefinedAttributeName);
+				break;
+			case UE::Interchange::EAttributeTypes::UInt16:
+				InterchangeFactoryNode->AddApplyAndFillDelegates<uint16>(AttributeKey.Key, ParentClass, *UserDefinedAttributeName);
+				break;
+			case UE::Interchange::EAttributeTypes::UInt32:
+				InterchangeFactoryNode->AddApplyAndFillDelegates<uint32>(AttributeKey.Key, ParentClass, *UserDefinedAttributeName);
+				break;
+			case UE::Interchange::EAttributeTypes::UInt64:
+				InterchangeFactoryNode->AddApplyAndFillDelegates<uint64>(AttributeKey.Key, ParentClass, *UserDefinedAttributeName);
+				break;
+			case UE::Interchange::EAttributeTypes::Float:
+				InterchangeFactoryNode->AddApplyAndFillDelegates<float>(AttributeKey.Key, ParentClass, *UserDefinedAttributeName);
+				break;
+			case UE::Interchange::EAttributeTypes::Float16:
+				InterchangeFactoryNode->AddApplyAndFillDelegates<FFloat16>(AttributeKey.Key, ParentClass, *UserDefinedAttributeName);
+				break;
+			case UE::Interchange::EAttributeTypes::Vector2f:
+				InterchangeFactoryNode->AddApplyAndFillDelegates<FVector2f>(AttributeKey.Key, ParentClass, *UserDefinedAttributeName);
+				break;
+			case UE::Interchange::EAttributeTypes::Vector3f:
+				InterchangeFactoryNode->AddApplyAndFillDelegates<FVector3f>(AttributeKey.Key, ParentClass, *UserDefinedAttributeName);
+				break;
+			case UE::Interchange::EAttributeTypes::Vector4f:
+				InterchangeFactoryNode->AddApplyAndFillDelegates<FVector4f>(AttributeKey.Key, ParentClass, *UserDefinedAttributeName);
+				break;
+			case UE::Interchange::EAttributeTypes::Double:
+				InterchangeFactoryNode->AddApplyAndFillDelegates<double>(AttributeKey.Key, ParentClass, *UserDefinedAttributeName);
+				break;
+			case UE::Interchange::EAttributeTypes::Vector2d:
+				InterchangeFactoryNode->AddApplyAndFillDelegates<FVector2D>(AttributeKey.Key, ParentClass, *UserDefinedAttributeName);
+				break;
+			case UE::Interchange::EAttributeTypes::Vector3d:
+				InterchangeFactoryNode->AddApplyAndFillDelegates<FVector3d>(AttributeKey.Key, ParentClass, *UserDefinedAttributeName);
+				break;
+			case UE::Interchange::EAttributeTypes::Vector4d:
+				InterchangeFactoryNode->AddApplyAndFillDelegates<FVector4d>(AttributeKey.Key, ParentClass, *UserDefinedAttributeName);
+				break;
+			case UE::Interchange::EAttributeTypes::String:
+				InterchangeFactoryNode->AddApplyAndFillDelegates<FString>(AttributeKey.Key, ParentClass, *UserDefinedAttributeName);
+				break;
+			}
+		}
+	}
+}
+
+UE::Interchange::FAttributeKey UInterchangeUserDefinedAttributesAPI::MakeUserDefinedPropertyValueKey(const FString& UserDefinedAttributeName, bool RequiresDelegate)
+{
+	return MakeUserDefinedPropertyKey(UserDefinedAttributeName, RequiresDelegate);
+}
+
+UE::Interchange::FAttributeKey UInterchangeUserDefinedAttributesAPI::MakeUserDefinedPropertyPayloadKey(const FString& UserDefinedAttributeName, bool RequiresDelegate)
+{
+	const bool bGeneratePayloadKey = true;
+	return MakeUserDefinedPropertyKey(UserDefinedAttributeName, RequiresDelegate, bGeneratePayloadKey);
+}
+
+bool UInterchangeUserDefinedAttributesAPI::HasAttribute(const UInterchangeBaseNode* InterchangeSourceNode, const FString& InUserDefinedAttributeName, bool GeneratePayloadKey, bool& OutRequiresDelegate)
+{
+	UE::Interchange::FAttributeKey KeyWithDelegate = MakeUserDefinedPropertyKey(InUserDefinedAttributeName,/*RequiresDelegate =*/ true, GeneratePayloadKey);
+	UE::Interchange::FAttributeKey KeyWithoutDelegate = MakeUserDefinedPropertyKey(InUserDefinedAttributeName,/*RequiresDelegate =*/ false, GeneratePayloadKey);
+	OutRequiresDelegate = true;
+	if (InterchangeSourceNode->HasAttribute(KeyWithDelegate))
+	{
+		return true;
+	}
+	else if (InterchangeSourceNode->HasAttribute(KeyWithoutDelegate))
+	{
+		OutRequiresDelegate = false;
+		return true;
+	}
+
+	return false;
+}
+
+UE::Interchange::FAttributeKey UInterchangeUserDefinedAttributesAPI::MakeUserDefinedPropertyKey(const FString& UserDefinedAttributeName, bool RequiresDelegate, bool GeneratePayloadKey /*= false*/)
+{ 
+	//Create a unique Key for this user defined attribute
+	FString Prefix = UserDefinedAttributeBaseKey;
+	if (RequiresDelegate)
+	{
+		Prefix = Prefix + UserDefinedAttributeDelegateKey;
+	}
+
+	const FString Suffix = GeneratePayloadKey ? UserDefinedAttributePayLoadPostKey : UserDefinedAttributeValuePostKey;
+
+	return UE::Interchange::FAttributeKey(Prefix + UserDefinedAttributeName + Suffix);
+}
\ No newline at end of file
diff --git a/Engine/Source/Runtime/Interchange/Core/Public/Nodes/InterchangeUserDefinedAttribute.h b/Engine/Source/Runtime/Interchange/Core/Public/Nodes/InterchangeUserDefinedAttribute.h
--- a/Engine/Source/Runtime/Interchange/Core/Public/Nodes/InterchangeUserDefinedAttribute.h
+++ b/Engine/Source/Runtime/Interchange/Core/Public/Nodes/InterchangeUserDefinedAttribute.h
@@ -18,6 +18,7 @@
 #include "InterchangeUserDefinedAttribute.generated.h"
 
 struct FFrame;
+class UInterchangeFactoryBaseNode;
 
 USTRUCT(BlueprintType)
 struct FInterchangeUserDefinedAttributeInfo
@@ -30,6 +31,8 @@ struct FInterchangeUserDefinedAttributeInfo
 	UE::Interchange::EAttributeTypes Type;
 
 	TOptional<FString> PayloadKey;
+
+	bool RequiresDelegate;
 };
 
 /**
@@ -53,22 +56,22 @@ public:
 	 *        Payload key will point on a FRichCurve payload.
 	 */
 	template<typename ValueType>
-	static bool CreateUserDefinedAttribute(UInterchangeBaseNode* InterchangeNode, const FString& UserDefinedAttributeName, const ValueType& Value, const TOptional<FString>& PayloadKey);
+	static bool CreateUserDefinedAttribute(UInterchangeBaseNode* InterchangeNode, const FString& UserDefinedAttributeName, const ValueType& Value, const TOptional<FString>& PayloadKey, bool RequiresDelegate = false);
 
 	UFUNCTION(BlueprintCallable, Category = "Interchange | Node | UserDefinedAttribute")
-	static INTERCHANGECORE_API bool CreateUserDefinedAttribute_Boolean(UInterchangeBaseNode* InterchangeNode, const FString& UserDefinedAttributeName, const bool& Value, const FString& PayloadKey);
-	
+	static INTERCHANGECORE_API bool CreateUserDefinedAttribute_Boolean(UInterchangeBaseNode* InterchangeNode, const FString& UserDefinedAttributeName, const bool& Value, const FString& PayloadKey, bool RequiresDelegate = false);
+
 	UFUNCTION(BlueprintCallable, Category = "Interchange | Node | UserDefinedAttribute")
-	static INTERCHANGECORE_API bool CreateUserDefinedAttribute_Float(UInterchangeBaseNode* InterchangeNode, const FString& UserDefinedAttributeName, const float& Value, const FString& PayloadKey);
+	static INTERCHANGECORE_API bool CreateUserDefinedAttribute_Float(UInterchangeBaseNode* InterchangeNode, const FString& UserDefinedAttributeName, const float& Value, const FString& PayloadKey, bool RequiresDelegate = false);
 
 	UFUNCTION(BlueprintCallable, Category = "Interchange | Node | UserDefinedAttribute")
-	static INTERCHANGECORE_API bool CreateUserDefinedAttribute_Double(UInterchangeBaseNode* InterchangeNode, const FString& UserDefinedAttributeName, const double& Value, const FString& PayloadKey);
+	static INTERCHANGECORE_API bool CreateUserDefinedAttribute_Double(UInterchangeBaseNode* InterchangeNode, const FString& UserDefinedAttributeName, const double& Value, const FString& PayloadKey, bool RequiresDelegate = false);
 
 	UFUNCTION(BlueprintCallable, Category = "Interchange | Node | UserDefinedAttribute")
-	static INTERCHANGECORE_API bool CreateUserDefinedAttribute_Int32(UInterchangeBaseNode* InterchangeNode, const FString& UserDefinedAttributeName, const int32& Value, const FString& PayloadKey);
+	static INTERCHANGECORE_API bool CreateUserDefinedAttribute_Int32(UInterchangeBaseNode* InterchangeNode, const FString& UserDefinedAttributeName, const int32& Value, const FString& PayloadKey, bool RequiresDelegate = false);
 
 	UFUNCTION(BlueprintCallable, Category = "Interchange | Node | UserDefinedAttribute")
-	static INTERCHANGECORE_API bool CreateUserDefinedAttribute_FString(UInterchangeBaseNode* InterchangeNode, const FString& UserDefinedAttributeName, const FString& Value, const FString& PayloadKey);
+	static INTERCHANGECORE_API bool CreateUserDefinedAttribute_FString(UInterchangeBaseNode* InterchangeNode, const FString& UserDefinedAttributeName, const FString& Value, const FString& PayloadKey, bool RequiresDelegate = false);
 
 	/**
 	 * Remove the specified user defined attribute
@@ -114,24 +117,35 @@ public:
 	UFUNCTION(BlueprintCallable, Category = "Interchange | Node | UserDefinedAttribute")
 	static INTERCHANGECORE_API void DuplicateAllUserDefinedAttribute(const UInterchangeBaseNode* InterchangeSourceNode, UInterchangeBaseNode* InterchangeDestinationNode, bool bAddSourceNodeName);
 
+	static INTERCHANGECORE_API void AddApplyAndFillDelegatesToFactory(UInterchangeFactoryBaseNode* InterchangeFactoryNode, UClass* ParentClass);
+
+	static INTERCHANGECORE_API UE::Interchange::FAttributeKey MakeUserDefinedPropertyValueKey(const FString& UserDefinedAttributeName,bool RequiresDelegate);
+	
+	static INTERCHANGECORE_API UE::Interchange::FAttributeKey MakeUserDefinedPropertyPayloadKey(const FString& UserDefinedAttributeName, bool RequiresDelegate);
+
+private:
+	static INTERCHANGECORE_API bool HasAttribute(const UInterchangeBaseNode* InterchangeSourceNode, const FString& InUserDefinedAttributeName, bool GeneratePayloadKey, bool& OutRequiresDelegate);
+	static INTERCHANGECORE_API UE::Interchange::FAttributeKey MakeUserDefinedPropertyKey(const FString& UserDefinedAttributeName,bool RequiresDelegate, bool GeneratePayloadKey = false);
+	
 private:
 	static INTERCHANGECORE_API const FString UserDefinedAttributeBaseKey;
 	static INTERCHANGECORE_API const FString UserDefinedAttributeValuePostKey;
 	static INTERCHANGECORE_API const FString UserDefinedAttributePayLoadPostKey;
+	static INTERCHANGECORE_API const FString UserDefinedAttributeDelegateKey;
 };
 
 template<typename ValueType>
-bool UInterchangeUserDefinedAttributesAPI::CreateUserDefinedAttribute(UInterchangeBaseNode* InterchangeNode, const FString& UserDefinedAttributeName, const ValueType& Value, const TOptional<FString>& PayloadKey)
+bool UInterchangeUserDefinedAttributesAPI::CreateUserDefinedAttribute(UInterchangeBaseNode* InterchangeNode, const FString& UserDefinedAttributeName, const ValueType& Value, const TOptional<FString>& PayloadKey, bool RequiresDelegate /*= true*/)
 {
 	check(InterchangeNode);
-	//Create a unique Key for this user defined attribute
-	FString StorageBaseKey = UserDefinedAttributeBaseKey + UserDefinedAttributeName;
-	UE::Interchange::FAttributeKey UserDefinedValueKey = UE::Interchange::FAttributeKey(StorageBaseKey + UserDefinedAttributeValuePostKey);
+	
+	UE::Interchange::FAttributeKey UserDefinedValueKey = MakeUserDefinedPropertyKey(UserDefinedAttributeName, RequiresDelegate);
 	if (InterchangeNode->HasAttribute(UserDefinedValueKey))
 	{
 		return false;
 	}
-	UE::Interchange::FAttributeKey UserDefinedPayloadKey = UE::Interchange::FAttributeKey(StorageBaseKey + UserDefinedAttributePayLoadPostKey);
+	
+	UE::Interchange::FAttributeKey UserDefinedPayloadKey = MakeUserDefinedPropertyKey(UserDefinedAttributeName, RequiresDelegate, true);
 	if (InterchangeNode->HasAttribute(UserDefinedPayloadKey))
 	{
 		return false;
@@ -149,14 +163,13 @@ template<typename ValueType>
 bool UInterchangeUserDefinedAttributesAPI::GetUserDefinedAttribute(const UInterchangeBaseNode* InterchangeNode, const FString& UserDefinedAttributeName, ValueType& OutValue, TOptional<FString>& OutPayloadKey)
 {
 	check(InterchangeNode);
-	FString StorageBaseKey = UserDefinedAttributeBaseKey + UserDefinedAttributeName;
-	UE::Interchange::FAttributeKey UserDefinedValueKey = UE::Interchange::FAttributeKey(StorageBaseKey + UserDefinedAttributeValuePostKey);
-
-	if (!InterchangeNode->HasAttribute(UserDefinedValueKey))
+	bool RequiresDelegate;
+	if (!HasAttribute(InterchangeNode, UserDefinedAttributeName, false, RequiresDelegate))
 	{
 		return false;
 	}
 
+	UE::Interchange::FAttributeKey UserDefinedValueKey = MakeUserDefinedPropertyKey(UserDefinedAttributeName, RequiresDelegate);
 	if (!InterchangeNode->GetAttribute<ValueType>(UserDefinedValueKey.Key, OutValue))
 	{
 		return false;
@@ -164,7 +177,7 @@ bool UInterchangeUserDefinedAttributesAPI::GetUserDefinedAttribute(const UInterc
 
 	//Payload is optional
 	OutPayloadKey.Reset();
-	UE::Interchange::FAttributeKey UserDefinedPayloadKey = UE::Interchange::FAttributeKey(StorageBaseKey + UserDefinedAttributePayLoadPostKey);
+	UE::Interchange::FAttributeKey UserDefinedPayloadKey = MakeUserDefinedPropertyKey(UserDefinedAttributeName, RequiresDelegate, true);
 	if (InterchangeNode->HasAttribute(UserDefinedPayloadKey))
 	{
 		FString PayloadKeyValue;
@@ -177,6 +190,7 @@ bool UInterchangeUserDefinedAttributesAPI::GetUserDefinedAttribute(const UInterc
 		{
 			return false;
 		}
+
 	}
 	return true;
 }
