Interchange: Prevent asset to loose import data

#rnx

diff --git a/Engine/Source/Runtime/Interchange/Engine/Private/InterchangeAssetImportData.cpp b/Engine/Source/Runtime/Interchange/Engine/Private/InterchangeAssetImportData.cpp
--- a/Engine/Source/Runtime/Interchange/Engine/Private/InterchangeAssetImportData.cpp
+++ b/Engine/Source/Runtime/Interchange/Engine/Private/InterchangeAssetImportData.cpp
@@ -101,37 +101,6 @@ void UInterchangeAssetImportData::Serialize(FArchive& Ar)
 	Ar.UsingCustomVersion(FInterchangeCustomVersion::GUID);
 	int32 CustomVersion = Ar.CustomVer(FInterchangeCustomVersion::GUID);
 
-	if (Ar.IsSaving())
-	{
-		if (UInterchangeManager::IsInterchangeImportEnabled())
-		{
-			//serialize
-			if (TransientNodeContainer)
-			{
-				FLargeMemoryWriter NodeContainerAr;
-				TransientNodeContainer->SerializeNodeContainerData(NodeContainerAr);
-				CachedNodeContainer = TArray64<uint8>(NodeContainerAr.GetData(), NodeContainerAr.TotalSize());
-			}
-			else
-			{
-				CachedNodeContainer.Reset();
-			}
-
-			CachedPipelines.Reset(TransientPipelines.Num());
-			for (TObjectPtr<UObject>& PipelineObjectPtr : TransientPipelines)
-			{
-				UObject* PipelineObject = PipelineObjectPtr.Get();
-				if (PipelineObject)
-				{
-					FString PipelineJSON = SerializePipeline(PipelineObject);
-
-					FString PipelineFullName = PipelineObject->GetClass()->GetFullName();
-					CachedPipelines.Emplace(PipelineFullName, PipelineJSON);
-				}
-			}
-		}
-	}
-
 	if (CustomVersion >= FInterchangeCustomVersion::SerializedInterchangeObjectStoring)
 	{
 		Ar << CachedNodeContainer;
@@ -149,7 +118,19 @@ UInterchangeBaseNodeContainer* UInterchangeAssetImportData::GetNodeContainer() c
 
 void UInterchangeAssetImportData::SetNodeContainer(UInterchangeBaseNodeContainer* InNodeContainer) const
 {
-	TransientNodeContainer = InNodeContainer; 
+	TransientNodeContainer = InNodeContainer;
+
+	//Serialize cache
+	if (TransientNodeContainer)
+	{
+		FLargeMemoryWriter NodeContainerAr;
+		TransientNodeContainer->SerializeNodeContainerData(NodeContainerAr);
+		CachedNodeContainer = TArray64<uint8>(NodeContainerAr.GetData(), NodeContainerAr.TotalSize());
+	}
+	else
+	{
+		CachedNodeContainer.Reset();
+	}
 }
 
 
@@ -161,6 +142,20 @@ void UInterchangeAssetImportData::SetPipelines(const TArray<UObject*>& InPipelin
 	{
 		TransientPipelines.Add(Pipeline);
 	}
+
+	//Serialize cache
+	CachedPipelines.Reset(TransientPipelines.Num());
+	for (TObjectPtr<UObject>& PipelineObjectPtr : TransientPipelines)
+	{
+		UObject* PipelineObject = PipelineObjectPtr.Get();
+		if (PipelineObject)
+		{
+			FString PipelineJSON = SerializePipeline(PipelineObject);
+
+			FString PipelineFullName = PipelineObject->GetClass()->GetFullName();
+			CachedPipelines.Emplace(PipelineFullName, PipelineJSON);
+		}
+	}
 }
 
 TArray<UObject*> UInterchangeAssetImportData::GetPipelines() const
diff --git a/Engine/Source/Runtime/Interchange/Engine/Public/InterchangeAssetImportData.h b/Engine/Source/Runtime/Interchange/Engine/Public/InterchangeAssetImportData.h
--- a/Engine/Source/Runtime/Interchange/Engine/Public/InterchangeAssetImportData.h
+++ b/Engine/Source/Runtime/Interchange/Engine/Public/InterchangeAssetImportData.h
@@ -157,6 +157,6 @@ private:
 	void ProcessContainerCache() const;
 	void ProcessPipelinesCache() const;
 	void ProcessDeprecatedData() const;
-	TArray64<uint8> CachedNodeContainer;
-	TArray<TPair<FString, FString>> CachedPipelines; //Class, Data(serialized JSON) pair
+	mutable TArray64<uint8> CachedNodeContainer;
+	mutable TArray<TPair<FString, FString>> CachedPipelines; //Class, Data(serialized JSON) pair
 };
