Fixed flaw in PostLoad based on changes to Serialize

#rnx

diff --git a/Engine/Source/Runtime/Interchange/Engine/Private/InterchangeAssetImportData.cpp b/Engine/Source/Runtime/Interchange/Engine/Private/InterchangeAssetImportData.cpp
--- a/Engine/Source/Runtime/Interchange/Engine/Private/InterchangeAssetImportData.cpp
+++ b/Engine/Source/Runtime/Interchange/Engine/Private/InterchangeAssetImportData.cpp
@@ -41,16 +41,7 @@ void UInterchangeAssetImportData::PostLoad()
 
 	if (Pipelines_DEPRECATED.Num() > 0)
 	{
-		TransientPipelines.Empty();
-
-		for (TObjectPtr<UObject>& PipelineObject : Pipelines_DEPRECATED)
-		{
-			if (PipelineObject)
-			{
-				TransientPipelines.Add(PipelineObject.Get());
-			}
-		}
-
+		SetPipelines(Pipelines_DEPRECATED);
 		Pipelines_DEPRECATED.Empty();
 	}
 }
@@ -140,7 +131,10 @@ void UInterchangeAssetImportData::SetPipelines(const TArray<UObject*>& InPipelin
 
 	for (UObject* Pipeline : InPipelines)
 	{
-		TransientPipelines.Add(Pipeline);
+		if (ensure(Pipeline))
+		{
+			TransientPipelines.Add(Pipeline);
+		}
 	}
 
 	//Serialize cache
