Fixed undesired harmless ensure

#rnx

diff --git a/Engine/Source/Runtime/Interchange/Engine/Private/InterchangeAssetImportData.cpp b/Engine/Source/Runtime/Interchange/Engine/Private/InterchangeAssetImportData.cpp
--- a/Engine/Source/Runtime/Interchange/Engine/Private/InterchangeAssetImportData.cpp
+++ b/Engine/Source/Runtime/Interchange/Engine/Private/InterchangeAssetImportData.cpp
@@ -131,7 +131,7 @@ void UInterchangeAssetImportData::SetPipelines(const TArray<UObject*>& InPipelin
 
 	for (UObject* Pipeline : InPipelines)
 	{
-		if (ensure(Pipeline))
+		if (Pipeline)
 		{
 			TransientPipelines.Add(Pipeline);
 		}
