Another fix to flawed logic in D: UInterchangeAssetImportData class

#rnx

diff --git a/Engine/Source/Runtime/Interchange/Engine/Private/InterchangeImportCommon.cpp b/Engine/Source/Runtime/Interchange/Engine/Private/InterchangeImportCommon.cpp
--- a/Engine/Source/Runtime/Interchange/Engine/Private/InterchangeImportCommon.cpp
+++ b/Engine/Source/Runtime/Interchange/Engine/Private/InterchangeImportCommon.cpp
@@ -67,7 +67,7 @@ namespace UE::Interchange
 		void EndSetupAssetData(FFactoryCommon::FUpdateImportAssetDataParameters& Parameters, UInterchangeAssetImportData* AssetImportData)
 		{
 			TRACE_CPUPROFILER_EVENT_SCOPE("UE::Interchange::Private::ImportCommon::EndSetupAssetData")
-			AssetImportData->SetNodeContainer(NewObject<UInterchangeBaseNodeContainer>(AssetImportData));
+			UInterchangeBaseNodeContainer* FactoryNodeContainer = NewObject<UInterchangeBaseNodeContainer>(AssetImportData);
 			//We copy only the factory node dependencies, we use this only 
 			if (UInterchangeFactoryBaseNode* FactoryNode = Cast<UInterchangeFactoryBaseNode>(Parameters.NodeContainer->GetFactoryNode(Parameters.NodeUniqueID)))
 			{
@@ -75,9 +75,10 @@ namespace UE::Interchange
 				AssetImportData->NodeUniqueID = Parameters.NodeUniqueID;
 				TArray<FString> FactoryDependencies;
 				FactoryDependencies.Add(AssetImportData->NodeUniqueID);
-				RecursivelyDuplicateFactoryNodeDependencies(Parameters.NodeContainer, AssetImportData->GetNodeContainer(), FactoryDependencies);
+				RecursivelyDuplicateFactoryNodeDependencies(Parameters.NodeContainer, FactoryNodeContainer, FactoryDependencies);
 			}
-			
+			AssetImportData->SetNodeContainer(FactoryNodeContainer);
+
 			TArray<UObject*> NewPipelines;
 			for (const UObject* Pipeline : Parameters.Pipelines)
 			{
