[Insights] Fixed name of cpu scope to not include quotes.

#rnx

diff --git a/Engine/Source/Runtime/Interchange/Engine/Private/InterchangeManager.cpp b/Engine/Source/Runtime/Interchange/Engine/Private/InterchangeManager.cpp
--- a/Engine/Source/Runtime/Interchange/Engine/Private/InterchangeManager.cpp
+++ b/Engine/Source/Runtime/Interchange/Engine/Private/InterchangeManager.cpp
@@ -1096,7 +1096,7 @@ bool UInterchangeManager::CanReimport(const UObject* Object, TArray<FString>& Ou
 
 void UInterchangeManager::StartQueuedTasks(bool bCancelAllTasks /*= false*/)
 {
-	TRACE_CPUPROFILER_EVENT_SCOPE("UInterchangeManager::StartQueuedTasks")
+	TRACE_CPUPROFILER_EVENT_SCOPE_STR("UInterchangeManager::StartQueuedTasks")
 	ensure(IsInterchangeImportEnabled());
 	if (!ensure(IsInGameThread()))
 	{
