Quick fix for issue in UE::Interchange::MakeName on Linux: Root cause FString initialized with a non-null terminated array of TCHARs.
FInal fix will go to Main as it may be more pervasive.

#rnx

diff --git a/Engine/Source/Runtime/Interchange/Core/Private/InterchangeHelper.cpp b/Engine/Source/Runtime/Interchange/Core/Private/InterchangeHelper.cpp
--- a/Engine/Source/Runtime/Interchange/Core/Private/InterchangeHelper.cpp
+++ b/Engine/Source/Runtime/Interchange/Core/Private/InterchangeHelper.cpp
@@ -6,22 +6,13 @@ namespace UE::Interchange
 {
 	FString MakeName(const FString& InName, bool bIsJoint)
 	{
-		FString SpecialChars = bIsJoint ? UnsupportedJointNameCharacters : UnsupportedNameCharacters;
+		const TCHAR* InvalidChar = bIsJoint ? INVALID_OBJECTNAME_CHARACTERS TEXT("+ ") : INVALID_OBJECTNAME_CHARACTERS;
 		FString TmpName = InName;
 
-		// Remove namespaces
-		int32 LastNamespaceTokenIndex = INDEX_NONE;
-		if (TmpName.FindLastChar(TEXT(':'), LastNamespaceTokenIndex))
+		while (*InvalidChar)
 		{
-			const bool bAllowShrinking = true;
-			//+1 to remove the ':' character we found
-			TmpName.RightChopInline(LastNamespaceTokenIndex + 1, bAllowShrinking);
-		}
-
-		//Remove the special chars
-		for (size_t CharIndex = 0; CharIndex < SpecialChars.Len(); CharIndex++)
-		{
-			TmpName.ReplaceCharInline(SpecialChars[CharIndex], TEXT('_'), ESearchCase::CaseSensitive);
+			TmpName.ReplaceCharInline(*InvalidChar, TCHAR('_'), ESearchCase::CaseSensitive);
+			++InvalidChar;
 		}
 
 		return TmpName;
diff --git a/Engine/Source/Runtime/Interchange/Core/Public/InterchangeHelper.h b/Engine/Source/Runtime/Interchange/Core/Public/InterchangeHelper.h
--- a/Engine/Source/Runtime/Interchange/Core/Public/InterchangeHelper.h
+++ b/Engine/Source/Runtime/Interchange/Core/Public/InterchangeHelper.h
@@ -7,11 +7,9 @@
 
 namespace UE::Interchange
 {
-	const TCHAR UnsupportedNameCharacters[] = { TEXT('.'), TEXT(','), TEXT('/'), TEXT('`'), TEXT('%') };
-	const TCHAR UnsupportedJointNameCharacters[] = { TEXT('.'), TEXT(','), TEXT('/'), TEXT('`'), TEXT('%'), TEXT('+'), TEXT(' ') };
-
 	/**
-	* Replaces any unsupported characters with "_" character, and removes namespace indicator ":" character
-	*/
+	 * Replaces any unsupported characters with "_" character, and removes namespace indicator ":" character
+	 * #todo_interchange: Replace with SanitizeObjectName from InterchangeManager.cpp
+	 */
 	INTERCHANGECORE_API FString MakeName(const FString& InName, bool bIsJoint = false);
 };
\ No newline at end of file
