Fixed flaw in package creation during import

#rnx

diff --git a/Engine/Source/Runtime/Interchange/Engine/Private/InterchangeManager.cpp b/Engine/Source/Runtime/Interchange/Engine/Private/InterchangeManager.cpp
--- a/Engine/Source/Runtime/Interchange/Engine/Private/InterchangeManager.cpp
+++ b/Engine/Source/Runtime/Interchange/Engine/Private/InterchangeManager.cpp
@@ -1431,9 +1431,11 @@ UInterchangeManager::ImportInternal(const FString& ContentPath, const UInterchan
 		Pipeline->AdjustSettingsForContext(Context, TaskData.ReimportObject);
 	};
 
-	FDateTime CurrentDateTime(FDateTime::Now());
-	const FString TransientPackageBasePath = GetTransientPackage()->GetPathName();
-	const FString PackageName = FString::Printf(TEXT("InterchangePipelinePackage_%llu"), CurrentDateTime.GetTicks());
+	// Use counter to guarantee uniqueness of packages on each call to ImportInternal
+	static uint64 ImportCounter = 0;
+	static FString TransientPackageBasePath = GetTransientPackage()->GetPathName();
+
+	const FString PackageName = FString::Printf(TEXT("InterchangePipelinePackage-%llu"), ++ImportCounter);
 	AsyncHelper->PipelineInstancesPackageName = TransientPackageBasePath / PackageName;
 
 	UPackage* PipelineInstancesPackage = CreatePackage(*AsyncHelper->PipelineInstancesPackageName);
diff --git a/Engine/Source/Runtime/Interchange/Engine/Private/Tasks/InterchangeTaskParsing.cpp b/Engine/Source/Runtime/Interchange/Engine/Private/Tasks/InterchangeTaskParsing.cpp
--- a/Engine/Source/Runtime/Interchange/Engine/Private/Tasks/InterchangeTaskParsing.cpp
+++ b/Engine/Source/Runtime/Interchange/Engine/Private/Tasks/InterchangeTaskParsing.cpp
@@ -333,7 +333,7 @@ void UE::Interchange::FTaskParsing::DoTask(ENamedThreads::Type CurrentThread, co
 				continue;
 			}
 			UInterchangeResultError_Generic* ErrorResult = NewObject<UInterchangeResultError_Generic>(GetTransientPackage(), UInterchangeResultError_Generic::StaticClass());
-			ErrorResult->Text = NSLOCTEXT("InterchangeTaskParsing", "TranslatorNoAssetImported", "There was no asset to import in the provided source data (found no mesh and no animation).");
+			ErrorResult->Text = NSLOCTEXT("InterchangeTaskParsing", "TranslatorNoAssetImported", "There was no data to import in the provided source data.");
 			ErrorResult->SourceAssetName = Translator->GetSourceData()->GetFilename();
 			Translator->AddMessage(ErrorResult);
 		}
