Interchange: Fix a race condition when filling the allowed class array in the async helper

#rnx

diff --git a/Engine/Source/Runtime/Interchange/Engine/Private/InterchangeManager.cpp b/Engine/Source/Runtime/Interchange/Engine/Private/InterchangeManager.cpp
--- a/Engine/Source/Runtime/Interchange/Engine/Private/InterchangeManager.cpp
+++ b/Engine/Source/Runtime/Interchange/Engine/Private/InterchangeManager.cpp
@@ -30,6 +30,7 @@
 #include "Misc/AssertionMacros.h"
 #include "Misc/AsyncTaskNotification.h"
 #include "Misc/MessageDialog.h"
+#include "Misc/NamePermissionList.h"
 #include "Misc/ScopedSlowTask.h"
 #include "Misc/DateTime.h"
 #include "Nodes/InterchangeBaseNodeContainer.h"
@@ -209,6 +210,37 @@ void UE::Interchange::FImportAsyncHelper::AddReferencedObjects(FReferenceCollect
 	Collector.AddReferencedObjects(CreatedFactories);
 }
 
+bool UE::Interchange::FImportAsyncHelper::IsClassImportAllowed(UClass* Class)
+{
+#if WITH_EDITOR
+	//Lock the classes
+	FScopeLock Lock(&ClassPermissionLock);
+
+	if (AllowedClasses.Contains(Class))
+	{
+		return true;
+	}
+	else if (DeniedClasses.Contains(Class))
+	{
+		return false;
+	}
+
+	IAssetTools& AssetTools = FAssetToolsModule::GetModule().Get();
+	TSharedPtr<FPathPermissionList> AssetClassPermissionList = AssetTools.GetAssetClassPathPermissionList(EAssetClassAction::ImportAsset);
+	if (AssetClassPermissionList && AssetClassPermissionList->HasFiltering())
+	{
+		if (!AssetClassPermissionList->PassesFilter(Class->GetPathName()))
+		{
+			UE_LOG(LogInterchangeEngine, Display, TEXT("The creation of asset of class '%s' is not allowed in this project."), *Class->GetName());
+			DeniedClasses.Add(Class);
+			return false;
+		}
+	}
+	AllowedClasses.Add(Class);
+#endif //WITH_EDITOR
+	return true;
+}
+
 /*
 * // Created factories map, Key is factory node UID
 			mutable FCriticalSection CreatedFactoriesLock;
diff --git a/Engine/Source/Runtime/Interchange/Engine/Private/Tasks/InterchangeTaskImportObject.cpp b/Engine/Source/Runtime/Interchange/Engine/Private/Tasks/InterchangeTaskImportObject.cpp
--- a/Engine/Source/Runtime/Interchange/Engine/Private/Tasks/InterchangeTaskImportObject.cpp
+++ b/Engine/Source/Runtime/Interchange/Engine/Private/Tasks/InterchangeTaskImportObject.cpp
@@ -1,9 +1,6 @@
 // Copyright Epic Games, Inc. All Rights Reserved.
 #include "InterchangeTaskImportObject.h"
 
-#if WITH_EDITOR
-#include "AssetToolsModule.h"
-#endif //WITH_EDITOR
 #include "AssetCompilingManager.h"
 #include "Async/TaskGraphInterfaces.h"
 #include "CoreMinimal.h"
@@ -15,7 +12,6 @@
 #include "InterchangeResult.h"
 #include "InterchangeSourceData.h"
 #include "InterchangeTranslatorBase.h"
-#include "Misc/NamePermissionList.h"
 #include "Misc/Paths.h"
 #include "PackageUtils/PackageUtils.h"
 #include "Stats/Stats.h"
@@ -122,43 +118,11 @@ namespace UE::Interchange::Private
 
 	bool CanImportClass(UE::Interchange::FImportAsyncHelper& AsyncHelper, UInterchangeFactoryBaseNode& FactoryNode, int32 SourceIndex)
 	{
-	#if WITH_EDITOR
 		if (UClass* Class = FactoryNode.GetObjectClass())
 		{
-			if (AsyncHelper.AllowedClasses.Contains(Class))
-			{
-				return true;
-			}
-			else if (AsyncHelper.DeniedClasses.Contains(Class))
-			{
-				return false;
-			}
-			else
-			{
-				IAssetTools& AssetTools = FAssetToolsModule::GetModule().Get();
-				TSharedPtr<FPathPermissionList> AssetClassPermissionList = AssetTools.GetAssetClassPathPermissionList(EAssetClassAction::ImportAsset);
-				if (AssetClassPermissionList && AssetClassPermissionList->HasFiltering())
-				{
-					if (!AssetClassPermissionList->PassesFilter(Class->GetPathName()))
-					{
-						UE_LOG(LogInterchangeEngine, Display, TEXT("The creation of asset of class '%s' is not allowed in this project."), *Class->GetName());
-						AsyncHelper.DeniedClasses.Add(Class);
-						return false;
-					}
-				}
-
-				AsyncHelper.AllowedClasses.Add(Class);
-
-				return true;
-			}
+			return AsyncHelper.IsClassImportAllowed(Class);
 		}
-		else
-		{
-			return false;
-		}
-	#endif //WITH_EDITOR
-
-		return true;
+		return false;
 	}
 
 	UInterchangeFactoryBase::FImportAssetResult InternalImportObjectStartup(TSharedPtr<FImportAsyncHelper, ESPMode::ThreadSafe> AsyncHelper
diff --git a/Engine/Source/Runtime/Interchange/Engine/Public/InterchangeManager.h b/Engine/Source/Runtime/Interchange/Engine/Public/InterchangeManager.h
--- a/Engine/Source/Runtime/Interchange/Engine/Public/InterchangeManager.h
+++ b/Engine/Source/Runtime/Interchange/Engine/Public/InterchangeManager.h
@@ -220,18 +220,15 @@ namespace UE
 			// Package where the Pipeline Instances are stored during an import.
 			FString PipelineInstancesPackageName;
 
+			//Return true if we can import this class, false otherwise
+			bool IsClassImportAllowed(UClass* Class);
+
 			UPackage* GetCreatedPackage(const FString& PackageName) const;
 			void AddCreatedPackage(const FString& PackageName, UPackage* Package);
 
 			UInterchangeFactoryBase* GetCreatedFactory(const FString& FactoryNodeUniqueId) const;
 			void AddCreatedFactory(const FString& FactoryNodeUniqueId, UInterchangeFactoryBase* Factory);
 
-			// Set of classes which creation has been denied
-			TSet<UClass*> DeniedClasses;
-
-			// Set of classes which creation is allowed
-			TSet<UClass*> AllowedClasses;
-
 			struct FImportedObjectInfo
 			{
 				UObject* ImportedObject = nullptr; // The object that was imported
@@ -276,6 +273,12 @@ namespace UE
 			void CleanUp();
 
 		private:
+			FCriticalSection ClassPermissionLock;
+			// Set of classes which creation has been denied
+			TSet<UClass*> DeniedClasses;
+			// Set of classes which creation is allowed
+			TSet<UClass*> AllowedClasses;
+
 			//Create package map, Key is package name. We cannot create package asynchronously so we have to create a game thread task to do this
 			mutable FCriticalSection CreatedPackagesLock;
 			TMap<FString, UPackage*> CreatedPackages;
