Modified Interchange engine in order to add Interchange Import to Dataprep

diff --git a/Engine/Source/Runtime/Interchange/Engine/Private/InterchangeManager.cpp b/Engine/Source/Runtime/Interchange/Engine/Private/InterchangeManager.cpp
--- a/Engine/Source/Runtime/Interchange/Engine/Private/InterchangeManager.cpp
+++ b/Engine/Source/Runtime/Interchange/Engine/Private/InterchangeManager.cpp
@@ -1409,6 +1409,7 @@ UInterchangeManager::ImportInternal(const FString& ContentPath, const UInterchan
 	TaskData.bFollowRedirectors = ImportAssetParameters.bFollowRedirectors;
 	TaskData.ImportType = ImportType;
 	TaskData.ReimportObject = ImportAssetParameters.ReimportAsset;
+	TaskData.ImportLevel = ImportAssetParameters.ImportLevel;
 
 	TSharedRef<UE::Interchange::FImportAsyncHelper, ESPMode::ThreadSafe> AsyncHelper = CreateAsyncHelper(TaskData, ImportAssetParameters);
 	AsyncHelper->UniqueId = UniqueId;
diff --git a/Engine/Source/Runtime/Interchange/Engine/Private/Tasks/InterchangeTaskCreateSceneObjects.cpp b/Engine/Source/Runtime/Interchange/Engine/Private/Tasks/InterchangeTaskCreateSceneObjects.cpp
--- a/Engine/Source/Runtime/Interchange/Engine/Private/Tasks/InterchangeTaskCreateSceneObjects.cpp
+++ b/Engine/Source/Runtime/Interchange/Engine/Private/Tasks/InterchangeTaskCreateSceneObjects.cpp
@@ -57,10 +57,11 @@ void UE::Interchange::FTaskCreateSceneObjects::DoTask(ENamedThreads::Type Curren
 	}
 
 	UObject* ReimportObject = AsyncHelper->TaskData.ReimportObject;
-	ULevel* CurrentLevel = GWorld->GetCurrentLevel();
-	const FString WorldPath = GWorld->GetOutermost()->GetPathName();
-	const FString WorldName = GWorld->GetName();
-	const FString NodePrefix = CurrentLevel->GetName() + TEXT(".");
+	ULevel* ImportLevel = AsyncHelper->TaskData.ImportLevel ? AsyncHelper->TaskData.ImportLevel : GWorld->GetCurrentLevel();
+	UWorld* ImportWorld = ImportLevel->GetWorld();
+	const FString WorldPath = ImportWorld->GetOutermost()->GetPathName();
+	const FString WorldName = ImportWorld->GetName();
+	const FString NodePrefix = ImportLevel->GetName() + TEXT(".");
 
 	for (UInterchangeFactoryBaseNode* FactoryNode : FactoryNodes)
 	{
@@ -79,7 +80,7 @@ void UE::Interchange::FTaskCreateSceneObjects::DoTask(ENamedThreads::Type Curren
 		UInterchangeFactoryBase::FImportSceneObjectsParams CreateSceneObjectsParams;
 		CreateSceneObjectsParams.ObjectName = NodeDisplayName;
 		CreateSceneObjectsParams.FactoryNode = FactoryNode;
-		CreateSceneObjectsParams.Level = CurrentLevel;
+		CreateSceneObjectsParams.Level = ImportLevel;
 		CreateSceneObjectsParams.ReimportObject = FFactoryCommon::GetObjectToReimport(ReimportObject, *FactoryNode, WorldPath, WorldName, NodePrefix + NodeDisplayName);
 		CreateSceneObjectsParams.ReimportFactoryNode = FFactoryCommon::GetFactoryNode(ReimportObject, WorldPath, WorldName, NodePrefix + NodeDisplayName);
 
diff --git a/Engine/Source/Runtime/Interchange/Engine/Public/InterchangeManager.h b/Engine/Source/Runtime/Interchange/Engine/Public/InterchangeManager.h
--- a/Engine/Source/Runtime/Interchange/Engine/Public/InterchangeManager.h
+++ b/Engine/Source/Runtime/Interchange/Engine/Public/InterchangeManager.h
@@ -33,6 +33,7 @@ class UInterchangeFactoryBase;
 class UInterchangeFactoryBaseNode;
 class UInterchangePipelineBase;
 class UInterchangePythonPipelineBase;
+class ULevel;
 
 /** Some utilities delegates for the automation of interchange */
 DECLARE_DYNAMIC_DELEGATE_OneParam(FOnObjectImportDoneDynamic, UObject*, Object);
@@ -93,6 +94,9 @@ namespace UE
 
 			//True if we are reimporting assets or scene
 			UObject* ReimportObject = nullptr;
+
+			//Level to import into in case of a scene import
+			ULevel* ImportLevel = nullptr;
 		};
 
 		class FImportResult : protected FGCObject
@@ -353,6 +357,10 @@ struct FImportAssetParameters
 	UPROPERTY(EditAnywhere, BlueprintReadWrite, Category = "Interchange|ImportAsset", meta = (AllowedClasses = "/Script/InterchangeCore.InterchangePipelineBase, /Script/InterchangeEngine.InterchangeBlueprintPipelineBase, /Script/InterchangeEngine.InterchangePythonPipelineAsset"))
 	TArray<FSoftObjectPath> OverridePipelines;
 
+	//Level to import into in case of a scene import
+	UPROPERTY(EditAnywhere, BlueprintReadWrite, Category = "Interchange|ImportAsset")
+	TObjectPtr<ULevel> ImportLevel = nullptr;
+
 	/* Delegates used track the imported objects */
 
 	// This is called each time an asset is imported or reimported from the import call
