Fix Interchange - �Import Done� log record as a failure on import

diff --git a/Engine/Source/Runtime/Interchange/Engine/Private/InterchangeManager.cpp b/Engine/Source/Runtime/Interchange/Engine/Private/InterchangeManager.cpp
--- a/Engine/Source/Runtime/Interchange/Engine/Private/InterchangeManager.cpp
+++ b/Engine/Source/Runtime/Interchange/Engine/Private/InterchangeManager.cpp
@@ -1885,7 +1885,14 @@ void UInterchangeManager::ReleaseAsyncHelper(TWeakPtr<UE::Interchange::FImportAs
 			}
 			else
 			{
-				TitleText = NSLOCTEXT("Interchange", "Asynchronous_import_end", "Import Done");
+				if(bSucceeded)
+				{
+					TitleText = NSLOCTEXT("Interchange", "Asynchronous_import_end", "Import Done");
+				}
+				else
+				{
+					TitleText = NSLOCTEXT("Interchange", "Asynchronous_import_failed", "Import Failed");
+				}
 			}
 
 			Notification->SetComplete(TitleText, FText::GetEmpty(), bSucceeded);
