Make sure all options from the import system are supported properly.

#rnx

diff --git a/Engine/Plugins/Interchange/Runtime/Source/Pipelines/Private/InterchangeGenericAssetsPipeline.cpp b/Engine/Plugins/Interchange/Runtime/Source/Pipelines/Private/InterchangeGenericAssetsPipeline.cpp
--- a/Engine/Plugins/Interchange/Runtime/Source/Pipelines/Private/InterchangeGenericAssetsPipeline.cpp
+++ b/Engine/Plugins/Interchange/Runtime/Source/Pipelines/Private/InterchangeGenericAssetsPipeline.cpp
@@ -553,7 +553,12 @@ void UInterchangeGenericAssetsPipeline::ImplementUseSourceNameForAssetOption()
 	//If we import only one mesh, we want to rename the mesh using the file name.
 	const int32 MeshesImportedNodeCount = SkeletalMeshNodeUids.Num() + StaticMeshNodeUids.Num();
 
-	FString OverrideAssetName = IsStandAlonePipeline() ? AssetName : FString();
+	FString OverrideAssetName = IsStandAlonePipeline() ? DestinationName : FString();
+	if(OverrideAssetName.IsEmpty() && IsStandAlonePipeline())
+	{
+		OverrideAssetName = AssetName;
+	}
+
 	//SkeletalMesh it must always be run even if there is no rename option, skeleton and physics asset will be rename properly
 	MeshPipeline->ImplementUseSourceNameForAssetOptionSkeletalMesh(MeshesImportedNodeCount, bUseSourceNameForAsset, OverrideAssetName);
 
diff --git a/Engine/Plugins/Interchange/Runtime/Source/Pipelines/Private/InterchangeGenericMaterialPipeline.cpp b/Engine/Plugins/Interchange/Runtime/Source/Pipelines/Private/InterchangeGenericMaterialPipeline.cpp
--- a/Engine/Plugins/Interchange/Runtime/Source/Pipelines/Private/InterchangeGenericMaterialPipeline.cpp
+++ b/Engine/Plugins/Interchange/Runtime/Source/Pipelines/Private/InterchangeGenericMaterialPipeline.cpp
@@ -742,7 +742,14 @@ void UInterchangeGenericMaterialPipeline::ExecutePipeline(UInterchangeBaseNodeCo
 		}
 	}
 
-	if (IsStandAlonePipeline() && !AssetName.IsEmpty())
+	//If we have a valid override name
+	FString OverrideAssetName = IsStandAlonePipeline() ? DestinationName : FString();
+	if (OverrideAssetName.IsEmpty() && IsStandAlonePipeline())
+	{
+		OverrideAssetName = AssetName;
+	}
+
+	if (IsStandAlonePipeline() && !OverrideAssetName.IsEmpty())
 	{
 		TArray<UInterchangeBaseMaterialFactoryNode*> BaseMaterialNodes;
 		BaseNodeContainer->IterateNodesOfType<UInterchangeBaseMaterialFactoryNode>([&BaseMaterialNodes](const FString& NodeUid, UInterchangeBaseMaterialFactoryNode* MaterialNode)
@@ -752,8 +759,8 @@ void UInterchangeGenericMaterialPipeline::ExecutePipeline(UInterchangeBaseNodeCo
 
 		if(BaseMaterialNodes.Num() == 1)
 		{
-			BaseMaterialNodes[0]->SetAssetName(AssetName);
-			BaseMaterialNodes[0]->SetDisplayLabel(AssetName);
+			BaseMaterialNodes[0]->SetAssetName(OverrideAssetName);
+			BaseMaterialNodes[0]->SetDisplayLabel(OverrideAssetName);
 		}
 	}
 }
diff --git a/Engine/Plugins/Interchange/Runtime/Source/Pipelines/Private/InterchangeGenericTexturePipeline.cpp b/Engine/Plugins/Interchange/Runtime/Source/Pipelines/Private/InterchangeGenericTexturePipeline.cpp
--- a/Engine/Plugins/Interchange/Runtime/Source/Pipelines/Private/InterchangeGenericTexturePipeline.cpp
+++ b/Engine/Plugins/Interchange/Runtime/Source/Pipelines/Private/InterchangeGenericTexturePipeline.cpp
@@ -173,11 +173,17 @@ void UInterchangeGenericTexturePipeline::ExecutePipeline(UInterchangeBaseNodeCon
 			TextureFactoryNode = HandleCreationOfTextureFactoryNode(TextureNode);
 		}
 		//If we have a valid override name
-		const bool bOverrideAssetName = TextureNodes.Num() == 1 && IsStandAlonePipeline() && !AssetName.IsEmpty();
+		FString OverrideAssetName = IsStandAlonePipeline() ? DestinationName : FString();
+		if (OverrideAssetName.IsEmpty() && IsStandAlonePipeline())
+		{
+			OverrideAssetName = AssetName;
+		}
+
+		const bool bOverrideAssetName = TextureNodes.Num() == 1 && IsStandAlonePipeline() && !OverrideAssetName.IsEmpty();
 		if (TextureFactoryNode && bOverrideAssetName)
 		{
-			TextureFactoryNode->SetAssetName(AssetName);
-			TextureFactoryNode->SetDisplayLabel(AssetName);
+			TextureFactoryNode->SetAssetName(OverrideAssetName);
+			TextureFactoryNode->SetDisplayLabel(OverrideAssetName);
 		}
 	}
 }
diff --git a/Engine/Source/Developer/AssetTools/Private/AssetTools.cpp b/Engine/Source/Developer/AssetTools/Private/AssetTools.cpp
--- a/Engine/Source/Developer/AssetTools/Private/AssetTools.cpp
+++ b/Engine/Source/Developer/AssetTools/Private/AssetTools.cpp
@@ -3315,12 +3315,12 @@ TArray<UObject*> UAssetToolsImpl::ImportAssetsInternal(const TArray<FString>& Fi
 				ImportAssetParameters.bIsAutomated = bAutomatedImport;
 				ImportAssetParameters.bFollowRedirectors = UE::AssetTools::Private::bFollowRedirectorsWhenImporting;
 				ImportAssetParameters.ReimportAsset = nullptr;
+				ImportAssetParameters.bReplaceExisting = bForceOverrideExisting;
 				if (Params.AssetImportTask && !Params.AssetImportTask->DestinationName.IsEmpty())
 				{
-					UE_LOG(LogAssetTools, Warning, TEXT("Interchange do not use the DestinationName parameter. Users must specify the destination name in a pipeline through the Options parameter."));
+					ImportAssetParameters.DestinationName = Params.AssetImportTask->DestinationName;
 				}
 
-
 				TFunction<void(UE::Interchange::FImportResult&)> AppendImportResult =
 					// Note: ImportStatus captured by value so that the lambda keeps the shared ptr alive
 					[ImportStatus, ImportTask = Params.AssetImportTask](UE::Interchange::FImportResult& Result)
diff --git a/Engine/Source/Editor/UnrealEd/Private/SkinWeightsUtilities.cpp b/Engine/Source/Editor/UnrealEd/Private/SkinWeightsUtilities.cpp
--- a/Engine/Source/Editor/UnrealEd/Private/SkinWeightsUtilities.cpp
+++ b/Engine/Source/Editor/UnrealEd/Private/SkinWeightsUtilities.cpp
@@ -120,6 +120,7 @@ bool FSkinWeightsUtilities::ImportAlternateSkinWeight(USkeletalMesh* SkeletalMes
 			}
 		}
 
+		ImportAssetParameters.DestinationName = FGuid::NewGuid().ToString(EGuidFormats::Digits);
 		//TODO create a pipeline that set all the proper skeletalmesh options (look at the legacy system setup)
 		UE::Interchange::FAssetImportResultRef AssetImportResult = InterchangeManager.ImportAssetAsync(ImportAssetPath, ScopedSourceData.GetSourceData(), ImportAssetParameters);
 		AssetImportResult->WaitUntilDone(); //TODO, do not stall the main thread here, WaitUntilDone will tick taskgraph so the job can complete even if we wait on the game thread.
diff --git a/Engine/Source/Runtime/Interchange/Core/Public/InterchangePipelineBase.h b/Engine/Source/Runtime/Interchange/Core/Public/InterchangePipelineBase.h
--- a/Engine/Source/Runtime/Interchange/Core/Public/InterchangePipelineBase.h
+++ b/Engine/Source/Runtime/Interchange/Core/Public/InterchangePipelineBase.h
@@ -340,6 +340,12 @@ public:
 	 */
 	INTERCHANGECORE_API virtual bool SupportReimport() const { return true; }
 
+	/*
+	 * The Unreal import system have an option to force a name if we import only one main asset (one texture, one mesh or one animation).
+	 * The generic asset pipeline is using this information to behave like expected.
+	 */
+	FString DestinationName;
+
 protected:
 
 	UE_DEPRECATED(5.2, "This function is replace by ExecutePipeline.")
diff --git a/Engine/Source/Runtime/Interchange/Engine/Private/InterchangeManager.cpp b/Engine/Source/Runtime/Interchange/Engine/Private/InterchangeManager.cpp
--- a/Engine/Source/Runtime/Interchange/Engine/Private/InterchangeManager.cpp
+++ b/Engine/Source/Runtime/Interchange/Engine/Private/InterchangeManager.cpp
@@ -1341,6 +1341,11 @@ UInterchangeManager::ImportInternal(const FString& ContentPath, const UInterchan
 		}
 		const bool bIsPipelineOverride = ImportAssetParameters.OverridePipelines.Num() > 0;
 		Attribs.Add(FAnalyticsEventAttribute(TEXT("Parameters.IsPipelineOverrided"), bIsPipelineOverride));
+		Attribs.Add(FAnalyticsEventAttribute(TEXT("Parameters.bReplaceExisting"), ImportAssetParameters.bReplaceExisting));
+		if(!ImportAssetParameters.DestinationName.IsEmpty())
+		{
+			Attribs.Add(FAnalyticsEventAttribute(TEXT("Parameters.DestinationName"), ImportAssetParameters.DestinationName));
+		}
 	}
 
 	if (!ensure(IsInGameThread()))
@@ -1410,6 +1415,8 @@ UInterchangeManager::ImportInternal(const FString& ContentPath, const UInterchan
 	TaskData.ImportType = ImportType;
 	TaskData.ReimportObject = ImportAssetParameters.ReimportAsset;
 	TaskData.ImportLevel = ImportAssetParameters.ImportLevel;
+	TaskData.DestinationName = ImportAssetParameters.DestinationName;
+	TaskData.bReplaceExisting = ImportAssetParameters.bReplaceExisting;
 
 	TSharedRef<UE::Interchange::FImportAsyncHelper, ESPMode::ThreadSafe> AsyncHelper = CreateAsyncHelper(TaskData, ImportAssetParameters);
 	AsyncHelper->UniqueId = UniqueId;
@@ -1462,6 +1469,7 @@ UInterchangeManager::ImportInternal(const FString& ContentPath, const UInterchan
 		}
 
 		Pipeline->AdjustSettingsForContext(Context, TaskData.ReimportObject);
+		Pipeline->DestinationName = TaskData.DestinationName;
 	};
 
 	// Use counter to guarantee uniqueness of packages on each call to ImportInternal
diff --git a/Engine/Source/Runtime/Interchange/Engine/Private/Tasks/InterchangeTaskImportObject.cpp b/Engine/Source/Runtime/Interchange/Engine/Private/Tasks/InterchangeTaskImportObject.cpp
--- a/Engine/Source/Runtime/Interchange/Engine/Private/Tasks/InterchangeTaskImportObject.cpp
+++ b/Engine/Source/Runtime/Interchange/Engine/Private/Tasks/InterchangeTaskImportObject.cpp
@@ -380,6 +380,18 @@ void UE::Interchange::FTaskImportObject_GameThread::DoTask(ENamedThreads::Type C
 			return;
 		}
 
+		if (!bPackageWasCreated && !AsyncHelper->TaskData.bReplaceExisting)
+		{
+			//Do not replace existing asset, the option tell us to not override it. Skip this asset witha display message
+			UInterchangeResultWarning_Generic* Message = Factory->AddMessage<UInterchangeResultWarning_Generic>();
+			Message->SourceAssetName = AsyncHelper->SourceDatas[SourceIndex]->GetFilename();
+			Message->DestinationAssetName = AssetName;
+			Message->AssetType = FactoryNode->GetObjectClass();
+			Message->Text = FText::Format(NSLOCTEXT("FTaskImportObject_GameThread", "CouldntReplaceExistingAsset", "The option bReplaceExisting is false so we are not overriding the asset named '{0}'.")
+				, FText::FromString(AssetName));
+			return;
+		}
+
 		if (!bPackageWasCreated && AsyncHelper->TaskData.bFollowRedirectors)
 		{
 			if (UObjectRedirector* Redirector = FindObject<UObjectRedirector>(Pkg, *AssetName))
diff --git a/Engine/Source/Runtime/Interchange/Engine/Public/InterchangeManager.h b/Engine/Source/Runtime/Interchange/Engine/Public/InterchangeManager.h
--- a/Engine/Source/Runtime/Interchange/Engine/Public/InterchangeManager.h
+++ b/Engine/Source/Runtime/Interchange/Engine/Public/InterchangeManager.h
@@ -97,6 +97,12 @@ namespace UE
 
 			//Level to import into in case of a scene import
 			ULevel* ImportLevel = nullptr;
+
+			/** Optional custom name to import as */
+			FString DestinationName;
+
+			/** Overwrite existing assets */
+			bool bReplaceExisting = true;
 		};
 
 		class FImportResult : protected FGCObject
@@ -361,6 +367,14 @@ struct FImportAssetParameters
 	UPROPERTY(EditAnywhere, BlueprintReadWrite, Category = "Interchange|ImportAsset")
 	TObjectPtr<ULevel> ImportLevel = nullptr;
 
+	/** Optional custom name to import as */
+	UPROPERTY(EditAnywhere, BlueprintReadWrite, Category = "Interchange|ImportAsset")
+	FString DestinationName;
+
+	/** Overwrite existing assets */
+	UPROPERTY(EditAnywhere, BlueprintReadWrite, Category = "Interchange|ImportAsset")
+	bool bReplaceExisting = true;
+
 	/* Delegates used track the imported objects */
 
 	// This is called each time an asset is imported or reimported from the import call
