Fix non unity build
#rnx

diff --git a/Engine/Source/Runtime/Interchange/Engine/Public/InterchangePipelineConfigurationBase.h b/Engine/Source/Runtime/Interchange/Engine/Public/InterchangePipelineConfigurationBase.h
--- a/Engine/Source/Runtime/Interchange/Engine/Public/InterchangePipelineConfigurationBase.h
+++ b/Engine/Source/Runtime/Interchange/Engine/Public/InterchangePipelineConfigurationBase.h
@@ -6,6 +6,7 @@
 #include "UObject/Object.h"
 #include "UObject/ObjectMacros.h"
 #include "InterchangeSourceData.h"
+#include "Nodes/InterchangeBaseNodeContainer.h"
 
 #include "InterchangePipelineConfigurationBase.generated.h"
 
