Make sure "import all" option from the import dialog reused the same options in case we import different file format (using different translator).

#rnx

diff --git a/Engine/Source/Runtime/Interchange/Engine/Private/InterchangeManager.cpp b/Engine/Source/Runtime/Interchange/Engine/Private/InterchangeManager.cpp
--- a/Engine/Source/Runtime/Interchange/Engine/Private/InterchangeManager.cpp
+++ b/Engine/Source/Runtime/Interchange/Engine/Private/InterchangeManager.cpp
@@ -1428,10 +1428,12 @@ UInterchangeManager::ImportInternal(const FString& ContentPath, const UInterchan
 	//Array of source data to build one graph per source
 	AsyncHelper->SourceDatas.Add(DuplicateSourceData);
 	constexpr int32 SourceIndex = 0;
+	UInterchangeTranslatorBase* AsyncTranslator = nullptr;
 	//Get all the translators for the source datas
 	for (int32 SourceDataIndex = 0; SourceDataIndex < AsyncHelper->SourceDatas.Num(); ++SourceDataIndex)
 	{
-		ensure(AsyncHelper->Translators.Add(GetTranslatorForSourceData(AsyncHelper->SourceDatas[SourceDataIndex])) == SourceDataIndex);
+		AsyncTranslator = GetTranslatorForSourceData(AsyncHelper->SourceDatas[SourceDataIndex]);
+		ensure(AsyncHelper->Translators.Add(AsyncTranslator) == SourceDataIndex);
 	}
 
 	//Create the node graphs for each source data (StrongObjectPtr has to be created on the main thread)
@@ -1485,14 +1487,27 @@ UInterchangeManager::ImportInternal(const FString& ContentPath, const UInterchan
 	PipelineInstancesPackage->ClearFlags(RF_Public | RF_Standalone);
 	PipelineInstancesPackage->SetPackageFlags(PKG_NewlyCreated);
 
-	if ( ImportAssetParameters.OverridePipelines.Num() == 0 )
+	const bool bSkipImportDialog = AsyncTranslator ? ImportAllWithSamePipelines.Contains(AsyncTranslator->GetClass()) : false;
+	if (bSkipImportDialog)
 	{
-		const bool bIsUnattended = FApp::IsUnattended() || GIsAutomationTesting || ImportAssetParameters.bIsAutomated;
-
+		TArray<UInterchangePipelineBase*> LastImportPipelines = ImportAllWithSamePipelines.FindChecked(AsyncTranslator->GetClass());
+		for (const UInterchangePipelineBase* LastImportPipeline : LastImportPipelines)
+		{
+			if(UInterchangePipelineBase* Pipeline = DuplicateObject<UInterchangePipelineBase>(LastImportPipeline, GetTransientPackage()))
+			{
+				AsyncHelper->Pipelines.Add(Pipeline);
+				AsyncHelper->OriginalPipelines.Add(Pipeline);
+				UE::Interchange::Private::FillPipelineAnalyticData(Pipeline, UniqueId, FString());
+			}
+		}
+	}
+	else if ( ImportAssetParameters.OverridePipelines.Num() == 0 )
+	{
+		
+		const bool bIsUnattended = FApp::IsUnattended() || GIsAutomationTesting || ImportAssetParameters.bIsAutomated || bSkipImportDialog;
 #if WITH_EDITORONLY_DATA
 		const bool bShowPipelineStacksConfigurationDialog = !bIsUnattended
 															&& FInterchangeProjectSettingsUtils::ShouldShowPipelineStacksConfigurationDialog(bImportScene, *SourceData)
-															&& !bImportAllWithDefault
 															&& !bImportCanceled
 															&& !IsRunningCommandlet();
 #else
@@ -1602,6 +1617,19 @@ UInterchangeManager::ImportInternal(const FString& ContentPath, const UInterchan
 			}
 		}
 
+		auto SetImportAllWithSamePipelines = [this, &AsyncTranslator](TArray<UInterchangePipelineBase*>& ToDuplicatePipelines)
+		{
+			TArray<UInterchangePipelineBase*>& PipelineList = ImportAllWithSamePipelines.FindOrAdd(AsyncTranslator->GetClass());
+			for (const UInterchangePipelineBase* Pipeline : ToDuplicatePipelines)
+			{
+				if (UInterchangePipelineBase* DupPipeline = DuplicateObject<UInterchangePipelineBase>(Pipeline, GetTransientPackage()))
+				{
+					DupPipeline->SetInternalFlags(EInternalObjectFlags::Async);
+					PipelineList.Add(DupPipeline);
+				}
+			}
+		};
+
 		if (bIsReimport)
 		{
 			if (RegisteredPipelineConfiguration && bShowPipelineStacksConfigurationDialog && !bIsUnattended)
@@ -1616,7 +1644,7 @@ UInterchangeManager::ImportInternal(const FString& ContentPath, const UInterchan
 				}
 				if (DialogResult == EInterchangePipelineConfigurationDialogResult::ImportAll)
 				{
-					bImportAllWithDefault = true;
+					SetImportAllWithSamePipelines(OutPipelines);
 				}
 			}
 			else
@@ -1648,7 +1676,7 @@ UInterchangeManager::ImportInternal(const FString& ContentPath, const UInterchan
 				}
 				if (DialogResult == EInterchangePipelineConfigurationDialogResult::ImportAll)
 				{
-					bImportAllWithDefault = true;
+					SetImportAllWithSamePipelines(OutPipelines);
 				}
 			}
 			else
@@ -1662,14 +1690,6 @@ UInterchangeManager::ImportInternal(const FString& ContentPath, const UInterchan
 				{
 					//When we do not show the UI we use the original stack
 					OutPipelines = StackInfoPtr->Pipelines;
-					if (bImportAllWithDefault)
-					{
-						//When the user want to use the same settings for all source data we just need to load the settings save by the initial dialog
-						for (UInterchangePipelineBase* Pipeline : OutPipelines)
-						{
-							Pipeline->LoadSettings(DefaultStackName);
-						}
-					}
 				}
 				else if (PipelineStacks.Num() > 0)
 				{
@@ -1930,7 +1950,16 @@ void UInterchangeManager::ReleaseAsyncHelper(TWeakPtr<UE::Interchange::FImportAs
 			Notification = nullptr; //This should delete the notification
 		}
 
-		bImportAllWithDefault = false;
+		//Release import all pipelines so they can be garbage collect
+		for (TPair<UClass*, TArray<UInterchangePipelineBase*>>& ImportAllWithSamePipelinesPair : ImportAllWithSamePipelines)
+		{
+			for (UInterchangePipelineBase* Pipeline : ImportAllWithSamePipelinesPair.Value)
+			{
+				Pipeline->ClearInternalFlags(EInternalObjectFlags::Async);
+			}
+		}
+		ImportAllWithSamePipelines.Empty();
+
 		bImportCanceled = false;
 	}
 	else if(Notification.IsValid())
diff --git a/Engine/Source/Runtime/Interchange/Engine/Public/InterchangeManager.h b/Engine/Source/Runtime/Interchange/Engine/Public/InterchangeManager.h
--- a/Engine/Source/Runtime/Interchange/Engine/Public/InterchangeManager.h
+++ b/Engine/Source/Runtime/Interchange/Engine/Public/InterchangeManager.h
@@ -694,8 +694,8 @@ private:
 	FThreadSafeBool bIsActive = false;
 
 	//If the user wants to use the same import pipeline stack for all the queue task
-	//This boolean is reset to false when the ImportTasks array is empty.
-	bool bImportAllWithDefault = false;
+	//This map is reset when the ImportTasks array is empty.
+	TMap<UClass*, TArray<UInterchangePipelineBase*>> ImportAllWithSamePipelines;
 
 	//Indicates that the import process was canceled by the user.
 	//This boolean is reset to false when the ImportTasks array is empty.
