Interchange: Fix crash when async texture build is not finish and the pipeline post import callback try to access the render data.

#rnx

diff --git a/Engine/Source/Runtime/Interchange/Engine/Private/InterchangeManager.cpp b/Engine/Source/Runtime/Interchange/Engine/Private/InterchangeManager.cpp
--- a/Engine/Source/Runtime/Interchange/Engine/Private/InterchangeManager.cpp
+++ b/Engine/Source/Runtime/Interchange/Engine/Private/InterchangeManager.cpp
@@ -515,6 +515,10 @@ FGraphEventArray UE::Interchange::FImportAsyncHelper::GetCompletionTaskGraphEven
 	TasksToComplete.Append(ImportObjectTasks);
 	TasksToComplete.Append(FinalizeImportObjectTasks);
 	TasksToComplete.Append(SceneTasks);
+	if (WaitAssetCompilationTask.GetReference())
+	{
+		TasksToComplete.Add(WaitAssetCompilationTask);
+	}
 	TasksToComplete.Append(PipelinePostImportTasks);
 
 	if (PreAsyncCompletionTask.GetReference())
diff --git a/Engine/Source/Runtime/Interchange/Engine/Private/Tasks/InterchangeTaskParsing.cpp b/Engine/Source/Runtime/Interchange/Engine/Private/Tasks/InterchangeTaskParsing.cpp
--- a/Engine/Source/Runtime/Interchange/Engine/Private/Tasks/InterchangeTaskParsing.cpp
+++ b/Engine/Source/Runtime/Interchange/Engine/Private/Tasks/InterchangeTaskParsing.cpp
@@ -406,27 +406,35 @@ void UE::Interchange::FTaskParsing::DoTask(ENamedThreads::Type CurrentThread, co
 	}
 
 	//Add an async task for pre completion
-	
 	FGraphEventArray PreCompletionPrerequistes;
 	AsyncHelper->PreCompletionTask = TGraphTask<FTaskPreCompletion>::CreateTask(&CompletionPrerequistes).ConstructAndDispatchWhenReady(InterchangeManager, WeakAsyncHelper);
 	PreCompletionPrerequistes.Add(AsyncHelper->PreCompletionTask);
 
+	FGraphEventArray AssetCompilationPrerequistes;
+	//Start the wait for asset compilation task on a async task
+	for (int32 SourceIndex = 0; SourceIndex < AsyncHelper->SourceDatas.Num(); ++SourceIndex)
+	{
+		AsyncHelper->WaitAssetCompilationTask = TGraphTask<FTaskWaitAssetCompilation>::CreateTask(&PreCompletionPrerequistes).ConstructAndDispatchWhenReady(SourceIndex, WeakAsyncHelper);
+		AssetCompilationPrerequistes.Add(AsyncHelper->WaitAssetCompilationTask);
+	}
+
 	//Start the Post pipeline task
+	FGraphEventArray PostImportPrerequistes;
 	for (int32 SourceIndex = 0; SourceIndex < AsyncHelper->SourceDatas.Num(); ++SourceIndex)
 	{
 		for (int32 GraphPipelineIndex = 0; GraphPipelineIndex < AsyncHelper->Pipelines.Num(); ++GraphPipelineIndex)
 		{
 			int32 GraphPipelineTaskIndex = AsyncHelper->PipelinePostImportTasks.Add(
-				TGraphTask<FTaskPipelinePostImport>::CreateTask(&(PreCompletionPrerequistes)).ConstructAndDispatchWhenReady(SourceIndex, GraphPipelineIndex, WeakAsyncHelper)
+				TGraphTask<FTaskPipelinePostImport>::CreateTask(&(AssetCompilationPrerequistes)).ConstructAndDispatchWhenReady(SourceIndex, GraphPipelineIndex, WeakAsyncHelper)
 			);
 			//Ensure we run the pipeline in the same order we create the task, since the pipeline modifies the node container, its important that its not processed in parallel, Adding the one we start to the prerequisites
 			//is the way to go here
-			PreCompletionPrerequistes.Add(AsyncHelper->PipelinePostImportTasks[GraphPipelineTaskIndex]);
+			PostImportPrerequistes.Add(AsyncHelper->PipelinePostImportTasks[GraphPipelineTaskIndex]);
 		}
 	}
 
 	FGraphEventArray PreAsyncCompletionPrerequistes;
-	AsyncHelper->PreAsyncCompletionTask = TGraphTask<FTaskPreAsyncCompletion>::CreateTask(&PreCompletionPrerequistes).ConstructAndDispatchWhenReady(InterchangeManager, WeakAsyncHelper);
+	AsyncHelper->PreAsyncCompletionTask = TGraphTask<FTaskPreAsyncCompletion>::CreateTask(&PostImportPrerequistes).ConstructAndDispatchWhenReady(InterchangeManager, WeakAsyncHelper);
 	PreAsyncCompletionPrerequistes.Add(AsyncHelper->PreAsyncCompletionTask);
 
 	AsyncHelper->CompletionTask = TGraphTask<FTaskCompletion>::CreateTask(&PreAsyncCompletionPrerequistes).ConstructAndDispatchWhenReady(InterchangeManager, WeakAsyncHelper);
diff --git a/Engine/Source/Runtime/Interchange/Engine/Private/Tasks/InterchangeTaskPipeline.cpp b/Engine/Source/Runtime/Interchange/Engine/Private/Tasks/InterchangeTaskPipeline.cpp
--- a/Engine/Source/Runtime/Interchange/Engine/Private/Tasks/InterchangeTaskPipeline.cpp
+++ b/Engine/Source/Runtime/Interchange/Engine/Private/Tasks/InterchangeTaskPipeline.cpp
@@ -1,11 +1,15 @@
 // Copyright Epic Games, Inc. All Rights Reserved.
 #include "InterchangeTaskPipeline.h"
 
+#include "AssetCompilingManager.h"
+#include "Async/Async.h"
 #include "Async/TaskGraphInterfaces.h"
 #include "CoreMinimal.h"
+#include "GenericPlatform/GenericPlatformProcess.h"
 #include "InterchangeEngineLogPrivate.h"
 #include "InterchangeManager.h"
 #include "InterchangePipelineBase.h"
+#include "Interfaces/Interface_AsyncCompilation.h"
 #include "Nodes/InterchangeBaseNodeContainer.h"
 #include "Nodes/InterchangeFactoryBaseNode.h"
 #include "Stats/Stats.h"
@@ -51,6 +55,72 @@ void UE::Interchange::FTaskPipeline::DoTask(ENamedThreads::Type CurrentThread, c
 	}
 }
 
+void UE::Interchange::FTaskWaitAssetCompilation::DoTask(ENamedThreads::Type CurrentThread, const FGraphEventRef& MyCompletionGraphEvent)
+{
+	TRACE_CPUPROFILER_EVENT_SCOPE("UE::Interchange::FTaskWaitAssetCompilation::DoTask")
+#if INTERCHANGE_TRACE_ASYNCHRONOUS_TASK_ENABLED
+		INTERCHANGE_TRACE_ASYNCHRONOUS_TASK(WaitAssetCompilation)
+#endif
+
+#if WITH_EDITOR
+
+	TSharedPtr<FImportAsyncHelper, ESPMode::ThreadSafe> AsyncHelper = WeakAsyncHelper.Pin();
+	if (!ensure(AsyncHelper.IsValid()) || AsyncHelper->bCancel)
+	{
+		return;
+	}
+
+	if (!ensure(AsyncHelper->SourceDatas.IsValidIndex(SourceIndex)))
+	{
+		return;
+	}
+
+	TArray<UObject*> ImportedObjects;
+
+	auto FillImportedObjectsFromSource = [&ImportedObjects](const TArray<UE::Interchange::FImportAsyncHelper::FImportedObjectInfo>& ImportedInfos)
+		{
+			ImportedObjects.Reserve(ImportedObjects.Num() + ImportedInfos.Num());
+			for (const UE::Interchange::FImportAsyncHelper::FImportedObjectInfo& ImportedInfo : ImportedInfos)
+			{
+				ImportedObjects.Add(ImportedInfo.ImportedObject);
+			}
+		};
+
+	AsyncHelper->IterateImportedAssets(SourceIndex, FillImportedObjectsFromSource);
+	AsyncHelper->IterateImportedSceneObjects(SourceIndex, FillImportedObjectsFromSource);
+
+	//Make sure all assets compilation are done before calling the pipeline post import task, let other thread execute if assets are not compile yet and wait 50ms before a new query
+	FPlatformProcess::ConditionalSleep([&ImportedObjects]()
+		{
+			//Compilation status cannot be ask in async thread, query the compile status on the main thread with a small fast function
+			//This ensure we dont stall the main thread until all assets are compile.
+			bool bCompilationFinish = false;
+			Async(EAsyncExecution::TaskGraphMainThread, [&bCompilationFinish, &ImportedObjects]()
+				{
+					TRACE_CPUPROFILER_EVENT_SCOPE("UE::Interchange::FTaskWaitAssetCompilation::DoTask::IsCompilingLambda_GameThread");
+					//Make sure all asset compiling managers are up to date, In case the game thread is waiting for the import to finish (like automation test or synchronous import)
+					FAssetCompilingManager::Get().ProcessAsyncTasks();
+
+					bCompilationFinish = true;
+					for (int32 ObjectIndex = 0; ObjectIndex < ImportedObjects.Num(); ++ObjectIndex)
+					{
+						UObject* ImportObject = ImportedObjects[ObjectIndex];
+						if (IInterface_AsyncCompilation* AssetCompilationInterface = Cast<IInterface_AsyncCompilation>(ImportObject))
+						{
+							if (AssetCompilationInterface->IsCompiling())
+							{
+								bCompilationFinish = false;
+								break;
+							}
+						}
+					}
+				}).Wait();
+			return bCompilationFinish;
+		}, 0.05f);
+
+#endif //WITH_EDITOR
+}
+
 void UE::Interchange::FTaskPipelinePostImport::DoTask(ENamedThreads::Type CurrentThread, const FGraphEventRef& MyCompletionGraphEvent)
 {
 	TRACE_CPUPROFILER_EVENT_SCOPE("UE::Interchange::FTaskPipelinePostImport::DoTask")
@@ -110,7 +180,7 @@ void UE::Interchange::FTaskPipelinePostImport::DoTask(ENamedThreads::Type Curren
 	}
 	UInterchangePipelineBase* Pipeline = AsyncHelper->Pipelines[PipelineIndex];
 
-	//Call the pipeline outside of the lock, we do this in case the pipeline take a long time. We call it for each asset created by this import
+	//Call the pipeline for each asset created by this import
 	for (int32 ObjectIndex = 0; ObjectIndex < ImportedObjects.Num(); ++ObjectIndex)
 	{
 		Pipeline->ScriptedExecutePostImportPipeline(NodeContainer, NodeUniqueIDs[ObjectIndex], ImportedObjects[ObjectIndex], IsAssetsReimported[ObjectIndex]);
diff --git a/Engine/Source/Runtime/Interchange/Engine/Private/Tasks/InterchangeTaskPipeline.h b/Engine/Source/Runtime/Interchange/Engine/Private/Tasks/InterchangeTaskPipeline.h
--- a/Engine/Source/Runtime/Interchange/Engine/Private/Tasks/InterchangeTaskPipeline.h
+++ b/Engine/Source/Runtime/Interchange/Engine/Private/Tasks/InterchangeTaskPipeline.h
@@ -55,6 +55,38 @@ namespace UE
 			void DoTask(ENamedThreads::Type CurrentThread, const FGraphEventRef& MyCompletionGraphEvent);
 		};
 
+		//We want to be sure any asset compilation is finish before calling FTaskPipelinePostImport, we use a async task to wait until they are done
+		class FTaskWaitAssetCompilation
+		{
+		private:
+			int32 SourceIndex;
+			TWeakPtr<FImportAsyncHelper, ESPMode::ThreadSafe> WeakAsyncHelper;
+
+		public:
+			FTaskWaitAssetCompilation(int32 InSourceIndex, TWeakPtr<FImportAsyncHelper, ESPMode::ThreadSafe> InAsyncHelper)
+				: SourceIndex(InSourceIndex)
+				, WeakAsyncHelper(InAsyncHelper)
+			{
+			}
+
+			ENamedThreads::Type GetDesiredThread()
+			{
+				return ENamedThreads::AnyBackgroundThreadNormalTask;
+			}
+
+			static ESubsequentsMode::Type GetSubsequentsMode()
+			{
+				return ESubsequentsMode::TrackSubsequents;
+			}
+
+			FORCEINLINE TStatId GetStatId() const
+			{
+				RETURN_QUICK_DECLARE_CYCLE_STAT(FTaskWaitAssetCompilation, STATGROUP_TaskGraphTasks);
+			}
+
+			void DoTask(ENamedThreads::Type CurrentThread, const FGraphEventRef& MyCompletionGraphEvent);
+		};
+
 		class FTaskPipelinePostImport
 		{
 		private:
diff --git a/Engine/Source/Runtime/Interchange/Engine/Public/InterchangeManager.h b/Engine/Source/Runtime/Interchange/Engine/Public/InterchangeManager.h
--- a/Engine/Source/Runtime/Interchange/Engine/Public/InterchangeManager.h
+++ b/Engine/Source/Runtime/Interchange/Engine/Public/InterchangeManager.h
@@ -216,6 +216,7 @@ namespace UE
 
 			TArray<FGraphEventRef> TranslatorTasks;
 			TArray<FGraphEventRef> PipelineTasks;
+			FGraphEventRef WaitAssetCompilationTask;
 			TArray<FGraphEventRef> PipelinePostImportTasks;
 			FGraphEventRef ParsingTask;
 			TArray<FGraphEventRef> BeginImportObjectTasks;
