Backout a material interface change to derive from IInterface_AsyncCompilation.

#rnx

diff --git a/Engine/Source/Runtime/Engine/Classes/Materials/MaterialInterface.h b/Engine/Source/Runtime/Engine/Classes/Materials/MaterialInterface.h
--- a/Engine/Source/Runtime/Engine/Classes/Materials/MaterialInterface.h
+++ b/Engine/Source/Runtime/Engine/Classes/Materials/MaterialInterface.h
@@ -252,7 +252,7 @@ public:
 };
 
 UCLASS(abstract, BlueprintType, MinimalAPI, HideCategories = (Thumbnail))
-class UMaterialInterface : public UObject, public IBlendableInterface, public IInterface_AssetUserData, public IInterface_AsyncCompilation
+class UMaterialInterface : public UObject, public IBlendableInterface, public IInterface_AssetUserData
 {
 	GENERATED_UCLASS_BODY()
 
@@ -998,13 +998,11 @@ public:
 	*/
 	virtual bool IsComplete() const { return true; }
 
-	/** IInterface_AsyncCompilation begin*/
 #if WITH_EDITOR
-	ENGINE_API virtual bool IsCompiling() const override { return false; };
+	ENGINE_API virtual bool IsCompiling() const { return false; };
 #else
 	FORCEINLINE bool IsCompiling() const { return false; }
 #endif
-	/** IInterface_AsyncCompilation end*/
 
 	/** @brief Checks to see if this material has all its shaders cached and if not, will perform a synchronous compilation of those.
 	*
diff --git a/Engine/Source/Runtime/Interchange/Engine/Private/Tasks/InterchangeTaskPipeline.cpp b/Engine/Source/Runtime/Interchange/Engine/Private/Tasks/InterchangeTaskPipeline.cpp
--- a/Engine/Source/Runtime/Interchange/Engine/Private/Tasks/InterchangeTaskPipeline.cpp
+++ b/Engine/Source/Runtime/Interchange/Engine/Private/Tasks/InterchangeTaskPipeline.cpp
@@ -10,6 +10,7 @@
 #include "InterchangeManager.h"
 #include "InterchangePipelineBase.h"
 #include "Interfaces/Interface_AsyncCompilation.h"
+#include "Materials/MaterialInterface.h"
 #include "Nodes/InterchangeBaseNodeContainer.h"
 #include "Nodes/InterchangeFactoryBaseNode.h"
 #include "Stats/Stats.h"
@@ -105,6 +106,14 @@ void UE::Interchange::FTaskWaitAssetCompilation::DoTask(ENamedThreads::Type Curr
 					for (int32 ObjectIndex = 0; ObjectIndex < ImportedObjects.Num(); ++ObjectIndex)
 					{
 						UObject* ImportObject = ImportedObjects[ObjectIndex];
+						if (UMaterialInterface* MaterialInterface = Cast<UMaterialInterface>(ImportObject))
+						{
+							if (MaterialInterface->IsCompiling())
+							{
+								bCompilationFinish = false;
+								break;
+							}
+						}
 						if (IInterface_AsyncCompilation* AssetCompilationInterface = Cast<IInterface_AsyncCompilation>(ImportObject))
 						{
 							if (AssetCompilationInterface->IsCompiling())
