Fix issue when material already exist and importing a mesh.

#rnx

diff --git a/Engine/Source/Runtime/Interchange/Engine/Private/Tasks/InterchangeTaskImportObject.cpp b/Engine/Source/Runtime/Interchange/Engine/Private/Tasks/InterchangeTaskImportObject.cpp
--- a/Engine/Source/Runtime/Interchange/Engine/Private/Tasks/InterchangeTaskImportObject.cpp
+++ b/Engine/Source/Runtime/Interchange/Engine/Private/Tasks/InterchangeTaskImportObject.cpp
@@ -309,6 +309,7 @@ void UE::Interchange::FTaskImportObject_GameThread::DoTask(ENamedThreads::Type C
 	}
 	UInterchangeBaseNodeContainer* NodeContainer = AsyncHelper->BaseNodeContainers[SourceIndex].Get();
 
+	bool bSkipObjectNoReplace = false;
 	UObject* ExistingAsset = nullptr;
 	//If we do a reimport no need to create a package
 	if (ObjectToReimport)
@@ -389,7 +390,7 @@ void UE::Interchange::FTaskImportObject_GameThread::DoTask(ENamedThreads::Type C
 			Message->AssetType = FactoryNode->GetObjectClass();
 			Message->Text = FText::Format(NSLOCTEXT("FTaskImportObject_GameThread", "CouldntReplaceExistingAsset", "The option bReplaceExisting is false so we are not overriding the asset named '{0}'.")
 				, FText::FromString(AssetName));
-			return;
+			bSkipObjectNoReplace = true;
 		}
 
 		if (!bPackageWasCreated && AsyncHelper->TaskData.bFollowRedirectors)
@@ -449,6 +450,10 @@ void UE::Interchange::FTaskImportObject_GameThread::DoTask(ENamedThreads::Type C
 
 		//Set the factory node reference to the existing object, so we do not need to find the asset in the factory
 		FactoryNode->SetCustomReferenceObject(FSoftObjectPath(ExistingAsset));
+		if (bSkipObjectNoReplace)
+		{
+			return;
+		}
 
 		if (UInterchangeManager::GetInterchangeManager().IsObjectBeingImported(ExistingAsset))
 		{
