Fix issue with Flush Rendering command call twice

#rnx

diff --git a/Engine/Plugins/Interchange/Runtime/Source/Dispatcher/Private/InterchangeCommands.cpp b/Engine/Plugins/Interchange/Runtime/Source/Dispatcher/Private/InterchangeCommands.cpp
--- a/Engine/Plugins/Interchange/Runtime/Source/Dispatcher/Private/InterchangeCommands.cpp
+++ b/Engine/Plugins/Interchange/Runtime/Source/Dispatcher/Private/InterchangeCommands.cpp
@@ -10,6 +10,21 @@ namespace UE
 {
 	namespace Interchange
 	{
+		int32 DispatcherCommandVersion::GetMajor()
+		{
+			return 17;
+		}
+
+		int32 DispatcherCommandVersion::GetMinor()
+		{
+			return 0;
+		}
+
+		int32 DispatcherCommandVersion::GetPatch()
+		{
+			return 0;
+		}
+
 		TSharedPtr<ICommand> CreateCommand(ECommandId CommandType)
 		{
 			switch (CommandType)
diff --git a/Engine/Plugins/Interchange/Runtime/Source/Dispatcher/Public/InterchangeCommands.h b/Engine/Plugins/Interchange/Runtime/Source/Dispatcher/Public/InterchangeCommands.h
--- a/Engine/Plugins/Interchange/Runtime/Source/Dispatcher/Public/InterchangeCommands.h
+++ b/Engine/Plugins/Interchange/Runtime/Source/Dispatcher/Public/InterchangeCommands.h
@@ -17,11 +17,11 @@ namespace UE
 		{
 		public:
 			//Major version should be updated when there is an existing API that has been change.
-			static int32 GetMajor() { return 16; }
+			static int32 GetMajor();
 			//Minor version must be updated when there is an addition to the API.
-			static int32 GetMinor() { return 0; }
+			static int32 GetMinor();
 			//Patch version should be update if there is some bug fixes in the private code.
-			static int32 GetPatch() { return 0; }
+			static int32 GetPatch();
 
 			//LWCDisabled version tell if the code was compile with serialization compatibility.
 			static bool GetLWCDisabled()
diff --git a/Engine/Plugins/Interchange/Runtime/Source/Import/Private/Fbx/InterchangeFbxTranslator.cpp b/Engine/Plugins/Interchange/Runtime/Source/Import/Private/Fbx/InterchangeFbxTranslator.cpp
--- a/Engine/Plugins/Interchange/Runtime/Source/Import/Private/Fbx/InterchangeFbxTranslator.cpp
+++ b/Engine/Plugins/Interchange/Runtime/Source/Import/Private/Fbx/InterchangeFbxTranslator.cpp
@@ -189,6 +189,10 @@ void UInterchangeFbxTranslator::ImportFinish()
 
 TOptional<UE::Interchange::FImportImage> UInterchangeFbxTranslator::GetTexturePayloadData(const FString& PayLoadKey, TOptional<FString>& AlternateTexturePath) const
 {
+	if (PayLoadKey.IsEmpty())
+	{
+		return TOptional<UE::Interchange::FImportImage>();
+	}
 	UE::Interchange::Private::FScopedTranslator ScopedTranslator(PayLoadKey, Results);
 	const IInterchangeTexturePayloadInterface* TextureTranslator = ScopedTranslator.GetPayLoadInterface<IInterchangeTexturePayloadInterface>();
 	if (!ensure(TextureTranslator))
diff --git a/Engine/Plugins/Interchange/Runtime/Source/Parsers/Fbx/Private/FbxMaterial.cpp b/Engine/Plugins/Interchange/Runtime/Source/Parsers/Fbx/Private/FbxMaterial.cpp
--- a/Engine/Plugins/Interchange/Runtime/Source/Parsers/Fbx/Private/FbxMaterial.cpp
+++ b/Engine/Plugins/Interchange/Runtime/Source/Parsers/Fbx/Private/FbxMaterial.cpp
@@ -37,6 +37,10 @@ namespace UE
 
 			const UInterchangeTexture2DNode* FFbxMaterial::CreateTexture2DNode(UInterchangeBaseNodeContainer& NodeContainer, const FString& TextureFilePath)
 			{
+				if (TextureFilePath.IsEmpty())
+				{
+					return nullptr;
+				}
 				const FString TextureName = FPaths::GetBaseFilename(TextureFilePath);
 				const FString TextureNodeID = UInterchangeTextureNode::MakeNodeUid(TextureName);
 
diff --git a/Engine/Plugins/Tests/InterchangeTests/Source/InterchangeTests/Private/InterchangeImportTest.cpp b/Engine/Plugins/Tests/InterchangeTests/Source/InterchangeTests/Private/InterchangeImportTest.cpp
--- a/Engine/Plugins/Tests/InterchangeTests/Source/InterchangeTests/Private/InterchangeImportTest.cpp
+++ b/Engine/Plugins/Tests/InterchangeTests/Source/InterchangeTests/Private/InterchangeImportTest.cpp
@@ -1,16 +1,16 @@
 // Copyright Epic Games, Inc. All Rights Reserved.
 
+#include "AssetCompilingManager.h"
 #include "AssetRegistry/ARFilter.h"
+#include "AssetRegistry/AssetRegistryModule.h"
+#include "Editor.h"
+#include "Editor/Transactor.h"
 #include "HAL/FileManager.h"
 #include "Misc/AutomationTest.h"
 #include "InterchangeDispatcher.h"
 #include "InterchangeImportTestData.h"
 #include "InterchangeImportTestPlan.h"
 #include "InterchangeImportTestStepBase.h"
-#include "AssetRegistry/AssetRegistryModule.h"
-#include "Editor.h"
-#include "Editor/Transactor.h"
-
 #include "Modules/ModuleManager.h"
 #include "ObjectTools.h"
 
@@ -326,11 +326,15 @@ bool FInterchangeImportTest::RunTest(const FString& Path)
 			}
 			else
 			{
+
 				ObjectsToDelete.Add(ResultObject);
 			}
 		}
 	}
 
+	//Make sure all compilation is done before deleting some objects
+	FAssetCompilingManager::Get().FinishAllCompilation();
+
 	constexpr bool bShowConfirmation = false;
 	ObjectTools::ForceDeleteObjects(ObjectsToDelete, bShowConfirmation);
 
diff --git a/Engine/Source/Runtime/Interchange/Engine/Private/Tasks/InterchangeTaskParsing.cpp b/Engine/Source/Runtime/Interchange/Engine/Private/Tasks/InterchangeTaskParsing.cpp
--- a/Engine/Source/Runtime/Interchange/Engine/Private/Tasks/InterchangeTaskParsing.cpp
+++ b/Engine/Source/Runtime/Interchange/Engine/Private/Tasks/InterchangeTaskParsing.cpp
@@ -419,7 +419,6 @@ void UE::Interchange::FTaskParsing::DoTask(ENamedThreads::Type CurrentThread, co
 	}
 
 	//Start the Post pipeline task
-	FGraphEventArray PostImportPrerequistes;
 	for (int32 SourceIndex = 0; SourceIndex < AsyncHelper->SourceDatas.Num(); ++SourceIndex)
 	{
 		for (int32 GraphPipelineIndex = 0; GraphPipelineIndex < AsyncHelper->Pipelines.Num(); ++GraphPipelineIndex)
@@ -429,12 +428,12 @@ void UE::Interchange::FTaskParsing::DoTask(ENamedThreads::Type CurrentThread, co
 			);
 			//Ensure we run the pipeline in the same order we create the task, since the pipeline modifies the node container, its important that its not processed in parallel, Adding the one we start to the prerequisites
 			//is the way to go here
-			PostImportPrerequistes.Add(AsyncHelper->PipelinePostImportTasks[GraphPipelineTaskIndex]);
+			AssetCompilationPrerequistes.Add(AsyncHelper->PipelinePostImportTasks[GraphPipelineTaskIndex]);
 		}
 	}
 
 	FGraphEventArray PreAsyncCompletionPrerequistes;
-	AsyncHelper->PreAsyncCompletionTask = TGraphTask<FTaskPreAsyncCompletion>::CreateTask(&PostImportPrerequistes).ConstructAndDispatchWhenReady(InterchangeManager, WeakAsyncHelper);
+	AsyncHelper->PreAsyncCompletionTask = TGraphTask<FTaskPreAsyncCompletion>::CreateTask(&AssetCompilationPrerequistes).ConstructAndDispatchWhenReady(InterchangeManager, WeakAsyncHelper);
 	PreAsyncCompletionPrerequistes.Add(AsyncHelper->PreAsyncCompletionTask);
 
 	AsyncHelper->CompletionTask = TGraphTask<FTaskCompletion>::CreateTask(&PreAsyncCompletionPrerequistes).ConstructAndDispatchWhenReady(InterchangeManager, WeakAsyncHelper);
diff --git a/Engine/Source/Runtime/RenderCore/Private/RenderingThread.cpp b/Engine/Source/Runtime/RenderCore/Private/RenderingThread.cpp
--- a/Engine/Source/Runtime/RenderCore/Private/RenderingThread.cpp
+++ b/Engine/Source/Runtime/RenderCore/Private/RenderingThread.cpp
@@ -1334,11 +1334,6 @@ static void GameThreadWaitForTask(const UE::Tasks::FTask& Task, bool bEmptyGameT
 			NumRecursiveCalls++;
 			if (NumRecursiveCalls > 1)
 			{
-				if(GIsAutomationTesting)
-				{
-					// temp test to log callstacks for this being triggered during automation tests
-					ensureMsgf(false, TEXT("FlushRenderingCommands called recursively! %d calls on the stack."), NumRecursiveCalls);
-				}
 				UE_LOG(LogRendererCore,Warning,TEXT("FlushRenderingCommands called recursively! %d calls on the stack."), NumRecursiveCalls);
 			}
 			if (NumRecursiveCalls > 1 || FTaskGraphInterface::Get().IsThreadProcessingTasks(ENamedThreads::GameThread))
