Interchange: Fix the cancel when importing multiple file with fbx in-process.

#rnx

diff --git a/Engine/Plugins/Interchange/Runtime/Source/Dispatcher/Private/InterchangeCommands.cpp b/Engine/Plugins/Interchange/Runtime/Source/Dispatcher/Private/InterchangeCommands.cpp
--- a/Engine/Plugins/Interchange/Runtime/Source/Dispatcher/Private/InterchangeCommands.cpp
+++ b/Engine/Plugins/Interchange/Runtime/Source/Dispatcher/Private/InterchangeCommands.cpp
@@ -12,7 +12,7 @@ namespace UE
 	{
 		int32 DispatcherCommandVersion::GetMajor()
 		{
-			return 17;
+			return 18;
 		}
 
 		int32 DispatcherCommandVersion::GetMinor()
diff --git a/Engine/Source/Runtime/Interchange/Engine/Private/InterchangeManager.cpp b/Engine/Source/Runtime/Interchange/Engine/Private/InterchangeManager.cpp
--- a/Engine/Source/Runtime/Interchange/Engine/Private/InterchangeManager.cpp
+++ b/Engine/Source/Runtime/Interchange/Engine/Private/InterchangeManager.cpp
@@ -1262,16 +1262,29 @@ void UInterchangeManager::StartQueuedTasks(bool bCancelAllTasks /*= false*/)
 		{
 			continue;
 		}
-		bool& TranslatorLock = NonParallelTranslatorLocks.FindChecked(ClassAndTasks.Key);
-		if (!TranslatorLock)
-		{
-			FQueuedTaskData QueuedTaskData = ClassAndTasks.Value[0];
-			QueuedTasks.Enqueue(QueuedTaskData);
-			TranslatorLock = true;
-			constexpr bool bAllowShrinking = false;
-			ClassAndTasks.Value.RemoveAt(0, 1, bAllowShrinking);
-			//No need to process an another the lock is set
-			continue;
+		if (bCancelAllTasks)
+		{
+			//Enqueue all the tasks they will be all cancel
+			for (FQueuedTaskData QueuedTaskData : ClassAndTasks.Value)
+			{
+				QueuedTasks.Enqueue(QueuedTaskData);
+			}
+			ClassAndTasks.Value.Reset();
+		}
+		else
+		{
+			//Lock the translator and enqueue only the first task
+			bool& TranslatorLock = NonParallelTranslatorLocks.FindChecked(ClassAndTasks.Key);
+			if (!TranslatorLock)
+			{
+				FQueuedTaskData QueuedTaskData = ClassAndTasks.Value[0];
+				QueuedTasks.Enqueue(QueuedTaskData);
+				TranslatorLock = true;
+				constexpr bool bAllowShrinking = false;
+				ClassAndTasks.Value.RemoveAt(0, 1, bAllowShrinking);
+				//No need to process an another the lock is set
+				continue;
+			}
 		}
 	}
 
@@ -2232,7 +2245,7 @@ void UInterchangeManager::CancelAllTasks()
 
 	//Cancel the queued tasks, we cannot simply not do them since, there is some promise objects
 	//to setup in the completion task
-	const bool bCancelAllTasks = true;
+	constexpr bool bCancelAllTasks = true;
 	StartQueuedTasks(bCancelAllTasks);
 
 	//Set the cancel state on all running tasks
@@ -2245,6 +2258,20 @@ void UInterchangeManager::CancelAllTasks()
 			AsyncHelper->InitCancel();
 		}
 	}
+	for (TPair<UClass*, TArray<FQueuedTaskData>>& ClassAndTasks : NonParallelTranslatorQueueTasks)
+	{
+		//After calling StartQueuedTasks with bCancelAllTasks at true, we should not have any waiting task here
+		if (ClassAndTasks.Value.IsEmpty())
+		{
+			continue;
+		}
+		//if we still have some task we need to Cancel them asap
+		FQueuedTaskData QueuedTaskData = ClassAndTasks.Value[0];
+		if (QueuedTaskData.AsyncHelper.IsValid())
+		{
+			QueuedTaskData.AsyncHelper->InitCancel();
+		}
+	}
 	//Tasks should all finish quite fast now
 };
 
diff --git a/Engine/Source/Runtime/Interchange/Engine/Private/Tasks/InterchangeTaskParsing.cpp b/Engine/Source/Runtime/Interchange/Engine/Private/Tasks/InterchangeTaskParsing.cpp
--- a/Engine/Source/Runtime/Interchange/Engine/Private/Tasks/InterchangeTaskParsing.cpp
+++ b/Engine/Source/Runtime/Interchange/Engine/Private/Tasks/InterchangeTaskParsing.cpp
@@ -328,7 +328,7 @@ void UE::Interchange::FTaskParsing::DoTask(ENamedThreads::Type CurrentThread, co
 		for(TObjectPtr<UInterchangeTranslatorBase> Translator : AsyncHelper->Translators)
 		{
 			//We need a valid translator to add a result message
-			if (!Translator)
+			if (!Translator || AsyncHelper->bCancel)
 			{
 				continue;
 			}
