Interchange: Make sure InterchangeWorker process are terminate right when we discard the translator.

#rnx

diff --git a/Engine/Source/Runtime/Interchange/Engine/Private/InterchangeManager.cpp b/Engine/Source/Runtime/Interchange/Engine/Private/InterchangeManager.cpp
--- a/Engine/Source/Runtime/Interchange/Engine/Private/InterchangeManager.cpp
+++ b/Engine/Source/Runtime/Interchange/Engine/Private/InterchangeManager.cpp
@@ -174,11 +174,16 @@ UE::Interchange::FScopedInterchangeImportEnableState::~FScopedInterchangeImportE
 
 UE::Interchange::FScopedSourceData::FScopedSourceData(const FString& Filename)
 {
-	//Found the translator
 	SourceDataPtr = TStrongObjectPtr<UInterchangeSourceData>(UInterchangeManager::GetInterchangeManager().CreateSourceData(Filename));
-	check(SourceDataPtr.IsValid());
+	ensure(SourceDataPtr.IsValid());
+}
+
+UE::Interchange::FScopedSourceData::~FScopedSourceData()
+{
+	SourceDataPtr.Reset();
 }
 
+
 UInterchangeSourceData* UE::Interchange::FScopedSourceData::GetSourceData() const
 {
 	return SourceDataPtr.Get();
@@ -190,6 +195,16 @@ UE::Interchange::FScopedTranslator::FScopedTranslator(const UInterchangeSourceDa
 	ScopedTranslatorPtr = TStrongObjectPtr<UInterchangeTranslatorBase>(UInterchangeManager::GetInterchangeManager().GetTranslatorForSourceData(SourceData));
 }
 
+UE::Interchange::FScopedTranslator::~FScopedTranslator()
+{
+	//Found the translator
+	if (ScopedTranslatorPtr.IsValid())
+	{
+		ScopedTranslatorPtr->ReleaseSource();
+	}
+	ScopedTranslatorPtr.Reset();
+}
+
 UInterchangeTranslatorBase* UE::Interchange::FScopedTranslator::GetTranslator()
 {
 	return ScopedTranslatorPtr.Get();
@@ -1139,7 +1154,13 @@ bool UInterchangeManager::CanTranslateSourceData(const UInterchangeSourceData* S
 	}
 #endif
 
-	return GetTranslatorForSourceData(SourceData) != nullptr;
+	if (UInterchangeTranslatorBase* Translator = GetTranslatorForSourceData(SourceData))
+	{
+		Translator->ReleaseSource();
+		return true;
+	}
+
+	return false;
 }
 
 bool UInterchangeManager::CanReimport(const UObject* Object, TArray<FString>& OutFilenames) const
@@ -1445,12 +1466,13 @@ UInterchangeManager::ImportInternal(const FString& ContentPath, const UInterchan
 	}
 	
 	{
-		const UInterchangeTranslatorBase* Translator = GetTranslatorForSourceData(SourceData);
+		UInterchangeTranslatorBase* Translator = GetTranslatorForSourceData(SourceData);
 		if(!Translator)
 		{
 			UE_LOG(LogInterchangeEngine, Error, TEXT("Cannot import file, the source data is not supported. See if you can enable for interchange the extension [%s]"), *FPaths::GetExtension(SourceData->GetFilename()));
 			return EarlyExit();
 		}
+		Translator->ReleaseSource();
 	}
 
 	if (FEngineAnalytics::IsAvailable())
diff --git a/Engine/Source/Runtime/Interchange/Engine/Public/InterchangeManager.h b/Engine/Source/Runtime/Interchange/Engine/Public/InterchangeManager.h
--- a/Engine/Source/Runtime/Interchange/Engine/Public/InterchangeManager.h
+++ b/Engine/Source/Runtime/Interchange/Engine/Public/InterchangeManager.h
@@ -60,6 +60,7 @@ namespace UE
 		{
 		public:
 			INTERCHANGEENGINE_API explicit FScopedSourceData(const FString& Filename);
+			INTERCHANGEENGINE_API ~FScopedSourceData();
 			INTERCHANGEENGINE_API UInterchangeSourceData* GetSourceData() const;
 		private:
 			TStrongObjectPtr<UInterchangeSourceData> SourceDataPtr = nullptr;
@@ -69,6 +70,7 @@ namespace UE
 		{
 		public:
 			INTERCHANGEENGINE_API explicit FScopedTranslator(const UInterchangeSourceData* SourceData);
+			INTERCHANGEENGINE_API ~FScopedTranslator();
 			INTERCHANGEENGINE_API UInterchangeTranslatorBase* GetTranslator();
 
 		private:
