Fix static analysis
#rnx

diff --git a/Engine/Source/Runtime/Interchange/Core/Private/Nodes/InterchangeFactoryBaseNode.cpp b/Engine/Source/Runtime/Interchange/Core/Private/Nodes/InterchangeFactoryBaseNode.cpp
--- a/Engine/Source/Runtime/Interchange/Core/Private/Nodes/InterchangeFactoryBaseNode.cpp
+++ b/Engine/Source/Runtime/Interchange/Core/Private/Nodes/InterchangeFactoryBaseNode.cpp
@@ -57,7 +57,7 @@ FString UInterchangeFactoryBaseNode::GetKeyDisplayName(const UE::Interchange::FA
 	const FString OriginalKeyName = KeyDisplayName;
 	if (NodeAttributeKey == Macro_CustomSubPathKey)
 	{
-		KeyDisplayName = TEXT("Import Sub-Path");
+		KeyDisplayName = TEXT("Asset Sub-Path");
 	}
 	else if (OriginalKeyName.Equals(UE::Interchange::FFactoryBaseNodeStaticData::FactoryDependenciesBaseKey()))
 	{
@@ -85,10 +85,6 @@ FString UInterchangeFactoryBaseNode::GetKeyDisplayName(const UE::Interchange::FA
 	{
 		KeyDisplayName = TEXT("Force Node Reimport");
 	}
-	else if (NodeAttributeKey == Macro_CustomSubPathKey)
-	{
-		KeyDisplayName = TEXT("Asset Sub Path");
-	}
 	else
 	{
 		KeyDisplayName = Super::GetKeyDisplayName(NodeAttributeKey);
