Interchange: Add the content base path to the ExecutePipeline API. We will need this data to be able to find existing material when we will implement the material workflow for asset import.

diff --git a/Engine/Plugins/Enterprise/DataprepEditor/Source/DataprepCore/Private/InterchangeDataprepPipeline.cpp b/Engine/Plugins/Enterprise/DataprepEditor/Source/DataprepCore/Private/InterchangeDataprepPipeline.cpp
--- a/Engine/Plugins/Enterprise/DataprepEditor/Source/DataprepCore/Private/InterchangeDataprepPipeline.cpp
+++ b/Engine/Plugins/Enterprise/DataprepEditor/Source/DataprepCore/Private/InterchangeDataprepPipeline.cpp
@@ -10,14 +10,14 @@
 #include "Engine/StaticMesh.h"
 #include "LevelSequence.h"
 
-void UInterchangeDataprepLevelPipeline::ExecutePipeline(UInterchangeBaseNodeContainer* NodeContainer, const TArray<UInterchangeSourceData*>& InSourceDatas)
+void UInterchangeDataprepLevelPipeline::ExecutePipeline(UInterchangeBaseNodeContainer* NodeContainer, const TArray<UInterchangeSourceData*>& InSourceDatas, const FString& ContentBasePath)
 {
 	if (!NodeContainer)
 	{
 		return;
 	}
 
-	Super::ExecutePipeline(NodeContainer, InSourceDatas);
+	Super::ExecutePipeline(NodeContainer, InSourceDatas, ContentBasePath);
 
 	static TMap<UClass*, FString> SubPathsPerClass{
 		{ UTexture::StaticClass(), "Textures" },
diff --git a/Engine/Plugins/Enterprise/DataprepEditor/Source/DataprepCore/Private/InterchangeDataprepPipeline.h b/Engine/Plugins/Enterprise/DataprepEditor/Source/DataprepCore/Private/InterchangeDataprepPipeline.h
--- a/Engine/Plugins/Enterprise/DataprepEditor/Source/DataprepCore/Private/InterchangeDataprepPipeline.h
+++ b/Engine/Plugins/Enterprise/DataprepEditor/Source/DataprepCore/Private/InterchangeDataprepPipeline.h
@@ -15,6 +15,6 @@ class UInterchangeDataprepLevelPipeline : public UInterchangePipelineBase
 	GENERATED_BODY()
 
 protected:
-	virtual void ExecutePipeline(UInterchangeBaseNodeContainer* BaseNodeContainer, const TArray<UInterchangeSourceData*>& SourceDatas) override;
+	virtual void ExecutePipeline(UInterchangeBaseNodeContainer* BaseNodeContainer, const TArray<UInterchangeSourceData*>& SourceDatas, const FString& ContentBasePath) override;
 	virtual void ExecutePostImportPipeline(const UInterchangeBaseNodeContainer* BaseNodeContainer, const FString& NodeKey, UObject* CreatedAsset, bool bIsAReimport) override;
 };
\ No newline at end of file
diff --git a/Engine/Plugins/Enterprise/DatasmithImporter/Source/DatasmithImporter/Private/DatasmithTextureImporter.cpp b/Engine/Plugins/Enterprise/DatasmithImporter/Source/DatasmithImporter/Private/DatasmithTextureImporter.cpp
--- a/Engine/Plugins/Enterprise/DatasmithImporter/Source/DatasmithImporter/Private/DatasmithTextureImporter.cpp
+++ b/Engine/Plugins/Enterprise/DatasmithImporter/Source/DatasmithImporter/Private/DatasmithTextureImporter.cpp
@@ -411,7 +411,7 @@ UE::Interchange::FAssetImportResultRef FDatasmithTextureImporter::CreateTextureA
 	return UInterchangeManager::GetInterchangeManager().ImportAssetAsync( ContentPath, ScopedSourceData.GetSourceData(), ImportAssetParameters );
 }
 
-void UDatasmithTexturePipeline::ExecutePipeline(UInterchangeBaseNodeContainer* BaseNodeContainer, const TArray<UInterchangeSourceData*>& SourceDatas)
+void UDatasmithTexturePipeline::ExecutePipeline(UInterchangeBaseNodeContainer* BaseNodeContainer, const TArray<UInterchangeSourceData*>& SourceDatas, const FString& ContentBasePath)
 {
 	if ( !TextureElement.IsValid() )
 	{
diff --git a/Engine/Plugins/Enterprise/DatasmithImporter/Source/DatasmithImporter/Public/DatasmithTextureImporter.h b/Engine/Plugins/Enterprise/DatasmithImporter/Source/DatasmithImporter/Public/DatasmithTextureImporter.h
--- a/Engine/Plugins/Enterprise/DatasmithImporter/Source/DatasmithImporter/Public/DatasmithTextureImporter.h
+++ b/Engine/Plugins/Enterprise/DatasmithImporter/Source/DatasmithImporter/Public/DatasmithTextureImporter.h
@@ -46,7 +46,7 @@ class UDatasmithTexturePipeline : public UInterchangePipelineBase
 	GENERATED_BODY()
 
 public:
-	virtual void ExecutePipeline(UInterchangeBaseNodeContainer* BaseNodeContainer, const TArray<UInterchangeSourceData*>& SourceDatas) override;
+	virtual void ExecutePipeline(UInterchangeBaseNodeContainer* BaseNodeContainer, const TArray<UInterchangeSourceData*>& SourceDatas, const FString& ContentBasePath) override;
 
 	bool IsScripted() override
 	{
diff --git a/Engine/Plugins/Experimental/Enterprise/DatasmithInterchange/Source/DatasmithInterchange/Private/InterchangeDatasmithLevelPipeline.cpp b/Engine/Plugins/Experimental/Enterprise/DatasmithInterchange/Source/DatasmithInterchange/Private/InterchangeDatasmithLevelPipeline.cpp
--- a/Engine/Plugins/Experimental/Enterprise/DatasmithInterchange/Source/DatasmithInterchange/Private/InterchangeDatasmithLevelPipeline.cpp
+++ b/Engine/Plugins/Experimental/Enterprise/DatasmithInterchange/Source/DatasmithInterchange/Private/InterchangeDatasmithLevelPipeline.cpp
@@ -22,11 +22,11 @@
 #include "Materials/MaterialInterface.h"
 #include "Misc/PackageName.h"
 
-void UInterchangeDatasmithLevelPipeline::ExecutePipeline(UInterchangeBaseNodeContainer* NodeContainer, const TArray<UInterchangeSourceData*>& InSourceDatas)
+void UInterchangeDatasmithLevelPipeline::ExecutePipeline(UInterchangeBaseNodeContainer* NodeContainer, const TArray<UInterchangeSourceData*>& InSourceDatas, const FString& ContentBasePath)
 {
 	using namespace UE::DatasmithInterchange;
 
-	Super::ExecutePipeline(NodeContainer, InSourceDatas);
+	Super::ExecutePipeline(NodeContainer, InSourceDatas, ContentBasePath);
 
 	// Add material factory dependencies for mesh actors where all overrides are filled with the same material
 	for (UInterchangeMeshActorFactoryNode* MeshActorFactoryNode : NodeUtils::GetNodes<UInterchangeMeshActorFactoryNode>(NodeContainer))
diff --git a/Engine/Plugins/Experimental/Enterprise/DatasmithInterchange/Source/DatasmithInterchange/Private/InterchangeDatasmithLevelPipeline.h b/Engine/Plugins/Experimental/Enterprise/DatasmithInterchange/Source/DatasmithInterchange/Private/InterchangeDatasmithLevelPipeline.h
--- a/Engine/Plugins/Experimental/Enterprise/DatasmithInterchange/Source/DatasmithInterchange/Private/InterchangeDatasmithLevelPipeline.h
+++ b/Engine/Plugins/Experimental/Enterprise/DatasmithInterchange/Source/DatasmithInterchange/Private/InterchangeDatasmithLevelPipeline.h
@@ -22,7 +22,7 @@ class UInterchangeDatasmithLevelPipeline : public UInterchangeGenericLevelPipeli
 	GENERATED_BODY()
 
 protected:
-	virtual void ExecutePipeline(UInterchangeBaseNodeContainer* BaseNodeContainer, const TArray<UInterchangeSourceData*>& SourceDatas) override;
+	virtual void ExecutePipeline(UInterchangeBaseNodeContainer* BaseNodeContainer, const TArray<UInterchangeSourceData*>& SourceDatas, const FString& ContentBasePath) override;
 	virtual void ExecutePostImportPipeline(const UInterchangeBaseNodeContainer* BaseNodeContainer, const FString& NodeKey, UObject* CreatedAsset, bool bIsAReimport) override;
 
 	virtual UInterchangeActorFactoryNode* CreateActorFactoryNode(const UInterchangeSceneNode* SceneNode, const UInterchangeBaseNode* TranslatedAssetNode) const override;
diff --git a/Engine/Plugins/Experimental/Enterprise/DatasmithInterchange/Source/DatasmithInterchange/Private/InterchangeDatasmithMaterialPipeline.cpp b/Engine/Plugins/Experimental/Enterprise/DatasmithInterchange/Source/DatasmithInterchange/Private/InterchangeDatasmithMaterialPipeline.cpp
--- a/Engine/Plugins/Experimental/Enterprise/DatasmithInterchange/Source/DatasmithInterchange/Private/InterchangeDatasmithMaterialPipeline.cpp
+++ b/Engine/Plugins/Experimental/Enterprise/DatasmithInterchange/Source/DatasmithInterchange/Private/InterchangeDatasmithMaterialPipeline.cpp
@@ -60,7 +60,7 @@ UInterchangeDatasmithMaterialPipeline::UInterchangeDatasmithMaterialPipeline()
 	MaterialImport = EInterchangeMaterialImportOption::ImportAsMaterials;
 }
 
-void UInterchangeDatasmithMaterialPipeline::ExecutePipeline(UInterchangeBaseNodeContainer* NodeContainer, const TArray<UInterchangeSourceData*>& InSourceDatas)
+void UInterchangeDatasmithMaterialPipeline::ExecutePipeline(UInterchangeBaseNodeContainer* NodeContainer, const TArray<UInterchangeSourceData*>& InSourceDatas, const FString& ContentBasePath)
 {
 	using namespace UE::DatasmithInterchange;
 
@@ -85,7 +85,7 @@ void UInterchangeDatasmithMaterialPipeline::ExecutePipeline(UInterchangeBaseNode
 			}
 		});
 
-	Super::ExecutePipeline(NodeContainer, InSourceDatas);
+	Super::ExecutePipeline(NodeContainer, InSourceDatas, ContentBasePath);
 
 	UpdateMaterialFactoryNodes(ShaderNodes);
 
diff --git a/Engine/Plugins/Experimental/Enterprise/DatasmithInterchange/Source/DatasmithInterchange/Private/InterchangeDatasmithMaterialPipeline.h b/Engine/Plugins/Experimental/Enterprise/DatasmithInterchange/Source/DatasmithInterchange/Private/InterchangeDatasmithMaterialPipeline.h
--- a/Engine/Plugins/Experimental/Enterprise/DatasmithInterchange/Source/DatasmithInterchange/Private/InterchangeDatasmithMaterialPipeline.h
+++ b/Engine/Plugins/Experimental/Enterprise/DatasmithInterchange/Source/DatasmithInterchange/Private/InterchangeDatasmithMaterialPipeline.h
@@ -27,7 +27,7 @@ class DATASMITHINTERCHANGE_API UInterchangeDatasmithMaterialPipeline : public UI
 	UInterchangeDatasmithMaterialPipeline();
 
 protected:
-	virtual void ExecutePipeline(UInterchangeBaseNodeContainer* BaseNodeContainer, const TArray<UInterchangeSourceData*>& SourceDatas) override;
+	virtual void ExecutePipeline(UInterchangeBaseNodeContainer* BaseNodeContainer, const TArray<UInterchangeSourceData*>& SourceDatas, const FString& ContentBasePath) override;
 	virtual void ExecutePostImportPipeline(const UInterchangeBaseNodeContainer* BaseNodeContainer, const FString& NodeKey, UObject* CreatedAsset, bool bIsAReimport) override;
 
 	virtual bool CanExecuteOnAnyThread(EInterchangePipelineTask PipelineTask) override
diff --git a/Engine/Plugins/Experimental/Enterprise/DatasmithInterchange/Source/DatasmithInterchange/Private/InterchangeDatasmithPipeline.cpp b/Engine/Plugins/Experimental/Enterprise/DatasmithInterchange/Source/DatasmithInterchange/Private/InterchangeDatasmithPipeline.cpp
--- a/Engine/Plugins/Experimental/Enterprise/DatasmithInterchange/Source/DatasmithInterchange/Private/InterchangeDatasmithPipeline.cpp
+++ b/Engine/Plugins/Experimental/Enterprise/DatasmithInterchange/Source/DatasmithInterchange/Private/InterchangeDatasmithPipeline.cpp
@@ -120,7 +120,7 @@ UInterchangeDatasmithPipeline::UInterchangeDatasmithPipeline()
 	AnimationPipeline->CommonSkeletalMeshesAndAnimationsProperties = CommonSkeletalMeshesAndAnimationsProperties;
 }
 
-void UInterchangeDatasmithPipeline::ExecutePipeline(UInterchangeBaseNodeContainer* InBaseNodeContainer, const TArray<UInterchangeSourceData*>& SourceDatas)
+void UInterchangeDatasmithPipeline::ExecutePipeline(UInterchangeBaseNodeContainer* InBaseNodeContainer, const TArray<UInterchangeSourceData*>& SourceDatas, const FString& ContentBasePath)
 {
 	using namespace UE::DatasmithInterchange;
 
@@ -132,12 +132,12 @@ void UInterchangeDatasmithPipeline::ExecutePipeline(UInterchangeBaseNodeContaine
 		return;
 	}
 
-	auto ExecutePreImportPipelineFunc = [this,SourceDatas](UInterchangePipelineBase* Pipeline)
+	auto ExecutePreImportPipelineFunc = [this, &SourceDatas, &ContentBasePath](UInterchangePipelineBase* Pipeline)
 	{
 		if (Pipeline)
 		{
 			Pipeline->SetResultsContainer(this->Results);
-			Pipeline->ScriptedExecutePipeline(this->BaseNodeContainer, SourceDatas);
+			Pipeline->ScriptedExecutePipeline(this->BaseNodeContainer, SourceDatas, ContentBasePath);
 		}
 	};
 
diff --git a/Engine/Plugins/Experimental/Enterprise/DatasmithInterchange/Source/DatasmithInterchange/Private/InterchangeDatasmithStaticMeshPipeline.cpp b/Engine/Plugins/Experimental/Enterprise/DatasmithInterchange/Source/DatasmithInterchange/Private/InterchangeDatasmithStaticMeshPipeline.cpp
--- a/Engine/Plugins/Experimental/Enterprise/DatasmithInterchange/Source/DatasmithInterchange/Private/InterchangeDatasmithStaticMeshPipeline.cpp
+++ b/Engine/Plugins/Experimental/Enterprise/DatasmithInterchange/Source/DatasmithInterchange/Private/InterchangeDatasmithStaticMeshPipeline.cpp
@@ -11,11 +11,11 @@
 #include "Engine/StaticMesh.h"
 #include "Materials/MaterialInterface.h"
 
-void UInterchangeDatasmithStaticMeshPipeline::ExecutePipeline(UInterchangeBaseNodeContainer* NodeContainer, const TArray<UInterchangeSourceData*>& InSourceDatas)
+void UInterchangeDatasmithStaticMeshPipeline::ExecutePipeline(UInterchangeBaseNodeContainer* NodeContainer, const TArray<UInterchangeSourceData*>& InSourceDatas, const FString& ContentBasePath)
 {
 	using namespace UE::DatasmithInterchange;
 
-	Super::ExecutePipeline(NodeContainer, InSourceDatas);
+	Super::ExecutePipeline(NodeContainer, InSourceDatas, ContentBasePath);
 
 	// Add material factory dependencies for meshes where all slots are filled with the same material
 	for (UInterchangeStaticMeshFactoryNode* MeshFactoryNode : NodeUtils::GetNodes<UInterchangeStaticMeshFactoryNode>(NodeContainer))
diff --git a/Engine/Plugins/Experimental/Enterprise/DatasmithInterchange/Source/DatasmithInterchange/Private/InterchangeDatasmithStaticMeshPipeline.h b/Engine/Plugins/Experimental/Enterprise/DatasmithInterchange/Source/DatasmithInterchange/Private/InterchangeDatasmithStaticMeshPipeline.h
--- a/Engine/Plugins/Experimental/Enterprise/DatasmithInterchange/Source/DatasmithInterchange/Private/InterchangeDatasmithStaticMeshPipeline.h
+++ b/Engine/Plugins/Experimental/Enterprise/DatasmithInterchange/Source/DatasmithInterchange/Private/InterchangeDatasmithStaticMeshPipeline.h
@@ -17,7 +17,7 @@ class DATASMITHINTERCHANGE_API UInterchangeDatasmithStaticMeshPipeline : public
 	GENERATED_BODY()
 
 protected:
-	virtual void ExecutePipeline(UInterchangeBaseNodeContainer* BaseNodeContainer, const TArray<UInterchangeSourceData*>& SourceDatas) override;
+	virtual void ExecutePipeline(UInterchangeBaseNodeContainer* BaseNodeContainer, const TArray<UInterchangeSourceData*>& SourceDatas, const FString& ContentBasePath) override;
 	virtual void ExecutePostImportPipeline(const UInterchangeBaseNodeContainer* BaseNodeContainer, const FString& NodeKey, UObject* CreatedAsset, bool bIsAReimport) override;
 
 	virtual bool CanExecuteOnAnyThread(EInterchangePipelineTask PipelineTask) override
diff --git a/Engine/Plugins/Experimental/Enterprise/DatasmithInterchange/Source/DatasmithInterchange/Private/InterchangeDatasmithTexturePipeline.cpp b/Engine/Plugins/Experimental/Enterprise/DatasmithInterchange/Source/DatasmithInterchange/Private/InterchangeDatasmithTexturePipeline.cpp
--- a/Engine/Plugins/Experimental/Enterprise/DatasmithInterchange/Source/DatasmithInterchange/Private/InterchangeDatasmithTexturePipeline.cpp
+++ b/Engine/Plugins/Experimental/Enterprise/DatasmithInterchange/Source/DatasmithInterchange/Private/InterchangeDatasmithTexturePipeline.cpp
@@ -11,11 +11,11 @@
 #include "InterchangeTextureNode.h"
 #include "IDatasmithSceneElements.h"
 
-void UInterchangeDatasmithTexturePipeline::ExecutePipeline(UInterchangeBaseNodeContainer* InNodeContainer, const TArray<UInterchangeSourceData*>& InSourceDatas)
+void UInterchangeDatasmithTexturePipeline::ExecutePipeline(UInterchangeBaseNodeContainer* InNodeContainer, const TArray<UInterchangeSourceData*>& InSourceDatas, const FString& ContentBasePath)
 {
 	using namespace UE::DatasmithInterchange;
 
-	Super::ExecutePipeline(InNodeContainer, InSourceDatas);
+	Super::ExecutePipeline(InNodeContainer, InSourceDatas, ContentBasePath);
 
 	for (UInterchangeTextureFactoryNode* TextureFactoryNode : NodeUtils::GetNodes<UInterchangeTextureFactoryNode>(InNodeContainer))
 	{
diff --git a/Engine/Plugins/Experimental/Enterprise/DatasmithInterchange/Source/DatasmithInterchange/Private/InterchangeDatasmithTexturePipeline.h b/Engine/Plugins/Experimental/Enterprise/DatasmithInterchange/Source/DatasmithInterchange/Private/InterchangeDatasmithTexturePipeline.h
--- a/Engine/Plugins/Experimental/Enterprise/DatasmithInterchange/Source/DatasmithInterchange/Private/InterchangeDatasmithTexturePipeline.h
+++ b/Engine/Plugins/Experimental/Enterprise/DatasmithInterchange/Source/DatasmithInterchange/Private/InterchangeDatasmithTexturePipeline.h
@@ -20,7 +20,7 @@ class DATASMITHINTERCHANGE_API UInterchangeDatasmithTexturePipeline : public UIn
 	GENERATED_BODY()
 
 public:
-	virtual void ExecutePipeline(UInterchangeBaseNodeContainer* BaseNodeContainer, const TArray<UInterchangeSourceData*>& SourceDatas) override;
+	virtual void ExecutePipeline(UInterchangeBaseNodeContainer* BaseNodeContainer, const TArray<UInterchangeSourceData*>& SourceDatas, const FString& ContentBasePath) override;
 
 private:
 	void PreImportTextureFactoryNode(UInterchangeBaseNodeContainer* BaseNodeContainer, UInterchangeTextureFactoryNode* TextureFactoryNode) const;
diff --git a/Engine/Plugins/Experimental/Enterprise/DatasmithInterchange/Source/DatasmithInterchange/Public/InterchangeDatasmithPipeline.h b/Engine/Plugins/Experimental/Enterprise/DatasmithInterchange/Source/DatasmithInterchange/Public/InterchangeDatasmithPipeline.h
--- a/Engine/Plugins/Experimental/Enterprise/DatasmithInterchange/Source/DatasmithInterchange/Public/InterchangeDatasmithPipeline.h
+++ b/Engine/Plugins/Experimental/Enterprise/DatasmithInterchange/Source/DatasmithInterchange/Public/InterchangeDatasmithPipeline.h
@@ -59,7 +59,7 @@ class DATASMITHINTERCHANGE_API UInterchangeDatasmithPipeline : public UInterchan
 	bool bCreateMaterialReferencesFolders = true;
 
 protected:
-	virtual void ExecutePipeline(UInterchangeBaseNodeContainer* InBaseNodeContainer, const TArray<UInterchangeSourceData*>& SourceDatas) override;
+	virtual void ExecutePipeline(UInterchangeBaseNodeContainer* InBaseNodeContainer, const TArray<UInterchangeSourceData*>& SourceDatas, const FString& ContentBasePath) override;
 	virtual void ExecutePostImportPipeline(const UInterchangeBaseNodeContainer* InBaseNodeContainer, const FString& NodeKey, UObject* CreatedAsset, bool bIsAReimport) override;
 
 	virtual bool CanExecuteOnAnyThread(EInterchangePipelineTask PipelineTask) override
diff --git a/Engine/Plugins/Interchange/Editor/Source/Pipelines/Private/InterchangeGraphInspectorPipeline.cpp b/Engine/Plugins/Interchange/Editor/Source/Pipelines/Private/InterchangeGraphInspectorPipeline.cpp
--- a/Engine/Plugins/Interchange/Editor/Source/Pipelines/Private/InterchangeGraphInspectorPipeline.cpp
+++ b/Engine/Plugins/Interchange/Editor/Source/Pipelines/Private/InterchangeGraphInspectorPipeline.cpp
@@ -11,7 +11,7 @@
 
 #include UE_INLINE_GENERATED_CPP_BY_NAME(InterchangeGraphInspectorPipeline)
 
-void UInterchangeGraphInspectorPipeline::ExecutePipeline(UInterchangeBaseNodeContainer* BaseNodeContainer, const TArray<UInterchangeSourceData*>& SourceDatas)
+void UInterchangeGraphInspectorPipeline::ExecutePipeline(UInterchangeBaseNodeContainer* BaseNodeContainer, const TArray<UInterchangeSourceData*>& SourceDatas, const FString& ContentBasePath)
 {
 	//Do not pop dialog if we are unattended or doing automation
 	if (FApp::IsUnattended() || IsRunningCommandlet())
diff --git a/Engine/Plugins/Interchange/Editor/Source/Pipelines/Private/SInterchangePipelineConfigurationDialog.cpp b/Engine/Plugins/Interchange/Editor/Source/Pipelines/Private/SInterchangePipelineConfigurationDialog.cpp
--- a/Engine/Plugins/Interchange/Editor/Source/Pipelines/Private/SInterchangePipelineConfigurationDialog.cpp
+++ b/Engine/Plugins/Interchange/Editor/Source/Pipelines/Private/SInterchangePipelineConfigurationDialog.cpp
@@ -750,7 +750,7 @@ void SInterchangePipelineConfigurationDialog::RefreshStack(bool bStackSelectionC
 	//When doing a reimport we do not want to save the setting because the context have special default
 	//value for some options like: (Import Materials, Import Textures...).
 	//So when doing a reimport switching stack is like doing a reset to default on all pipelines
-	if (!bReimport)
+	if (!bReimport || !bStackSelectionChange)
 	{
 		SaveAllPipelineSettings();
 	}
@@ -786,7 +786,7 @@ void SInterchangePipelineConfigurationDialog::RefreshStack(bool bStackSelectionC
 			if (UInterchangePipelineBase* GeneratedPipeline = UE::Interchange::GeneratePipelineInstanceInSourceAssetPackage(DefaultPipeline))
 			{
 				GeneratedPipeline->TransferAdjustSettings(DefaultPipeline);
-				if (Stack.StackName != ReimportStackName)
+				if (Stack.StackName != ReimportStackName || !bStackSelectionChange)
 				{
 					//Load the settings for this pipeline
 					GeneratedPipeline->LoadSettings(Stack.StackName);
@@ -895,7 +895,7 @@ FReply SInterchangePipelineConfigurationDialog::OnPreviewImport() const
 	{
 		const TSharedPtr<FInterchangePipelineItemType> PipelineItem = PipelineListViewItems[PipelineIndex];
 		PipelineItem->Pipeline->SetResultsContainer(Results);
-		PipelineItem->Pipeline->ScriptedExecutePipeline(DuplicateBaseNodeContainer, SourceDatas);
+		PipelineItem->Pipeline->ScriptedExecutePipeline(DuplicateBaseNodeContainer, SourceDatas, FString());
 	}
 
 	//Set all node in preview mode so hide the internal data attributes
diff --git a/Engine/Plugins/Interchange/Editor/Source/Pipelines/Public/InterchangeGraphInspectorPipeline.h b/Engine/Plugins/Interchange/Editor/Source/Pipelines/Public/InterchangeGraphInspectorPipeline.h
--- a/Engine/Plugins/Interchange/Editor/Source/Pipelines/Public/InterchangeGraphInspectorPipeline.h
+++ b/Engine/Plugins/Interchange/Editor/Source/Pipelines/Public/InterchangeGraphInspectorPipeline.h
@@ -35,7 +35,7 @@ public:
 
 protected:
 
-	virtual void ExecutePipeline(UInterchangeBaseNodeContainer* BaseNodeContainer, const TArray<UInterchangeSourceData*>& SourceDatas) override;
+	virtual void ExecutePipeline(UInterchangeBaseNodeContainer* BaseNodeContainer, const TArray<UInterchangeSourceData*>& SourceDatas, const FString& ContentBasePath) override;
 
 	//virtual bool ExecutePostImportPipeline(const UInterchangeBaseNodeContainer* BaseNodeContainer, const FName& NodeKey, UObject* CreatedAsset) override;
 
diff --git a/Engine/Plugins/Interchange/Runtime/Source/Import/Public/Fbx/InterchangeFbxTranslator.h b/Engine/Plugins/Interchange/Runtime/Source/Import/Public/Fbx/InterchangeFbxTranslator.h
--- a/Engine/Plugins/Interchange/Runtime/Source/Import/Public/Fbx/InterchangeFbxTranslator.h
+++ b/Engine/Plugins/Interchange/Runtime/Source/Import/Public/Fbx/InterchangeFbxTranslator.h
@@ -20,7 +20,7 @@
 
 /* Fbx translator class support import of texture, material, static mesh, skeletal mesh, */
 
-UCLASS(BlueprintType, Experimental)
+UCLASS(BlueprintType)
 class INTERCHANGEIMPORT_API UInterchangeFbxTranslator : public UInterchangeTranslatorBase
 , public IInterchangeTexturePayloadInterface
 , public IInterchangeMeshPayloadInterface
diff --git a/Engine/Plugins/Interchange/Runtime/Source/Pipelines/Private/InterchangeGenericAnimationPipeline.cpp b/Engine/Plugins/Interchange/Runtime/Source/Pipelines/Private/InterchangeGenericAnimationPipeline.cpp
--- a/Engine/Plugins/Interchange/Runtime/Source/Pipelines/Private/InterchangeGenericAnimationPipeline.cpp
+++ b/Engine/Plugins/Interchange/Runtime/Source/Pipelines/Private/InterchangeGenericAnimationPipeline.cpp
@@ -65,7 +65,7 @@ void UInterchangeGenericAnimationPipeline::AdjustSettingsForContext(EInterchange
 #endif //WITH_EDITOR
 }
 
-void UInterchangeGenericAnimationPipeline::ExecutePipeline(UInterchangeBaseNodeContainer* InBaseNodeContainer, const TArray<UInterchangeSourceData*>& InSourceDatas)
+void UInterchangeGenericAnimationPipeline::ExecutePipeline(UInterchangeBaseNodeContainer* InBaseNodeContainer, const TArray<UInterchangeSourceData*>& InSourceDatas, const FString& ContentBasePath)
 {
 	if (!InBaseNodeContainer)
 	{
diff --git a/Engine/Plugins/Interchange/Runtime/Source/Pipelines/Private/InterchangeGenericAssetsPipeline.cpp b/Engine/Plugins/Interchange/Runtime/Source/Pipelines/Private/InterchangeGenericAssetsPipeline.cpp
--- a/Engine/Plugins/Interchange/Runtime/Source/Pipelines/Private/InterchangeGenericAssetsPipeline.cpp
+++ b/Engine/Plugins/Interchange/Runtime/Source/Pipelines/Private/InterchangeGenericAssetsPipeline.cpp
@@ -864,7 +864,7 @@ TArray<FInterchangeConflictInfo> UInterchangeGenericAssetsPipeline::GetConflictI
 
 	UInterchangeResultsContainer* OldResults = Results;
 	Results = NewObject<UInterchangeResultsContainer>(GetTransientPackage());
-	ExecutePipeline(TransientBaseNodeContainer, SourceDatas);
+	ExecutePipeline(TransientBaseNodeContainer, SourceDatas, FString());
 	Results = OldResults;
 
 	//Create the materials conflict
@@ -999,7 +999,7 @@ void UInterchangeGenericAssetsPipeline::ShowConflictDialog(const FGuid& Conflict
 	}
 }
 
-void UInterchangeGenericAssetsPipeline::ExecutePipeline(UInterchangeBaseNodeContainer* InBaseNodeContainer, const TArray<UInterchangeSourceData*>& InSourceDatas)
+void UInterchangeGenericAssetsPipeline::ExecutePipeline(UInterchangeBaseNodeContainer* InBaseNodeContainer, const TArray<UInterchangeSourceData*>& InSourceDatas, const FString& ContentBasePath)
 {
 	check(CommonSkeletalMeshesAndAnimationsProperties);
 
@@ -1073,15 +1073,15 @@ void UInterchangeGenericAssetsPipeline::ExecutePipeline(UInterchangeBaseNodeCont
 
 	if (MaterialPipeline)
 	{
-		MaterialPipeline->ScriptedExecutePipeline(InBaseNodeContainer, InSourceDatas);
+		MaterialPipeline->ScriptedExecutePipeline(InBaseNodeContainer, InSourceDatas, ContentBasePath);
 	}
 	if (MeshPipeline)
 	{
-		MeshPipeline->ScriptedExecutePipeline(InBaseNodeContainer, InSourceDatas);
+		MeshPipeline->ScriptedExecutePipeline(InBaseNodeContainer, InSourceDatas, ContentBasePath);
 	}
 	if (AnimationPipeline)
 	{
-		AnimationPipeline->ScriptedExecutePipeline(InBaseNodeContainer, InSourceDatas);
+		AnimationPipeline->ScriptedExecutePipeline(InBaseNodeContainer, InSourceDatas, ContentBasePath);
 	}
 
 	ImplementUseSourceNameForAssetOption(InBaseNodeContainer, InSourceDatas);
diff --git a/Engine/Plugins/Interchange/Runtime/Source/Pipelines/Private/InterchangeGenericMaterialPipeline.cpp b/Engine/Plugins/Interchange/Runtime/Source/Pipelines/Private/InterchangeGenericMaterialPipeline.cpp
--- a/Engine/Plugins/Interchange/Runtime/Source/Pipelines/Private/InterchangeGenericMaterialPipeline.cpp
+++ b/Engine/Plugins/Interchange/Runtime/Source/Pipelines/Private/InterchangeGenericMaterialPipeline.cpp
@@ -558,7 +558,7 @@ bool UInterchangeGenericMaterialPipeline::IsPropertyChangeNeedRefresh(const FPro
 
 #endif //WITH_EDITOR
 
-void UInterchangeGenericMaterialPipeline::ExecutePipeline(UInterchangeBaseNodeContainer* InBaseNodeContainer, const TArray<UInterchangeSourceData*>& InSourceDatas)
+void UInterchangeGenericMaterialPipeline::ExecutePipeline(UInterchangeBaseNodeContainer* InBaseNodeContainer, const TArray<UInterchangeSourceData*>& InSourceDatas, const FString& ContentBasePath)
 {
 #if UE_BUILD_DEBUG
 	UE::Interchange::Materials::HashUtils::FMaterialHashDebugData HashDebugData(TEXT("InterchangeDebug/MaterialHashLogs/"));
@@ -592,7 +592,7 @@ void UInterchangeGenericMaterialPipeline::ExecutePipeline(UInterchangeBaseNodeCo
 	
 	if (TexturePipeline)
 	{
-		TexturePipeline->ScriptedExecutePipeline(InBaseNodeContainer, InSourceDatas);
+		TexturePipeline->ScriptedExecutePipeline(InBaseNodeContainer, InSourceDatas, ContentBasePath);
 	}
 
 	//Skip Material import if the toggle is off
diff --git a/Engine/Plugins/Interchange/Runtime/Source/Pipelines/Private/InterchangeGenericMeshPipeline.cpp b/Engine/Plugins/Interchange/Runtime/Source/Pipelines/Private/InterchangeGenericMeshPipeline.cpp
--- a/Engine/Plugins/Interchange/Runtime/Source/Pipelines/Private/InterchangeGenericMeshPipeline.cpp
+++ b/Engine/Plugins/Interchange/Runtime/Source/Pipelines/Private/InterchangeGenericMeshPipeline.cpp
@@ -206,7 +206,7 @@ UInterchangePipelineMeshesUtilities* UInterchangeGenericMeshPipeline::CreateMesh
 	return CreatedPipelineMeshesUtilities;
 }
 
-void UInterchangeGenericMeshPipeline::ExecutePipeline(UInterchangeBaseNodeContainer* InBaseNodeContainer, const TArray<UInterchangeSourceData*>& InSourceDatas)
+void UInterchangeGenericMeshPipeline::ExecutePipeline(UInterchangeBaseNodeContainer* InBaseNodeContainer, const TArray<UInterchangeSourceData*>& InSourceDatas, const FString& ContentBasePath)
 {
 	if (!InBaseNodeContainer)
 	{
diff --git a/Engine/Plugins/Interchange/Runtime/Source/Pipelines/Private/InterchangeGenericScenesPipeline.cpp b/Engine/Plugins/Interchange/Runtime/Source/Pipelines/Private/InterchangeGenericScenesPipeline.cpp
--- a/Engine/Plugins/Interchange/Runtime/Source/Pipelines/Private/InterchangeGenericScenesPipeline.cpp
+++ b/Engine/Plugins/Interchange/Runtime/Source/Pipelines/Private/InterchangeGenericScenesPipeline.cpp
@@ -216,7 +216,7 @@ void UInterchangeGenericLevelPipeline::AdjustSettingsForContext(EInterchangePipe
 	bIsReimportContext = ReimportAsset != nullptr;
 }
 
-void UInterchangeGenericLevelPipeline::ExecutePipeline(UInterchangeBaseNodeContainer* InBaseNodeContainer, const TArray<UInterchangeSourceData*>& InSourceDatas)
+void UInterchangeGenericLevelPipeline::ExecutePipeline(UInterchangeBaseNodeContainer* InBaseNodeContainer, const TArray<UInterchangeSourceData*>& InSourceDatas, const FString& ContentBasePath)
 {
 	if (!InBaseNodeContainer)
 	{
diff --git a/Engine/Plugins/Interchange/Runtime/Source/Pipelines/Private/InterchangeGenericTexturePipeline.cpp b/Engine/Plugins/Interchange/Runtime/Source/Pipelines/Private/InterchangeGenericTexturePipeline.cpp
--- a/Engine/Plugins/Interchange/Runtime/Source/Pipelines/Private/InterchangeGenericTexturePipeline.cpp
+++ b/Engine/Plugins/Interchange/Runtime/Source/Pipelines/Private/InterchangeGenericTexturePipeline.cpp
@@ -156,7 +156,7 @@ void UInterchangeGenericTexturePipeline::FilterPropertiesFromTranslatedData(UInt
 
 #endif //WITH_EDITOR
 
-void UInterchangeGenericTexturePipeline::ExecutePipeline(UInterchangeBaseNodeContainer* InBaseNodeContainer, const TArray<UInterchangeSourceData*>& InSourceDatas)
+void UInterchangeGenericTexturePipeline::ExecutePipeline(UInterchangeBaseNodeContainer* InBaseNodeContainer, const TArray<UInterchangeSourceData*>& InSourceDatas, const FString& ContentBasePath)
 {
 	if (!InBaseNodeContainer)
 	{
diff --git a/Engine/Plugins/Interchange/Runtime/Source/Pipelines/Private/InterchangeMaterialXPipeline.cpp b/Engine/Plugins/Interchange/Runtime/Source/Pipelines/Private/InterchangeMaterialXPipeline.cpp
--- a/Engine/Plugins/Interchange/Runtime/Source/Pipelines/Private/InterchangeMaterialXPipeline.cpp
+++ b/Engine/Plugins/Interchange/Runtime/Source/Pipelines/Private/InterchangeMaterialXPipeline.cpp
@@ -557,9 +557,9 @@ void UInterchangeMaterialXPipeline::AdjustSettingsForContext(EInterchangePipelin
 	}
 }
 
-void UInterchangeMaterialXPipeline::ExecutePipeline(UInterchangeBaseNodeContainer* NodeContainer, const TArray<UInterchangeSourceData*>& InSourceDatas)
+void UInterchangeMaterialXPipeline::ExecutePipeline(UInterchangeBaseNodeContainer* NodeContainer, const TArray<UInterchangeSourceData*>& InSourceDatas, const FString& ContentBasePath)
 {
-	Super::ExecutePipeline(NodeContainer, InSourceDatas);
+	Super::ExecutePipeline(NodeContainer, InSourceDatas, ContentBasePath);
 
 #if WITH_EDITOR
 	auto UpdateMaterialXNodes = [this](const FString& NodeUid, UInterchangeMaterialFunctionCallExpressionFactoryNode* FactorNode)
diff --git a/Engine/Plugins/Interchange/Runtime/Source/Pipelines/Private/InterchangeglTFPipeline.cpp b/Engine/Plugins/Interchange/Runtime/Source/Pipelines/Private/InterchangeglTFPipeline.cpp
--- a/Engine/Plugins/Interchange/Runtime/Source/Pipelines/Private/InterchangeglTFPipeline.cpp
+++ b/Engine/Plugins/Interchange/Runtime/Source/Pipelines/Private/InterchangeglTFPipeline.cpp
@@ -289,9 +289,9 @@ void UInterchangeGLTFPipeline::AdjustSettingsForContext(EInterchangePipelineCont
 	}
 }
 
-void UInterchangeGLTFPipeline::ExecutePipeline(UInterchangeBaseNodeContainer* NodeContainer, const TArray<UInterchangeSourceData*>& InSourceDatas)
+void UInterchangeGLTFPipeline::ExecutePipeline(UInterchangeBaseNodeContainer* NodeContainer, const TArray<UInterchangeSourceData*>& InSourceDatas, const FString& ContentBasePath)
 {
-	Super::ExecutePipeline(NodeContainer, InSourceDatas);
+	Super::ExecutePipeline(NodeContainer, InSourceDatas, ContentBasePath);
 
 	if ((FApp::IsGame() || bUseGLTFMaterialInstanceLibrary) && GLTFPipelineSettings)
 	{
diff --git a/Engine/Plugins/Interchange/Runtime/Source/Pipelines/Public/InterchangeGenericAnimationPipeline.h b/Engine/Plugins/Interchange/Runtime/Source/Pipelines/Public/InterchangeGenericAnimationPipeline.h
--- a/Engine/Plugins/Interchange/Runtime/Source/Pipelines/Public/InterchangeGenericAnimationPipeline.h
+++ b/Engine/Plugins/Interchange/Runtime/Source/Pipelines/Public/InterchangeGenericAnimationPipeline.h
@@ -118,7 +118,7 @@ public:
 	
 protected:
 
-	virtual void ExecutePipeline(UInterchangeBaseNodeContainer* InBaseNodeContainer, const TArray<UInterchangeSourceData*>& InSourceDatas) override;
+	virtual void ExecutePipeline(UInterchangeBaseNodeContainer* InBaseNodeContainer, const TArray<UInterchangeSourceData*>& InSourceDatas, const FString& ContentBasePath) override;
 
 	virtual bool CanExecuteOnAnyThread(EInterchangePipelineTask PipelineTask) override
 	{
diff --git a/Engine/Plugins/Interchange/Runtime/Source/Pipelines/Public/InterchangeGenericAssetsPipeline.h b/Engine/Plugins/Interchange/Runtime/Source/Pipelines/Public/InterchangeGenericAssetsPipeline.h
--- a/Engine/Plugins/Interchange/Runtime/Source/Pipelines/Public/InterchangeGenericAssetsPipeline.h
+++ b/Engine/Plugins/Interchange/Runtime/Source/Pipelines/Public/InterchangeGenericAssetsPipeline.h
@@ -128,7 +128,7 @@ public:
 
 protected:
 
-	virtual void ExecutePipeline(UInterchangeBaseNodeContainer* InBaseNodeContainer, const TArray<UInterchangeSourceData*>& InSourceDatas) override;
+	virtual void ExecutePipeline(UInterchangeBaseNodeContainer* InBaseNodeContainer, const TArray<UInterchangeSourceData*>& InSourceDatas, const FString& ContentBasePath) override;
 	virtual void ExecutePostFactoryPipeline(const UInterchangeBaseNodeContainer* BaseNodeContainer, const FString& NodeKey, UObject* CreatedAsset, bool bIsAReimport) override;
 	virtual void ExecutePostImportPipeline(const UInterchangeBaseNodeContainer* BaseNodeContainer, const FString& NodeKey, UObject* CreatedAsset, bool bIsAReimport) override;
 
diff --git a/Engine/Plugins/Interchange/Runtime/Source/Pipelines/Public/InterchangeGenericMaterialPipeline.h b/Engine/Plugins/Interchange/Runtime/Source/Pipelines/Public/InterchangeGenericMaterialPipeline.h
--- a/Engine/Plugins/Interchange/Runtime/Source/Pipelines/Public/InterchangeGenericMaterialPipeline.h
+++ b/Engine/Plugins/Interchange/Runtime/Source/Pipelines/Public/InterchangeGenericMaterialPipeline.h
@@ -80,7 +80,7 @@ public:
 #endif //WITH_EDITOR
 
 protected:
-	virtual void ExecutePipeline(UInterchangeBaseNodeContainer* InBaseNodeContainer, const TArray<UInterchangeSourceData*>& InSourceDatas) override;
+	virtual void ExecutePipeline(UInterchangeBaseNodeContainer* InBaseNodeContainer, const TArray<UInterchangeSourceData*>& InSourceDatas, const FString& ContentBasePath) override;
 	virtual void ExecutePostFactoryPipeline(const UInterchangeBaseNodeContainer* BaseNodeContainer, const FString& NodeKey, UObject* CreatedAsset, bool bIsAReimport) override;
 	virtual void ExecutePostImportPipeline(const UInterchangeBaseNodeContainer* BaseNodeContainer, const FString& NodeKey, UObject* CreatedAsset, bool bIsAReimport) override;
 	virtual void SetReimportSourceIndex(UClass* ReimportObjectClass, const int32 SourceFileIndex) override;
diff --git a/Engine/Plugins/Interchange/Runtime/Source/Pipelines/Public/InterchangeGenericMeshPipeline.h b/Engine/Plugins/Interchange/Runtime/Source/Pipelines/Public/InterchangeGenericMeshPipeline.h
--- a/Engine/Plugins/Interchange/Runtime/Source/Pipelines/Public/InterchangeGenericMeshPipeline.h
+++ b/Engine/Plugins/Interchange/Runtime/Source/Pipelines/Public/InterchangeGenericMeshPipeline.h
@@ -211,7 +211,7 @@ public:
 		, const bool bAutoDetectType);
 
 protected:
-	virtual void ExecutePipeline(UInterchangeBaseNodeContainer* InBaseNodeContainer, const TArray<UInterchangeSourceData*>& InSourceDatas) override;
+	virtual void ExecutePipeline(UInterchangeBaseNodeContainer* InBaseNodeContainer, const TArray<UInterchangeSourceData*>& InSourceDatas, const FString& ContentBasePath) override;
 
 	virtual void ExecutePostImportPipeline(const UInterchangeBaseNodeContainer* InBaseNodeContainer, const FString& NodeKey, UObject* CreatedAsset, bool bIsAReimport) override;
 
diff --git a/Engine/Plugins/Interchange/Runtime/Source/Pipelines/Public/InterchangeGenericScenesPipeline.h b/Engine/Plugins/Interchange/Runtime/Source/Pipelines/Public/InterchangeGenericScenesPipeline.h
--- a/Engine/Plugins/Interchange/Runtime/Source/Pipelines/Public/InterchangeGenericScenesPipeline.h
+++ b/Engine/Plugins/Interchange/Runtime/Source/Pipelines/Public/InterchangeGenericScenesPipeline.h
@@ -46,7 +46,7 @@ public:
 
 protected:
 
-	virtual void ExecutePipeline(UInterchangeBaseNodeContainer* InBaseNodeContainer, const TArray<UInterchangeSourceData*>& InSourceDatas) override;
+	virtual void ExecutePipeline(UInterchangeBaseNodeContainer* InBaseNodeContainer, const TArray<UInterchangeSourceData*>& InSourceDatas, const FString& ContentBasePath) override;
 	virtual void ExecutePostImportPipeline(const UInterchangeBaseNodeContainer* BaseNodeContainer, const FString& NodeKey, UObject* CreatedAsset, bool bIsAReimport) override;
 
 	virtual bool CanExecuteOnAnyThread(EInterchangePipelineTask PipelineTask) override
diff --git a/Engine/Plugins/Interchange/Runtime/Source/Pipelines/Public/InterchangeGenericTexturePipeline.h b/Engine/Plugins/Interchange/Runtime/Source/Pipelines/Public/InterchangeGenericTexturePipeline.h
--- a/Engine/Plugins/Interchange/Runtime/Source/Pipelines/Public/InterchangeGenericTexturePipeline.h
+++ b/Engine/Plugins/Interchange/Runtime/Source/Pipelines/Public/InterchangeGenericTexturePipeline.h
@@ -65,7 +65,7 @@ public:
 
 public:
 	virtual void AdjustSettingsForContext(EInterchangePipelineContext ImportType, TObjectPtr<UObject> ReimportAsset) override;
-	virtual void ExecutePipeline(UInterchangeBaseNodeContainer* InBaseNodeContainer, const TArray<UInterchangeSourceData*>& InSourceDatas) override;
+	virtual void ExecutePipeline(UInterchangeBaseNodeContainer* InBaseNodeContainer, const TArray<UInterchangeSourceData*>& InSourceDatas, const FString& ContentBasePath) override;
 	virtual void ExecutePostFactoryPipeline(const UInterchangeBaseNodeContainer* BaseNodeContainer, const FString& NodeKey, UObject* CreatedAsset, bool bIsAReimport) override;
 
 #if WITH_EDITOR
diff --git a/Engine/Plugins/Interchange/Runtime/Source/Pipelines/Public/InterchangeMaterialXPipeline.h b/Engine/Plugins/Interchange/Runtime/Source/Pipelines/Public/InterchangeMaterialXPipeline.h
--- a/Engine/Plugins/Interchange/Runtime/Source/Pipelines/Public/InterchangeMaterialXPipeline.h
+++ b/Engine/Plugins/Interchange/Runtime/Source/Pipelines/Public/InterchangeMaterialXPipeline.h
@@ -204,7 +204,7 @@ public:
 
 protected:
 	virtual void AdjustSettingsForContext(EInterchangePipelineContext ImportType, TObjectPtr<UObject> ReimportAsset) override;
-	virtual void ExecutePipeline(UInterchangeBaseNodeContainer* BaseNodeContainer, const TArray<UInterchangeSourceData*>& SourceDatas) override;
+	virtual void ExecutePipeline(UInterchangeBaseNodeContainer* BaseNodeContainer, const TArray<UInterchangeSourceData*>& SourceDatas, const FString& ContentBasePath) override;
 
 	virtual bool CanExecuteOnAnyThread(EInterchangePipelineTask PipelineTask) override
 	{
diff --git a/Engine/Plugins/Interchange/Runtime/Source/Pipelines/Public/InterchangePipelineMeshesUtilities.h b/Engine/Plugins/Interchange/Runtime/Source/Pipelines/Public/InterchangePipelineMeshesUtilities.h
--- a/Engine/Plugins/Interchange/Runtime/Source/Pipelines/Public/InterchangePipelineMeshesUtilities.h
+++ b/Engine/Plugins/Interchange/Runtime/Source/Pipelines/Public/InterchangePipelineMeshesUtilities.h
@@ -131,32 +131,32 @@ struct FInterchangePipelineMeshesUtilitiesContext
 	/**
 	 * Convert static mesh to skeletal mesh
 	 */
-	UPROPERTY(EditAnywhere, Category = "Interchange | Pipeline | MeshesContext")
+	UPROPERTY(EditAnywhere, BlueprintReadWrite, Category = "Interchange | Pipeline | MeshesContext")
 	bool bConvertStaticMeshToSkeletalMesh = false;
 
 	/**
 	 * Convert static mesh to skeletal mesh
 	 */
-	UPROPERTY(EditAnywhere, Category = "Interchange | Pipeline | MeshesContext")
+	UPROPERTY(EditAnywhere, BlueprintReadWrite, Category = "Interchange | Pipeline | MeshesContext")
 	bool bConvertSkeletalMeshToStaticMesh = false;
 
 	/**
 	 * Convert static mesh that has morph target to skeletal mesh
 	 */
-	UPROPERTY(EditAnywhere, Category = "Interchange | Pipeline | MeshesContext")
+	UPROPERTY(EditAnywhere, BlueprintReadWrite, Category = "Interchange | Pipeline | MeshesContext")
 	bool bConvertStaticsWithMorphTargetsToSkeletals = false;
 
 	/**
 	 * If checked, meshes nested in bone hierarchies will be imported instead of being converted to bones. If the mesh are not skinned they will
 	 * be added to skeletal mesh and remove from the static meshes.
 	 */
-	UPROPERTY(EditAnywhere, Category = "Interchange | Pipeline | MeshesContext")
+	UPROPERTY(EditAnywhere, BlueprintReadWrite, Category = "Interchange | Pipeline | MeshesContext")
 	bool bImportMeshesInBoneHierarchy = true;
 
 	/**
 	 * When querying geometry, this flag will not add MeshGeometry if there is a scene node pointing on a geometry.
 	 */
-	UPROPERTY(EditAnywhere, Category = "Interchange | Pipeline | MeshesContext")
+	UPROPERTY(EditAnywhere, BlueprintReadWrite, Category = "Interchange | Pipeline | MeshesContext")
 	bool bQueryGeometryOnlyIfNoInstance = true;
 
 	bool IsStaticMeshInstance(const FInterchangeMeshInstance& MeshInstance, UInterchangeBaseNodeContainer* BaseNodeContainer);
diff --git a/Engine/Plugins/Interchange/Runtime/Source/Pipelines/Public/InterchangeglTFPipeline.h b/Engine/Plugins/Interchange/Runtime/Source/Pipelines/Public/InterchangeglTFPipeline.h
--- a/Engine/Plugins/Interchange/Runtime/Source/Pipelines/Public/InterchangeglTFPipeline.h
+++ b/Engine/Plugins/Interchange/Runtime/Source/Pipelines/Public/InterchangeglTFPipeline.h
@@ -58,7 +58,7 @@ public:
 
 protected:
 	virtual void AdjustSettingsForContext(EInterchangePipelineContext ImportType, TObjectPtr<UObject> ReimportAsset) override;
-	virtual void ExecutePipeline(UInterchangeBaseNodeContainer* BaseNodeContainer, const TArray<UInterchangeSourceData*>& SourceDatas) override;
+	virtual void ExecutePipeline(UInterchangeBaseNodeContainer* BaseNodeContainer, const TArray<UInterchangeSourceData*>& SourceDatas, const FString& ContentBasePath) override;
 
 	virtual bool CanExecuteOnAnyThread(EInterchangePipelineTask PipelineTask) override
 	{
diff --git a/Engine/Source/Runtime/Interchange/Core/Public/InterchangePipelineBase.h b/Engine/Source/Runtime/Interchange/Core/Public/InterchangePipelineBase.h
--- a/Engine/Source/Runtime/Interchange/Core/Public/InterchangePipelineBase.h
+++ b/Engine/Source/Runtime/Interchange/Core/Public/InterchangePipelineBase.h
@@ -59,7 +59,7 @@ struct FInterchangePipelinePropertyStatePerContext
 	GENERATED_BODY()
 
 	/** If true, the property is visible. */
-	UPROPERTY(EditAnywhere, Category = "Property States")
+	UPROPERTY(EditAnywhere, BlueprintReadWrite, Category = "Property States")
 	bool bVisible = true;
 };
 
@@ -91,15 +91,15 @@ struct FInterchangePipelinePropertyStates
 	}
 
 	/** If true, the property is locked. */
-	UPROPERTY(EditAnywhere, Category = "Property States")
+	UPROPERTY(EditAnywhere, BlueprintReadWrite, Category = "Property States")
 	bool bLocked = false;
 
 	/** The property states for the import context */
-	UPROPERTY(EditAnywhere, Category = "Context Properties States")
+	UPROPERTY(EditAnywhere, BlueprintReadWrite, Category = "Context Properties States")
 	FInterchangePipelinePropertyStatePerContext ImportStates;
 
 	/** The property states for the reimport context */
-	UPROPERTY(EditAnywhere, Category = "Context Properties States")
+	UPROPERTY(EditAnywhere, BlueprintReadWrite, Category = "Context Properties States")
 	FInterchangePipelinePropertyStatePerContext ReimportStates;
 };
 
@@ -149,12 +149,12 @@ public:
 	 * @note - the FTaskPipeline is calling this function not the virtual one that is call by the default implementation.
 	 */
 	UFUNCTION(BlueprintNativeEvent, BlueprintCallable, Category = "Interchange | Pipeline")
-	INTERCHANGECORE_API void ScriptedExecutePipeline(UInterchangeBaseNodeContainer* BaseNodeContainer, const TArray<UInterchangeSourceData*>& SourceDatas);
+	INTERCHANGECORE_API void ScriptedExecutePipeline(UInterchangeBaseNodeContainer* BaseNodeContainer, const TArray<UInterchangeSourceData*>& SourceDatas, const FString& ContentBasePath);
 	/** The default implementation (call if the blueprint do not have any implementation) will call the virtual ExecutePipeline */
-	void ScriptedExecutePipeline_Implementation(UInterchangeBaseNodeContainer* BaseNodeContainer, const TArray<UInterchangeSourceData*>& SourceDatas)
+	void ScriptedExecutePipeline_Implementation(UInterchangeBaseNodeContainer* BaseNodeContainer, const TArray<UInterchangeSourceData*>& SourceDatas, const FString& ContentBasePath)
 	{
 		//By default we call the virtual import pipeline execution
-		ExecutePipeline(BaseNodeContainer, SourceDatas);
+		ExecutePipeline(BaseNodeContainer, SourceDatas, ContentBasePath);
 	}
 
 	/**
@@ -410,21 +410,10 @@ public:
 
 protected:
 
-	UE_DEPRECATED(5.2, "This function is replace by ExecutePipeline.")
-	virtual void ExecutePreImportPipeline(UInterchangeBaseNodeContainer* BaseNodeContainer, const TArray<UInterchangeSourceData*>& SourceDatas)
+	virtual void ExecutePipeline(UInterchangeBaseNodeContainer* BaseNodeContainer, const TArray<UInterchangeSourceData*>& SourceDatas, const FString& ContentBasePath)
 	{
 	}
 
-	/**
-	 * This function is called after the translation is done. It should create any factory node representing unreal assets to be imported by the factories
-	 */
-	virtual void ExecutePipeline(UInterchangeBaseNodeContainer* BaseNodeContainer, const TArray<UInterchangeSourceData*>& SourceDatas)
-	{
-		PRAGMA_DISABLE_DEPRECATION_WARNINGS
-		ExecutePreImportPipeline(BaseNodeContainer, SourceDatas);
-		PRAGMA_ENABLE_DEPRECATION_WARNINGS
-	}
-
 	/**
 	 * This function is called after the factory create an asset, but before PostEditChange is called.
 	 * It is always call on the game thread, code inside this function should not wait on other thread. (do not stall the game thread)
diff --git a/Engine/Source/Runtime/Interchange/Engine/Private/InterchangeManager.cpp b/Engine/Source/Runtime/Interchange/Engine/Private/InterchangeManager.cpp
--- a/Engine/Source/Runtime/Interchange/Engine/Private/InterchangeManager.cpp
+++ b/Engine/Source/Runtime/Interchange/Engine/Private/InterchangeManager.cpp
@@ -1360,12 +1360,12 @@ void UInterchangeManager::StartQueuedTasks(bool bCancelAllTasks /*= false*/)
 
 			if (GraphParsingPrerequistes.Num() > 0)
 			{
-				QueuedTaskData.AsyncHelper->ParsingTask = TGraphTask<UE::Interchange::FTaskParsing>::CreateTask(&GraphParsingPrerequistes).ConstructAndDispatchWhenReady(this, QueuedTaskData.PackageBasePath, WeakAsyncHelper);
+				QueuedTaskData.AsyncHelper->ParsingTask = TGraphTask<UE::Interchange::FTaskParsing>::CreateTask(&GraphParsingPrerequistes).ConstructAndDispatchWhenReady(this, WeakAsyncHelper);
 			}
 			else
 			{
 				//Fallback on the translator pipeline prerequisites (translator must be done if there is no pipeline)
-				QueuedTaskData.AsyncHelper->ParsingTask = TGraphTask<UE::Interchange::FTaskParsing>::CreateTask(&PipelinePrerequistes).ConstructAndDispatchWhenReady(this, QueuedTaskData.PackageBasePath, WeakAsyncHelper);
+				QueuedTaskData.AsyncHelper->ParsingTask = TGraphTask<UE::Interchange::FTaskParsing>::CreateTask(&PipelinePrerequistes).ConstructAndDispatchWhenReady(this, WeakAsyncHelper);
 			}
 
 			//The graph parsing task will create the FCreateAssetTask that will run after them, the FAssetImportTask will call the appropriate Post asset import pipeline when the asset is completed
@@ -1500,14 +1500,14 @@ UInterchangeManager::ImportInternal(const FString& ContentPath, const UInterchan
 	}
 
 	UInterchangeAssetImportData* OriginalAssetImportData = UInterchangeAssetImportData::GetFromObject(ImportAssetParameters.ReimportAsset);
-	FString PackageBasePath = ContentPath;
+	FString ContentBasePath = ContentPath;
 	if (!ImportAssetParameters.ReimportAsset)
 	{
-		UE::Interchange::SanitizeObjectPath(PackageBasePath);
+		UE::Interchange::SanitizeObjectPath(ContentBasePath);
 	}
 	else
 	{
-		PackageBasePath = FPaths::GetPath(ImportAssetParameters.ReimportAsset->GetPathName());
+		ContentBasePath = FPaths::GetPath(ImportAssetParameters.ReimportAsset->GetPathName());
 	}
 
 	const bool bIsReimport = OriginalAssetImportData && OriginalAssetImportData->GetPipelines().Num() > 0;
@@ -1868,10 +1868,10 @@ UInterchangeManager::ImportInternal(const FString& ContentPath, const UInterchan
 		AsyncHelper->CleanUp();
 	}
 
+	AsyncHelper->ContentBasePath = ContentBasePath;
 	//Queue the task cancel or not, we need to return a valid asset import result
 	FQueuedTaskData QueuedTaskData;
 	QueuedTaskData.AsyncHelper = AsyncHelper;
-	QueuedTaskData.PackageBasePath = PackageBasePath;
 	QueuedTaskData.TranslatorClass = AsyncTranslator->GetClass();
 
 	//If we cancel or abort the task we want to avoid putting it in the NonParallelTranslatorQueueTasks (the locks will not be release if the task doesn't start)
diff --git a/Engine/Source/Runtime/Interchange/Engine/Private/Tasks/InterchangeTaskParsing.cpp b/Engine/Source/Runtime/Interchange/Engine/Private/Tasks/InterchangeTaskParsing.cpp
--- a/Engine/Source/Runtime/Interchange/Engine/Private/Tasks/InterchangeTaskParsing.cpp
+++ b/Engine/Source/Runtime/Interchange/Engine/Private/Tasks/InterchangeTaskParsing.cpp
@@ -250,14 +250,14 @@ void UE::Interchange::FTaskParsing::DoTask(ENamedThreads::Type CurrentThread, co
 		{
 			return AsyncHelper->SceneTasks.Add_GetRef(
 				TGraphTask<FTaskCreateSceneObjects>::CreateTask(&(TaskData.Prerequisites))
-				.ConstructAndDispatchWhenReady(PackageBasePath, SourceIndex, WeakAsyncHelper, TaskData.Nodes, FactoryClass));
+				.ConstructAndDispatchWhenReady(AsyncHelper->ContentBasePath, SourceIndex, WeakAsyncHelper, TaskData.Nodes, FactoryClass));
 		}
 		else
 		{
 			FString PackageSubPath;
 			FactoryNode->GetCustomSubPath(PackageSubPath);
 
-			FString AssetFullPath = FPaths::Combine(PackageBasePath, PackageSubPath, FactoryNode->GetAssetName());
+			FString AssetFullPath = FPaths::Combine(AsyncHelper->ContentBasePath, PackageSubPath, FactoryNode->GetAssetName());
 
 			//Make sure there is no duplicate name full path
 			uint32 NameIndex = 1;
@@ -294,19 +294,19 @@ void UE::Interchange::FTaskParsing::DoTask(ENamedThreads::Type CurrentThread, co
 			{
 				FGraphEventArray ImportObjectTasksPrerequistes;
 				int32 BeginImportObjectTaskIndex = AsyncHelper->BeginImportObjectTasks.Add(
-					TGraphTask<FTaskImportObject_GameThread>::CreateTask(&(TaskData.Prerequisites)).ConstructAndDispatchWhenReady(PackageBasePath, SourceIndex, WeakAsyncHelper, FactoryNode, FactoryClass)
+					TGraphTask<FTaskImportObject_GameThread>::CreateTask(&(TaskData.Prerequisites)).ConstructAndDispatchWhenReady(AsyncHelper->ContentBasePath, SourceIndex, WeakAsyncHelper, FactoryNode, FactoryClass)
 				);
 				ImportObjectTasksPrerequistes.Add(AsyncHelper->BeginImportObjectTasks[BeginImportObjectTaskIndex]);
 
 				int32 ImportObjectTaskIndex = AsyncHelper->ImportObjectTasks.Add(
-					TGraphTask<FTaskImportObject_Async>::CreateTask(&(ImportObjectTasksPrerequistes)).ConstructAndDispatchWhenReady(PackageBasePath, SourceIndex, WeakAsyncHelper, FactoryNode)
+					TGraphTask<FTaskImportObject_Async>::CreateTask(&(ImportObjectTasksPrerequistes)).ConstructAndDispatchWhenReady(AsyncHelper->ContentBasePath, SourceIndex, WeakAsyncHelper, FactoryNode)
 				);
 
 				FGraphEventArray FinalizeImportObjectTasksPrerequistes;
 				FinalizeImportObjectTasksPrerequistes.Add(AsyncHelper->ImportObjectTasks[ImportObjectTaskIndex]);
 
 				int32 FinalizeCreateTaskIndex = AsyncHelper->FinalizeImportObjectTasks.Add(
-					TGraphTask<FTaskImportObjectFinalize_GameThread>::CreateTask(&(FinalizeImportObjectTasksPrerequistes)).ConstructAndDispatchWhenReady(PackageBasePath, SourceIndex, WeakAsyncHelper, FactoryNode)
+					TGraphTask<FTaskImportObjectFinalize_GameThread>::CreateTask(&(FinalizeImportObjectTasksPrerequistes)).ConstructAndDispatchWhenReady(AsyncHelper->ContentBasePath, SourceIndex, WeakAsyncHelper, FactoryNode)
 				);
 
 				CreatedTasksAssetNames.Add(AssetFullPath);
diff --git a/Engine/Source/Runtime/Interchange/Engine/Private/Tasks/InterchangeTaskParsing.h b/Engine/Source/Runtime/Interchange/Engine/Private/Tasks/InterchangeTaskParsing.h
--- a/Engine/Source/Runtime/Interchange/Engine/Private/Tasks/InterchangeTaskParsing.h
+++ b/Engine/Source/Runtime/Interchange/Engine/Private/Tasks/InterchangeTaskParsing.h
@@ -18,12 +18,10 @@ namespace UE
 		private:
 			UInterchangeManager* InterchangeManager;
 			TWeakPtr<FImportAsyncHelper, ESPMode::ThreadSafe> WeakAsyncHelper;
-			FString PackageBasePath;
 		public:
-			FTaskParsing(UInterchangeManager* InInterchangeManager, const FString& InPackageBasePath, TWeakPtr<FImportAsyncHelper, ESPMode::ThreadSafe> InAsyncHelper)
+			FTaskParsing(UInterchangeManager* InInterchangeManager, TWeakPtr<FImportAsyncHelper, ESPMode::ThreadSafe> InAsyncHelper)
 				: InterchangeManager(InInterchangeManager)
 				, WeakAsyncHelper(InAsyncHelper)
-				, PackageBasePath(InPackageBasePath)
 			{
 				check(InterchangeManager);
 			}
diff --git a/Engine/Source/Runtime/Interchange/Engine/Private/Tasks/InterchangeTaskPipeline.cpp b/Engine/Source/Runtime/Interchange/Engine/Private/Tasks/InterchangeTaskPipeline.cpp
--- a/Engine/Source/Runtime/Interchange/Engine/Private/Tasks/InterchangeTaskPipeline.cpp
+++ b/Engine/Source/Runtime/Interchange/Engine/Private/Tasks/InterchangeTaskPipeline.cpp
@@ -50,7 +50,7 @@ void UE::Interchange::FTaskPipeline::DoTask(ENamedThreads::Type CurrentThread, c
 
 			if (ensure(AsyncHelper->BaseNodeContainers[GraphIndex].IsValid()))
 			{
-				Pipeline->ScriptedExecutePipeline(AsyncHelper->BaseNodeContainers[GraphIndex].Get(), AsyncHelper->SourceDatas);
+				Pipeline->ScriptedExecutePipeline(AsyncHelper->BaseNodeContainers[GraphIndex].Get(), AsyncHelper->SourceDatas, AsyncHelper->ContentBasePath);
 			}
 		}
 	}
diff --git a/Engine/Source/Runtime/Interchange/Engine/Public/InterchangeManager.h b/Engine/Source/Runtime/Interchange/Engine/Public/InterchangeManager.h
--- a/Engine/Source/Runtime/Interchange/Engine/Public/InterchangeManager.h
+++ b/Engine/Source/Runtime/Interchange/Engine/Public/InterchangeManager.h
@@ -206,6 +206,9 @@ namespace UE
 			/** Unique id for this asynchelper. */
 			int32 UniqueId;
 
+			//The content base path where to import
+			FString ContentBasePath;
+
 			//The following Arrays are per source data
 			TArray<TStrongObjectPtr<UInterchangeBaseNodeContainer>> BaseNodeContainers;
 			TArray<TObjectPtr<UInterchangeSourceData>> SourceDatas;
@@ -675,7 +678,6 @@ protected:
 private:
 	struct FQueuedTaskData
 	{
-		FString PackageBasePath;
 		TSharedPtr<UE::Interchange::FImportAsyncHelper, ESPMode::ThreadSafe> AsyncHelper;
 		UClass* TranslatorClass = nullptr;
 	};
diff --git a/Engine/Source/Runtime/StaticMeshDescription/Private/StaticMeshOperations.cpp b/Engine/Source/Runtime/StaticMeshDescription/Private/StaticMeshOperations.cpp
--- a/Engine/Source/Runtime/StaticMeshDescription/Private/StaticMeshOperations.cpp
+++ b/Engine/Source/Runtime/StaticMeshDescription/Private/StaticMeshOperations.cpp
@@ -1343,10 +1343,6 @@ void FStaticMeshOperations::ComputeTangentsAndNormals(FMeshDescription& MeshDesc
 	// Second normal(F to A) = Normalize(F+G+H+A)
 	// We found the connected edge using the triangle that share edges
 
-	// @todo: provide an option to weight each contributing polygon normal according to the size of
-	// the angle it makes with the vertex being calculated. This means that triangulated faces whose
-	// internal edge meets the vertex doesn't get undue extra weight.
-
 	struct FTriangleCornerData
 	{
 		FVertexInstanceID VertexInstanceID;
