Interchange: Fix untagged memories

#rnx

diff --git a/Engine/Plugins/Interchange/Editor/Source/Pipelines/Private/SInterchangePipelineConfigurationDialog.cpp b/Engine/Plugins/Interchange/Editor/Source/Pipelines/Private/SInterchangePipelineConfigurationDialog.cpp
--- a/Engine/Plugins/Interchange/Editor/Source/Pipelines/Private/SInterchangePipelineConfigurationDialog.cpp
+++ b/Engine/Plugins/Interchange/Editor/Source/Pipelines/Private/SInterchangePipelineConfigurationDialog.cpp
@@ -57,6 +57,8 @@ void SInterchangePipelineItem::Construct(
 	const TSharedRef<STableViewBase>& OwnerTable,
 	TSharedPtr<FInterchangePipelineItemType> InPipelineElement)
 {
+	LLM_SCOPE_BYNAME(TEXT("Interchange"));
+
 	PipelineElement = InPipelineElement;
 	TObjectPtr<UInterchangePipelineBase> PipelineElementPtr = PipelineElement->Pipeline;
 	check(PipelineElementPtr.Get());
@@ -373,6 +375,7 @@ TSharedRef<SBox> SInterchangePipelineConfigurationDialog::SpawnPipelineConfigura
 
 void SInterchangePipelineConfigurationDialog::Construct(const FArguments& InArgs)
 {
+	LLM_SCOPE_BYNAME(TEXT("Interchange"));
 	//Make sure there is a valid default value
 
 	OwnerWindow = InArgs._OwnerWindow;
@@ -791,6 +794,8 @@ FReply SInterchangePipelineConfigurationDialog::OnKeyDown(const FGeometry& MyGeo
 
 void SInterchangePipelineConfigurationDialog::RefreshStack(bool bStackSelectionChange)
 {
+	LLM_SCOPE_BYNAME(TEXT("Interchange"));
+
 	//Save current stack settings, we want the same settings when we will go back to the same stack
 	//When doing a reimport we do not want to save the setting because the context have special default
 	//value for some options like: (Import Materials, Import Textures...).
diff --git a/Engine/Plugins/Interchange/Runtime/Source/Pipelines/Private/InterchangeGenericSkeletalMeshPipeline.cpp b/Engine/Plugins/Interchange/Runtime/Source/Pipelines/Private/InterchangeGenericSkeletalMeshPipeline.cpp
--- a/Engine/Plugins/Interchange/Runtime/Source/Pipelines/Private/InterchangeGenericSkeletalMeshPipeline.cpp
+++ b/Engine/Plugins/Interchange/Runtime/Source/Pipelines/Private/InterchangeGenericSkeletalMeshPipeline.cpp
@@ -72,6 +72,7 @@ namespace UE::Interchange::SkeletalMeshGenericPipeline
 
 void UInterchangeGenericMeshPipeline::ExecutePreImportPipelineSkeletalMesh()
 {
+	LLM_SCOPE_BYNAME(TEXT("Interchange"));
 	check(CommonMeshesProperties.IsValid());
 	if (!bImportSkeletalMeshes)
 	{
@@ -567,6 +568,7 @@ void UInterchangeGenericMeshPipeline::PostImportSkeletalMesh(UObject* CreatedAss
 void UInterchangeGenericMeshPipeline::PostImportPhysicsAssetImport(UObject* CreatedAsset, const UInterchangeFactoryBaseNode* FactoryNode)
 {
 	TRACE_CPUPROFILER_EVENT_SCOPE(UInterchangeGenericMeshPipeline::PostImportPhysicsAssetImport);
+	LLM_SCOPE_BYNAME(TEXT("Interchange"));
 #if WITH_EDITOR
 	if (!bCreatePhysicsAsset || !BaseNodeContainer)
 	{
diff --git a/Engine/Source/Runtime/Interchange/Core/Public/Nodes/InterchangeBaseNodeUtilities.h b/Engine/Source/Runtime/Interchange/Core/Public/Nodes/InterchangeBaseNodeUtilities.h
--- a/Engine/Source/Runtime/Interchange/Core/Public/Nodes/InterchangeBaseNodeUtilities.h
+++ b/Engine/Source/Runtime/Interchange/Core/Public/Nodes/InterchangeBaseNodeUtilities.h
@@ -25,12 +25,16 @@ class TArrayAttributeHelper
 public:
 	~TArrayAttributeHelper()
 	{
+		LLM_SCOPE_BYNAME(TEXT("Interchange"));
+
 		Attributes = nullptr;
 		KeyCount = NAME_None;
 	}
 
 	void Initialize(const TSharedPtr<FAttributeStorage, ESPMode::ThreadSafe>& InAttributes, const FString& BaseKeyName)
 	{
+		LLM_SCOPE_BYNAME(TEXT("Interchange"));
+
 		check(InAttributes.IsValid());
 		Attributes = InAttributes;
 		KeyCount = BaseKeyName;
@@ -298,6 +302,8 @@ class TMapAttributeHelper
 public:
 	void Initialize(const TSharedRef<FAttributeStorage, ESPMode::ThreadSafe>& InAttributes, const FString& BaseKeyName)
 	{
+		LLM_SCOPE_BYNAME(TEXT("Interchange"));
+
 		Attributes = InAttributes;
 		FString BaseTryName = BaseKeyName;
 		FAttributeKey KeyCountKey(MoveTemp(BaseTryName));
@@ -484,6 +490,8 @@ public:
 
 	void RebuildCache()
 	{
+		LLM_SCOPE_BYNAME(TEXT("Interchange"));
+
 		TSharedPtr<FAttributeStorage, ESPMode::ThreadSafe> AttributesPtr = Attributes.Pin();
 		if (AttributesPtr.IsValid() && KeyCountHandle.IsValid())
 		{
diff --git a/Engine/Source/Runtime/Interchange/Engine/Private/InterchangeManager.cpp b/Engine/Source/Runtime/Interchange/Engine/Private/InterchangeManager.cpp
--- a/Engine/Source/Runtime/Interchange/Engine/Private/InterchangeManager.cpp
+++ b/Engine/Source/Runtime/Interchange/Engine/Private/InterchangeManager.cpp
@@ -1199,6 +1199,9 @@ bool UInterchangeManager::CanReimport(const UObject* Object, TArray<FString>& Ou
 void UInterchangeManager::StartQueuedTasks(bool bCancelAllTasks /*= false*/)
 {
 	TRACE_CPUPROFILER_EVENT_SCOPE(UInterchangeManager::StartQueuedTasks)
+		
+	LLM_SCOPE_BYNAME(TEXT("Interchange"));
+
 	ensure(IsInterchangeImportEnabled());
 	if (!ensure(IsInGameThread()))
 	{
@@ -1412,6 +1415,9 @@ TTuple<UE::Interchange::FAssetImportResultRef, UE::Interchange::FSceneImportResu
 UInterchangeManager::ImportInternal(const FString& ContentPath, const UInterchangeSourceData* SourceData, const FImportAssetParameters& ImportAssetParameters, const UE::Interchange::EImportType ImportType)
 {
 	TRACE_CPUPROFILER_EVENT_SCOPE(UInterchangeManager::ImportInternal)
+	
+	LLM_SCOPE_BYNAME(TEXT("Interchange"));
+
 	ensure(IsInterchangeImportEnabled());
 	check(IsInGameThread());
 	static int32 GeneratedUniqueID = 0;
@@ -1620,6 +1626,7 @@ UInterchangeManager::ImportInternal(const FString& ContentPath, const UInterchan
 
 		auto TranslateSourceFile = [&AsyncHelper]()
 		{
+			LLM_SCOPE_BYNAME(TEXT("Interchange"));
 			FScopedSlowTask Progress(2.f, NSLOCTEXT("InterchangeManager", "TranslatingSourceFile...", "Translating source file..."));
 			Progress.MakeDialog();
 			Progress.EnterProgressFrame(1.f);
diff --git a/Engine/Source/Runtime/Interchange/Engine/Private/Tasks/InterchangeTaskCompletion.cpp b/Engine/Source/Runtime/Interchange/Engine/Private/Tasks/InterchangeTaskCompletion.cpp
--- a/Engine/Source/Runtime/Interchange/Engine/Private/Tasks/InterchangeTaskCompletion.cpp
+++ b/Engine/Source/Runtime/Interchange/Engine/Private/Tasks/InterchangeTaskCompletion.cpp
@@ -26,6 +26,9 @@ void UE::Interchange::FTaskPreCompletion::DoTask(ENamedThreads::Type CurrentThre
 #if INTERCHANGE_TRACE_ASYNCHRONOUS_TASK_ENABLED
 	INTERCHANGE_TRACE_ASYNCHRONOUS_TASK(PreCompletion)
 #endif
+	
+	LLM_SCOPE_BYNAME(TEXT("Interchange"));
+
 	TSharedPtr<FImportAsyncHelper, ESPMode::ThreadSafe> AsyncHelper = WeakAsyncHelper.Pin();
 	check(AsyncHelper.IsValid());
 
@@ -35,6 +38,7 @@ void UE::Interchange::FTaskPreCompletion::DoTask(ENamedThreads::Type CurrentThre
 
 	auto IterationCallback = [&AsyncHelper, &Results, &bIsAsset](int32 SourceIndex, const TArray<FImportAsyncHelper::FImportedObjectInfo>& ImportedObjects)
 	{
+		LLM_SCOPE_BYNAME(TEXT("Interchange"));
 		//Verify if the task was cancel
 		if (AsyncHelper->bCancel)
 		{
@@ -191,6 +195,9 @@ void UE::Interchange::FTaskCompletion::DoTask(ENamedThreads::Type CurrentThread,
 #if INTERCHANGE_TRACE_ASYNCHRONOUS_TASK_ENABLED
 	INTERCHANGE_TRACE_ASYNCHRONOUS_TASK(Completion)
 #endif
+
+	LLM_SCOPE_BYNAME(TEXT("Interchange"));
+
 	TSharedPtr<FImportAsyncHelper, ESPMode::ThreadSafe> AsyncHelper = WeakAsyncHelper.Pin();
 	check(AsyncHelper.IsValid());
 
diff --git a/Engine/Source/Runtime/Interchange/Engine/Private/Tasks/InterchangeTaskImportObject.cpp b/Engine/Source/Runtime/Interchange/Engine/Private/Tasks/InterchangeTaskImportObject.cpp
--- a/Engine/Source/Runtime/Interchange/Engine/Private/Tasks/InterchangeTaskImportObject.cpp
+++ b/Engine/Source/Runtime/Interchange/Engine/Private/Tasks/InterchangeTaskImportObject.cpp
@@ -249,6 +249,9 @@ void UE::Interchange::FTaskImportObject_GameThread::DoTask(ENamedThreads::Type C
 #if INTERCHANGE_TRACE_ASYNCHRONOUS_TASK_ENABLED
 	INTERCHANGE_TRACE_ASYNCHRONOUS_TASK(TaskImportObject_GameThread)
 #endif
+
+	LLM_SCOPE_BYNAME(TEXT("Interchange"));
+
 	using namespace UE::Interchange;
 
 	TSharedPtr<FImportAsyncHelper, ESPMode::ThreadSafe> AsyncHelper = WeakAsyncHelper.Pin();
@@ -525,6 +528,9 @@ void UE::Interchange::FTaskImportObject_Async::DoTask(ENamedThreads::Type Curren
 #if INTERCHANGE_TRACE_ASYNCHRONOUS_TASK_ENABLED
 	INTERCHANGE_TRACE_ASYNCHRONOUS_TASK(TaskImportObject_Async)
 #endif
+
+	LLM_SCOPE_BYNAME(TEXT("Interchange"));
+
 	using namespace UE::Interchange;
 
 	TSharedPtr<FImportAsyncHelper, ESPMode::ThreadSafe> AsyncHelper = WeakAsyncHelper.Pin();
@@ -548,6 +554,7 @@ void UE::Interchange::FTaskImportObject_Async::DoTask(ENamedThreads::Type Curren
 		, PackageBasePath
 		, [&Factory](UInterchangeFactoryBase::FImportAssetObjectParams& ImportAssetObjectParams)
 		{
+			LLM_SCOPE_BYNAME(TEXT("Interchange"));
 			return Factory->ImportAsset_Async(ImportAssetObjectParams);
 		});
 }
@@ -557,6 +564,9 @@ void UE::Interchange::FTaskImportObjectFinalize_GameThread::DoTask(ENamedThreads
 #if INTERCHANGE_TRACE_ASYNCHRONOUS_TASK_ENABLED
 	INTERCHANGE_TRACE_ASYNCHRONOUS_TASK(TaskImportObjectFinalize_GameThread)
 #endif
+
+	LLM_SCOPE_BYNAME(TEXT("Interchange"));
+
 	using namespace UE::Interchange;
 
 	TSharedPtr<FImportAsyncHelper, ESPMode::ThreadSafe> AsyncHelper = WeakAsyncHelper.Pin();
@@ -580,6 +590,7 @@ void UE::Interchange::FTaskImportObjectFinalize_GameThread::DoTask(ENamedThreads
 		, PackageBasePath
 		, [&Factory](UInterchangeFactoryBase::FImportAssetObjectParams& ImportAssetObjectParams)
 		{
+			LLM_SCOPE_BYNAME(TEXT("Interchange"));
 			return Factory->EndImportAsset_GameThread(ImportAssetObjectParams);
 		});
 
diff --git a/Engine/Source/Runtime/Interchange/Engine/Private/Tasks/InterchangeTaskParsing.cpp b/Engine/Source/Runtime/Interchange/Engine/Private/Tasks/InterchangeTaskParsing.cpp
--- a/Engine/Source/Runtime/Interchange/Engine/Private/Tasks/InterchangeTaskParsing.cpp
+++ b/Engine/Source/Runtime/Interchange/Engine/Private/Tasks/InterchangeTaskParsing.cpp
@@ -93,6 +93,8 @@ void UE::Interchange::FTaskParsing::DoTask(ENamedThreads::Type CurrentThread, co
 #if INTERCHANGE_TRACE_ASYNCHRONOUS_TASK_ENABLED
 	INTERCHANGE_TRACE_ASYNCHRONOUS_TASK(ParsingGraph)
 #endif
+	LLM_SCOPE_BYNAME(TEXT("Interchange"));
+
 	FGCScopeGuard GCScopeGuard;
 
 	TSharedPtr<FImportAsyncHelper, ESPMode::ThreadSafe> AsyncHelper = WeakAsyncHelper.Pin();
@@ -240,6 +242,7 @@ void UE::Interchange::FTaskParsing::DoTask(ENamedThreads::Type CurrentThread, co
 	TSet<FString> CreatedTasksAssetNames; // Tracks for which asset name we have created a task so that we don't have 2 tasks for the same asset name
 	TFunction<FGraphEventRef(FTaskData&)> CreateTasksFromData = [this, &AsyncHelper, &RenameAssets, &CreatedTasksAssetNames](FTaskData& TaskData)
 	{
+		LLM_SCOPE_BYNAME(TEXT("Interchange"));
 		check(TaskData.Nodes.Num() == 1); //We expect 1 node per asset task
 
 		const int32 SourceIndex = TaskData.SourceIndex;
@@ -292,6 +295,7 @@ void UE::Interchange::FTaskParsing::DoTask(ENamedThreads::Type CurrentThread, co
 			if (ensureMsgf(!CreatedTasksAssetNames.Contains(AssetFullPath),
 				TEXT("Found multiple task data with the same asset name (%s). Only one will be executed."), *AssetFullPath))
 			{
+				LLM_SCOPE_BYNAME(TEXT("Interchange"));
 				FGraphEventArray ImportObjectTasksPrerequistes;
 				int32 BeginImportObjectTaskIndex = AsyncHelper->BeginImportObjectTasks.Add(
 					TGraphTask<FTaskImportObject_GameThread>::CreateTask(&(TaskData.Prerequisites)).ConstructAndDispatchWhenReady(AsyncHelper->ContentBasePath, SourceIndex, WeakAsyncHelper, FactoryNode, FactoryClass)
diff --git a/Engine/Source/Runtime/Interchange/Engine/Private/Tasks/InterchangeTaskPipeline.cpp b/Engine/Source/Runtime/Interchange/Engine/Private/Tasks/InterchangeTaskPipeline.cpp
--- a/Engine/Source/Runtime/Interchange/Engine/Private/Tasks/InterchangeTaskPipeline.cpp
+++ b/Engine/Source/Runtime/Interchange/Engine/Private/Tasks/InterchangeTaskPipeline.cpp
@@ -27,6 +27,8 @@ void UE::Interchange::FTaskPipeline::DoTask(ENamedThreads::Type CurrentThread, c
 	INTERCHANGE_TRACE_ASYNCHRONOUS_TASK(PipelinePreImport)
 #endif
 
+	LLM_SCOPE_BYNAME(TEXT("Interchange"));
+
 	TOptional<FGCScopeGuard> GCScopeGuard;
 	if (!IsInGameThread())
 	{
@@ -62,6 +64,8 @@ void UE::Interchange::FTaskWaitAssetCompilation::DoTask(ENamedThreads::Type Curr
 #if INTERCHANGE_TRACE_ASYNCHRONOUS_TASK_ENABLED
 		INTERCHANGE_TRACE_ASYNCHRONOUS_TASK(WaitAssetCompilation)
 #endif
+	
+	LLM_SCOPE_BYNAME(TEXT("Interchange"));
 
 #if WITH_EDITOR
 
@@ -93,12 +97,14 @@ void UE::Interchange::FTaskWaitAssetCompilation::DoTask(ENamedThreads::Type Curr
 	//Make sure all assets compilation are done before calling the pipeline post import task, let other thread execute if assets are not compile yet and wait 50ms before a new query
 	FPlatformProcess::ConditionalSleep([&ImportedObjects]()
 		{
+			LLM_SCOPE_BYNAME(TEXT("Interchange"));
 			//Compilation status cannot be ask in async thread, query the compile status on the main thread with a small fast function
 			//This ensure we dont stall the main thread until all assets are compile.
 			bool bCompilationFinish = false;
 			Async(EAsyncExecution::TaskGraphMainThread, [&bCompilationFinish, &ImportedObjects]()
 				{
 					TRACE_CPUPROFILER_EVENT_SCOPE(UE::Interchange::FTaskWaitAssetCompilation::DoTask::IsCompilingLambda_GameThread);
+					LLM_SCOPE_BYNAME(TEXT("Interchange"));
 					//Make sure all asset compiling managers are up to date, In case the game thread is waiting for the import to finish (like automation test or synchronous import)
 					FAssetCompilingManager::Get().ProcessAsyncTasks();
 
@@ -137,6 +143,8 @@ void UE::Interchange::FTaskPipelinePostImport::DoTask(ENamedThreads::Type Curren
 	INTERCHANGE_TRACE_ASYNCHRONOUS_TASK(PipelinePostImport)
 #endif
 
+	LLM_SCOPE_BYNAME(TEXT("Interchange"));
+
 	TOptional<FGCScopeGuard> GCScopeGuard;
 	if (!IsInGameThread())
 	{
