Interchange: Allow to override an asset with a dialog message question.

#rnx

diff --git a/Engine/Source/Runtime/Interchange/Engine/Private/Tasks/InterchangeTaskImportObject.cpp b/Engine/Source/Runtime/Interchange/Engine/Private/Tasks/InterchangeTaskImportObject.cpp
--- a/Engine/Source/Runtime/Interchange/Engine/Private/Tasks/InterchangeTaskImportObject.cpp
+++ b/Engine/Source/Runtime/Interchange/Engine/Private/Tasks/InterchangeTaskImportObject.cpp
@@ -13,6 +13,8 @@
 #include "InterchangeSourceData.h"
 #include "InterchangeTranslatorBase.h"
 #include "Interfaces/Interface_AsyncCompilation.h"
+#include "Misc/App.h"
+#include "Misc/MessageDialog.h"
 #include "Misc/Paths.h"
 #include "PackageUtils/PackageUtils.h"
 #include "Stats/Stats.h"
@@ -385,18 +387,6 @@ void UE::Interchange::FTaskImportObject_GameThread::DoTask(ENamedThreads::Type C
 			return;
 		}
 
-		if (!bPackageWasCreated && !AsyncHelper->TaskData.bReplaceExisting)
-		{
-			//Do not replace existing asset, the option tell us to not override it. Skip this asset witha display message
-			UInterchangeResultWarning_Generic* Message = Factory->AddMessage<UInterchangeResultWarning_Generic>();
-			Message->SourceAssetName = AsyncHelper->SourceDatas[SourceIndex]->GetFilename();
-			Message->DestinationAssetName = AssetName;
-			Message->AssetType = FactoryNode->GetObjectClass();
-			Message->Text = FText::Format(NSLOCTEXT("FTaskImportObject_GameThread", "CouldntReplaceExistingAsset", "The option bReplaceExisting is false so we are not overriding the asset named '{0}'.")
-				, FText::FromString(AssetName));
-			bSkipObjectNoReplace = true;
-		}
-
 		if (!bPackageWasCreated && AsyncHelper->TaskData.bFollowRedirectors)
 		{
 			if (UObjectRedirector* Redirector = FindObject<UObjectRedirector>(Pkg, *AssetName))
@@ -412,6 +402,48 @@ void UE::Interchange::FTaskImportObject_GameThread::DoTask(ENamedThreads::Type C
 			}
 		}
 		ExistingAsset = StaticFindObject(nullptr, Pkg, *AssetName);
+		if (ExistingAsset && !AsyncHelper->TaskData.bReplaceExisting)
+		{
+			const FString AssetFullName = ExistingAsset->GetFullName();
+			//If the bReplaceExistingAllDialogAnswer was set do not show again the message dialog, simply reuse the previous answer.
+			if (AsyncHelper->TaskData.bReplaceExistingAllDialogAnswer.IsSet())
+			{
+				bSkipObjectNoReplace = !AsyncHelper->TaskData.bReplaceExistingAllDialogAnswer.GetValue();
+			}
+			else
+			{
+				bSkipObjectNoReplace = true;
+				FText OverrideDialogMessage = FText::Format(NSLOCTEXT("InterchangeTaskimportObject", "OverrideAssetMessage", "Are you sure you want to override asset '{0}'?")
+					, FText::FromString(AssetFullName));
+				if (!GIsAutomationTesting && !FApp::IsUnattended() && !FApp::IsGame())
+				{
+					EAppReturnType::Type DialogResult = FMessageDialog::Open(EAppMsgType::YesNoYesAllNoAll, OverrideDialogMessage);
+					switch (DialogResult)
+					{
+						case EAppReturnType::YesAll:
+							AsyncHelper->TaskData.bReplaceExistingAllDialogAnswer = true;
+						case EAppReturnType::Yes:
+							bSkipObjectNoReplace = false;
+							break;
+						case EAppReturnType::NoAll:
+							AsyncHelper->TaskData.bReplaceExistingAllDialogAnswer = false;
+						case EAppReturnType::No:
+							bSkipObjectNoReplace = true;
+					}
+				}
+			}
+
+			if (bSkipObjectNoReplace)
+			{
+				//Do not replace existing asset, the option tell us to not override it. Skip this asset witha display message
+				UInterchangeResultWarning_Generic* Message = Factory->AddMessage<UInterchangeResultWarning_Generic>();
+				Message->SourceAssetName = AsyncHelper->SourceDatas[SourceIndex]->GetFilename();
+				Message->DestinationAssetName = AssetFullName;
+				Message->AssetType = FactoryNode->GetObjectClass();
+				Message->Text = FText::Format(NSLOCTEXT("FTaskImportObject_GameThread", "CouldntReplaceExistingAsset", "The option bReplaceExisting is false, we are not overriding the asset named '{0}'.")
+					, FText::FromString(AssetFullName));
+			}
+		}
 	}
 
 	if (ExistingAsset)
diff --git a/Engine/Source/Runtime/Interchange/Engine/Public/InterchangeManager.h b/Engine/Source/Runtime/Interchange/Engine/Public/InterchangeManager.h
--- a/Engine/Source/Runtime/Interchange/Engine/Public/InterchangeManager.h
+++ b/Engine/Source/Runtime/Interchange/Engine/Public/InterchangeManager.h
@@ -106,6 +106,12 @@ namespace UE
 
 			/** Overwrite existing assets */
 			bool bReplaceExisting = true;
+
+			/**
+			 * Interchange import task will show a dialog in case user try to override an existing asset and bReplaceExisting is false,
+			 * if this optional is set, it will override or not all existing assets this task try to override.
+			 */
+			TOptional<bool> bReplaceExistingAllDialogAnswer;
 		};
 
 		class FImportResult : protected FGCObject
