Interchange: Limit Node AssetName to 256 characters

#rnx

diff --git a/Engine/Source/Runtime/Interchange/Core/Private/Nodes/InterchangeBaseNode.cpp b/Engine/Source/Runtime/Interchange/Core/Private/Nodes/InterchangeBaseNode.cpp
--- a/Engine/Source/Runtime/Interchange/Core/Private/Nodes/InterchangeBaseNode.cpp
+++ b/Engine/Source/Runtime/Interchange/Core/Private/Nodes/InterchangeBaseNode.cpp
@@ -415,20 +415,23 @@ FGuid UInterchangeBaseNode::GetHash() const
 
 FString UInterchangeBaseNode::GetAssetName() const
 {
-	if (!Attributes->ContainAttribute(UE::Interchange::FBaseNodeStaticData::AssetNameKey()))
+	FString OutName = GetDisplayLabel();
+	if (Attributes->ContainAttribute(UE::Interchange::FBaseNodeStaticData::AssetNameKey()))
 	{
-		return GetDisplayLabel();
+		UE::Interchange::FAttributeStorage::TAttributeHandle<FString> Handle = Attributes->GetAttributeHandle<FString>(UE::Interchange::FBaseNodeStaticData::AssetNameKey());
+		if (Handle.IsValid())
+		{
+			FString Value;
+			Handle.Get(Value);
+			OutName = Value;
+		}
 	}
-
-	UE::Interchange::FAttributeStorage::TAttributeHandle<FString> Handle = Attributes->GetAttributeHandle<FString>(UE::Interchange::FBaseNodeStaticData::AssetNameKey());
-	if (Handle.IsValid())
+	if (OutName.Len() > 256)
 	{
-		FString Value;
-		Handle.Get(Value);
-		return Value;
+		FString GuidStr = FGuid::NewGuid().ToString(EGuidFormats::Base36Encoded);
+		OutName = OutName.Left(115) + GuidStr + OutName.Right(115);
 	}
-
-	return GetDisplayLabel();
+	return OutName;
 }
 
 bool UInterchangeBaseNode::SetAssetName(const FString& AssetName)
