Interchange: Fix an issue when a asset name is longer then 1023 character

#rnx

diff --git a/Engine/Source/Runtime/Interchange/Core/Private/Nodes/InterchangeBaseNode.cpp b/Engine/Source/Runtime/Interchange/Core/Private/Nodes/InterchangeBaseNode.cpp
--- a/Engine/Source/Runtime/Interchange/Core/Private/Nodes/InterchangeBaseNode.cpp
+++ b/Engine/Source/Runtime/Interchange/Core/Private/Nodes/InterchangeBaseNode.cpp
@@ -2,6 +2,7 @@
 #include "Nodes/InterchangeBaseNode.h"
 
 #include "CoreMinimal.h"
+#include "Misc/SecureHash.h"
 #include "Types/AttributeStorage.h"
 #include "UObject/Class.h"
 #include "UObject/Object.h"
@@ -428,7 +429,15 @@ FString UInterchangeBaseNode::GetAssetName() const
 	}
 	if (OutName.Len() > 256)
 	{
-		FString GuidStr = FGuid::NewGuid().ToString(EGuidFormats::Base36Encoded);
+		// Compute a 128-bit hash based on the string and use that as a GUID :
+		FTCHARToUTF8 Converted(*OutName);
+		FMD5 MD5Gen;
+		MD5Gen.Update((const uint8*)Converted.Get(), Converted.Length());
+		uint32 Digest[4];
+		MD5Gen.Final((uint8*)Digest);
+		FString GuidStr = FGuid(Digest[0], Digest[1], Digest[2], Digest[3]).ToString(EGuidFormats::Base36Encoded);
+
+		//Put the guid between the first and last 115 charcaters
 		OutName = OutName.Left(115) + GuidStr + OutName.Right(115);
 	}
 	return OutName;
