Interchange: Polish the project settings. Rename some settings to better/simpler name. Add a per translator show import dialog override to allow more control. We can now not show the import dialog for texture but show it for one particular texture translator like the jpeg. This will be useful if users need to do a translator that need to show the options even if the asset category is not showing the import dialog.

#rnx

diff --git a/Engine/Source/Editor/UnrealEd/Private/Fbx/FbxFactory.cpp b/Engine/Source/Editor/UnrealEd/Private/Fbx/FbxFactory.cpp
--- a/Engine/Source/Editor/UnrealEd/Private/Fbx/FbxFactory.cpp
+++ b/Engine/Source/Editor/UnrealEd/Private/Fbx/FbxFactory.cpp
@@ -1148,7 +1148,7 @@ UObject* UFbxFactory::RecursiveImportNode(UnFbx::FFbxImporter* FbxImporter, void
 
 void UFbxFactory::CleanUp() 
 {
-	UnFbx::FFbxImporter* FFbxImporter = UnFbx::FFbxImporter::GetInstance();
+	UnFbx::FFbxImporter* FFbxImporter = UnFbx::FFbxImporter::GetInstance(true/*bDoNotCreate*/);
 	bDetectImportTypeOnImport = true;
 	bShowOption = true;
 	// load options
diff --git a/Engine/Source/Editor/UnrealEd/Private/Fbx/FbxMainImport.cpp b/Engine/Source/Editor/UnrealEd/Private/Fbx/FbxMainImport.cpp
--- a/Engine/Source/Editor/UnrealEd/Private/Fbx/FbxMainImport.cpp
+++ b/Engine/Source/Editor/UnrealEd/Private/Fbx/FbxMainImport.cpp
@@ -602,10 +602,16 @@ FFbxImporter::~FFbxImporter()
 //-------------------------------------------------------------------------
 //
 //-------------------------------------------------------------------------
-FFbxImporter* FFbxImporter::GetInstance()
+FFbxImporter* FFbxImporter::GetInstance(bool bDoNotCreate /*= false*/)
 {
 	if (!StaticInstance.IsValid())
 	{
+		//Return nullptr if we cannot create the instance
+		if (bDoNotCreate)
+		{
+			return nullptr;
+		}
+
 		StaticInstance = MakeShareable( new FFbxImporter() );
 	}
 	return StaticInstance.Get();
diff --git a/Engine/Source/Editor/UnrealEd/Public/FbxImporter.h b/Engine/Source/Editor/UnrealEd/Public/FbxImporter.h
--- a/Engine/Source/Editor/UnrealEd/Public/FbxImporter.h
+++ b/Engine/Source/Editor/UnrealEd/Public/FbxImporter.h
@@ -606,7 +606,7 @@ public:
 	/**
 	 * Returns the importer singleton. It will be created on the first request.
 	 */
-	UNREALED_API static FFbxImporter* GetInstance();
+	UNREALED_API static FFbxImporter* GetInstance(bool bDoNotCreate = false);
 	static void DeleteInstance();
 
 	/**
diff --git a/Engine/Source/Runtime/Interchange/Engine/Private/InterchangeManager.cpp b/Engine/Source/Runtime/Interchange/Engine/Private/InterchangeManager.cpp
--- a/Engine/Source/Runtime/Interchange/Engine/Private/InterchangeManager.cpp
+++ b/Engine/Source/Runtime/Interchange/Engine/Private/InterchangeManager.cpp
@@ -1623,11 +1623,11 @@ UInterchangeManager::ImportInternal(const FString& ContentPath, const UInterchan
 
 	//In runtime we do not have any pipeline configurator
 #if WITH_EDITORONLY_DATA
-	TSoftClassPtr <UInterchangePipelineConfigurationBase> PipelineConfigurationDialogClass = InterchangeImportSettings.PipelineConfigurationDialogClass;
+	TSoftClassPtr <UInterchangePipelineConfigurationBase> ImportDialogClass = InterchangeImportSettings.ImportDialogClass;
 
-	if (PipelineConfigurationDialogClass.IsValid())
+	if (ImportDialogClass.IsValid())
 	{
-		UClass* PipelineConfigurationClass = PipelineConfigurationDialogClass.LoadSynchronous();
+		UClass* PipelineConfigurationClass = ImportDialogClass.LoadSynchronous();
 		if (PipelineConfigurationClass)
 		{
 			RegisteredPipelineConfiguration = NewObject<UInterchangePipelineConfigurationBase>(GetTransientPackage(), PipelineConfigurationClass, NAME_None, RF_NoFlags);
diff --git a/Engine/Source/Runtime/Interchange/Engine/Private/InterchangeProjectSettings.cpp b/Engine/Source/Runtime/Interchange/Engine/Private/InterchangeProjectSettings.cpp
--- a/Engine/Source/Runtime/Interchange/Engine/Private/InterchangeProjectSettings.cpp
+++ b/Engine/Source/Runtime/Interchange/Engine/Private/InterchangeProjectSettings.cpp
@@ -110,7 +110,7 @@ bool FInterchangeProjectSettingsUtils::ShouldShowPipelineStacksConfigurationDial
 {
 	const FInterchangeImportSettings& ImportSettings = GetDefaultImportSettings(bIsSceneImport);
 
-	bool bShowPipelineStacksConfigurationDialog = ImportSettings.bShowPipelineStacksConfigurationDialog;
+	bool bShowImportDialog = ImportSettings.bShowImportDialog;
 
 	if (!bIsSceneImport)
 	{
@@ -119,18 +119,41 @@ bool FInterchangeProjectSettingsUtils::ShouldShowPipelineStacksConfigurationDial
 		if (UInterchangeTranslatorBase* Translator = ScopedTranslator.GetTranslator())
 		{
 			EInterchangeTranslatorAssetType SupportedAssetTypes = Translator->GetSupportedAssetTypes();
+			const UClass* TranslatorClass = Translator->GetClass();
 
+			//Iterate all override, if there is at least one override that show the imnport dialog we will show it.
+			bool bFoundOverride = false;
+			bool bShowFromOverrideStack = false;
 			const FInterchangeContentImportSettings& ContentImportSettings = GetDefault<UInterchangeProjectSettings>()->ContentImportSettings;
-			for (TMap<EInterchangeTranslatorAssetType, bool>::TConstIterator StackOverridesIt = ContentImportSettings.ShowPipelineStacksConfigurationDialogOverride.CreateConstIterator(); StackOverridesIt; ++StackOverridesIt)
+			for (const TPair<EInterchangeTranslatorAssetType, FInterchangeDialogOverride>& ShowImportDialog : ContentImportSettings.ShowImportDialogOverride)
 			{
-				if ((SupportedAssetTypes ^ StackOverridesIt->Key) < StackOverridesIt->Key)
+				if ((ShowImportDialog.Key == EInterchangeTranslatorAssetType::None && SupportedAssetTypes == EInterchangeTranslatorAssetType::None)
+					|| (static_cast<uint8>(ShowImportDialog.Key & SupportedAssetTypes) > 0))
 				{
-					bShowPipelineStacksConfigurationDialog = StackOverridesIt->Value;
-					break;
+					//Look if there is a per translator override
+					const FInterchangeDialogOverride& InterchangeDialogOverride = ShowImportDialog.Value;
+					const FInterchangePerTranslatorDialogOverride* FindPerTranslatorOverride = InterchangeDialogOverride.PerTranslatorImportDialogOverride.FindByPredicate([&TranslatorClass](const FInterchangePerTranslatorDialogOverride& InterchangePerTranslatorDialogOverride)
+						{
+							return (InterchangePerTranslatorDialogOverride.Translator.Get() == TranslatorClass);
+						});
+					if (FindPerTranslatorOverride)
+					{
+						bShowFromOverrideStack |= FindPerTranslatorOverride->bShowImportDialog;
+					}
+					else
+					{
+						//Simple override
+						bShowFromOverrideStack |= InterchangeDialogOverride.bShowImportDialog;
+					}
+					bFoundOverride = true;
 				}
 			}
+			if (bFoundOverride)
+			{
+				bShowImportDialog = bShowFromOverrideStack;
+			}
 		}
 	}
 
-	return bShowPipelineStacksConfigurationDialog;
+	return bShowImportDialog;
 }
\ No newline at end of file
diff --git a/Engine/Source/Runtime/Interchange/Engine/Public/InterchangeProjectSettings.h b/Engine/Source/Runtime/Interchange/Engine/Public/InterchangeProjectSettings.h
--- a/Engine/Source/Runtime/Interchange/Engine/Public/InterchangeProjectSettings.h
+++ b/Engine/Source/Runtime/Interchange/Engine/Public/InterchangeProjectSettings.h
@@ -51,13 +51,37 @@ struct FInterchangeImportSettings
 	UPROPERTY(EditAnywhere, Category = "Pipeline")
 	FName DefaultPipelineStack = NAME_None;
 
-	/** This tell interchange which pipeline configuration dialog to popup when we need to configure the pipelines.*/
+	/** This tell interchange which class to instance to show the import dialog.*/
 	UPROPERTY(EditAnywhere, Category = "Pipeline")
-	TSoftClassPtr <UInterchangePipelineConfigurationBase> PipelineConfigurationDialogClass;
+	TSoftClassPtr <UInterchangePipelineConfigurationBase> ImportDialogClass;
 
-	/** If enabled, the pipeline stacks configuration dialog will show when interchange must choose a pipeline to import or re-import. If disabled interchange will use the DefaultPipelineStack.*/
+	/** If enabled, the import option dialog will show when interchange import or re-import.*/
 	UPROPERTY(EditAnywhere, Category = "Pipeline")
-	bool bShowPipelineStacksConfigurationDialog = true;
+	bool bShowImportDialog = true;
+};
+
+USTRUCT()
+struct FInterchangePerTranslatorDialogOverride
+{
+	GENERATED_BODY()
+
+	UPROPERTY(EditAnywhere, Category = "DialogOverride", meta = (AllowedClasses = "/Script/InterchangeCore.InterchangeTranslatorBase"))
+	TSoftClassPtr<UInterchangeTranslatorBase> Translator;
+
+	UPROPERTY(EditAnywhere, Category = "DialogOverride")
+	bool bShowImportDialog = true;
+};
+
+USTRUCT()
+struct FInterchangeDialogOverride
+{
+	GENERATED_BODY()
+
+	UPROPERTY(EditAnywhere, Category = "DialogOverride")
+	bool bShowImportDialog = true;
+
+	UPROPERTY(EditAnywhere, Category = "DialogOverride", meta = (AllowedClasses = "/Script/InterchangeCore.InterchangePipelineBase, /Script/InterchangeEngine.InterchangeBlueprintPipelineBase, /Script/InterchangeEngine.InterchangePythonPipelineAsset"))
+	TArray<FInterchangePerTranslatorDialogOverride> PerTranslatorImportDialogOverride;
 };
 
 USTRUCT()
@@ -69,9 +93,9 @@ struct FInterchangeContentImportSettings : public FInterchangeImportSettings
 	UPROPERTY(EditAnywhere, Category = "Pipeline", Meta=(DisplayAfter="DefaultPipelineStack"))
 	TMap<EInterchangeTranslatorAssetType, FName> DefaultPipelineStackOverride;
 
-	/** This tell interchange which pipeline stack to select when importing.*/
-	UPROPERTY(EditAnywhere, Category = "Pipeline", Meta=(DisplayAfter="bShowPipelineStacksConfigurationDialog"))
-	TMap<EInterchangeTranslatorAssetType, bool> ShowPipelineStacksConfigurationDialogOverride;
+	/** This tell interchange if the import dialog should show or not when importing a particular type of asset.*/
+	UPROPERTY(EditAnywhere, Category = "Pipeline", Meta=(DisplayAfter="bShowImportDialog"))
+	TMap<EInterchangeTranslatorAssetType, FInterchangeDialogOverride> ShowImportDialogOverride;
 };
 
 UCLASS(config=Engine, meta=(DisplayName=Interchange), MinimalAPI)
