Replaced CONSTEXPR with constexpr.

diff --git a/Engine/Source/Runtime/Core/Public/Misc/Variant.h b/Engine/Source/Runtime/Core/Public/Misc/Variant.h
--- a/Engine/Source/Runtime/Core/Public/Misc/Variant.h
+++ b/Engine/Source/Runtime/Core/Public/Misc/Variant.h
@@ -95,7 +95,7 @@ enum class EVariantTypes : int32
  */
 template<typename T> struct TVariantTraits
 {
-	static CONSTEXPR EVariantTypes GetType()
+	static constexpr EVariantTypes GetType()
 	{
 		static_assert(!sizeof(T), "Variant trait must be specialized for this type.");
 		return EVariantTypes::Empty;
@@ -414,264 +414,264 @@ FORCEINLINE TArray<uint8> FVariant::GetValue<TArray<uint8> >() const
 /** Implements variant type traits for the built-in ANSICHAR type. */
 template<> struct TVariantTraits<ANSICHAR>
 {
-	static CONSTEXPR EVariantTypes GetType() { return EVariantTypes::Ansichar; }
+	static constexpr EVariantTypes GetType() { return EVariantTypes::Ansichar; }
 };
 
 
 /** Implements variant type traits for the built-in bool type. */
 template<> struct TVariantTraits<bool>
 {
-	static CONSTEXPR EVariantTypes GetType() { return EVariantTypes::Bool; }
+	static constexpr EVariantTypes GetType() { return EVariantTypes::Bool; }
 };
 
 
 /** Implements variant type traits for the built-in FBox type. */
 template<> struct TVariantTraits<FBox>
 {
-	static CONSTEXPR EVariantTypes GetType() { return EVariantTypes::Box; }
+	static constexpr EVariantTypes GetType() { return EVariantTypes::Box; }
 };
 
 
 /** Implements variant type traits for the built-in FBoxSphereBounds type. */
 template<> struct TVariantTraits<FBoxSphereBounds>
 {
-	static CONSTEXPR EVariantTypes GetType() { return EVariantTypes::BoxSphereBounds; }
+	static constexpr EVariantTypes GetType() { return EVariantTypes::BoxSphereBounds; }
 };
 
 
 /** Implements variant type traits for byte arrays. */
 template<> struct TVariantTraits<TArray<uint8> >
 {
-	static CONSTEXPR EVariantTypes GetType() { return EVariantTypes::ByteArray; }
+	static constexpr EVariantTypes GetType() { return EVariantTypes::ByteArray; }
 };
 
 
 /** Implements variant type traits for the built-in FColor type. */
 template<> struct TVariantTraits<FColor>
 {
-	static CONSTEXPR EVariantTypes GetType() { return EVariantTypes::Color; }
+	static constexpr EVariantTypes GetType() { return EVariantTypes::Color; }
 };
 
 
 /** Implements variant type traits for the built-in FDateTime type. */
 template<> struct TVariantTraits<FDateTime>
 {
-	static CONSTEXPR EVariantTypes GetType() { return EVariantTypes::DateTime; }
+	static constexpr EVariantTypes GetType() { return EVariantTypes::DateTime; }
 };
 
 
 /** Implements variant type traits for the built-in double type. */
 template<> struct TVariantTraits<double>
 {
-	static CONSTEXPR EVariantTypes GetType() { return EVariantTypes::Double; }
+	static constexpr EVariantTypes GetType() { return EVariantTypes::Double; }
 };
 
 
 /** Implements variant type traits for enumeration types. */
 template<typename EnumType> struct TVariantTraits<TEnumAsByte<EnumType> >
 {
-	static CONSTEXPR EVariantTypes GetType() { return EVariantTypes::Enum; }
+	static constexpr EVariantTypes GetType() { return EVariantTypes::Enum; }
 };
 
 
 /** Implements variant type traits for the built-in float type. */
 template<> struct TVariantTraits<float>
 {
-	static CONSTEXPR EVariantTypes GetType() { return EVariantTypes::Float; }
+	static constexpr EVariantTypes GetType() { return EVariantTypes::Float; }
 };
 
 
 /** Implements variant type traits for the built-in FGuid type. */
 template<> struct TVariantTraits<FGuid>
 {
-	static CONSTEXPR EVariantTypes GetType() { return EVariantTypes::Guid; }
+	static constexpr EVariantTypes GetType() { return EVariantTypes::Guid; }
 };
 
 
 /** Implements variant type traits for the built-in int8 type. */
 template<> struct TVariantTraits<int8>
 {
-	static CONSTEXPR EVariantTypes GetType() { return EVariantTypes::Int8; }
+	static constexpr EVariantTypes GetType() { return EVariantTypes::Int8; }
 };
 
 
 /** Implements variant type traits for the built-in int16 type. */
 template<> struct TVariantTraits<int16>
 {
-	static CONSTEXPR EVariantTypes GetType() { return EVariantTypes::Int16; }
+	static constexpr EVariantTypes GetType() { return EVariantTypes::Int16; }
 };
 
 
 /** Implements variant type traits for the built-in int32 type. */
 template<> struct TVariantTraits<int32>
 {
-	static CONSTEXPR EVariantTypes GetType() { return EVariantTypes::Int32; }
+	static constexpr EVariantTypes GetType() { return EVariantTypes::Int32; }
 };
 
 
 /** Implements variant type traits for the built-in int64 type. */
 template<> struct TVariantTraits<int64>
 {
-	static CONSTEXPR EVariantTypes GetType() { return EVariantTypes::Int64; }
+	static constexpr EVariantTypes GetType() { return EVariantTypes::Int64; }
 };
 
 
 /** Implements variant type traits for the built-in FIntPoint type. */
 template<> struct TVariantTraits<FIntPoint>
 {
-	static CONSTEXPR EVariantTypes GetType() { return EVariantTypes::IntPoint; }
+	static constexpr EVariantTypes GetType() { return EVariantTypes::IntPoint; }
 };
 
 
 /** Implements variant type traits for the built-in FIntVector type. */
 template<> struct TVariantTraits<FIntVector>
 {
-	static CONSTEXPR EVariantTypes GetType() { return EVariantTypes::IntVector; }
+	static constexpr EVariantTypes GetType() { return EVariantTypes::IntVector; }
 };
 
 
 /** Implements variant type traits for the built-in FIntRect type. */
 template<> struct TVariantTraits<FIntRect>
 {
-	static CONSTEXPR EVariantTypes GetType() { return EVariantTypes::IntRect; }
+	static constexpr EVariantTypes GetType() { return EVariantTypes::IntRect; }
 };
 
 
 /** Implements variant type traits for the built-in FLinearColor type. */
 template<> struct TVariantTraits<FLinearColor>
 {
-	static CONSTEXPR EVariantTypes GetType() { return EVariantTypes::LinearColor; }
+	static constexpr EVariantTypes GetType() { return EVariantTypes::LinearColor; }
 };
 
 
 /** Implements variant type traits for the built-in FMatrix type. */
 template<> struct TVariantTraits<FMatrix>
 {
-	static CONSTEXPR EVariantTypes GetType() { return EVariantTypes::Matrix; }
+	static constexpr EVariantTypes GetType() { return EVariantTypes::Matrix; }
 };
 
 
 /** Implements variant type traits for the built-in FPlane type. */
 template<> struct TVariantTraits<FPlane>
 {
-	static CONSTEXPR EVariantTypes GetType() { return EVariantTypes::Plane; }
+	static constexpr EVariantTypes GetType() { return EVariantTypes::Plane; }
 };
 
 
 /** Implements variant type traits for the built-in FQuat type. */
 template<> struct TVariantTraits<FQuat>
 {
-	static CONSTEXPR EVariantTypes GetType() { return EVariantTypes::Quat; }
+	static constexpr EVariantTypes GetType() { return EVariantTypes::Quat; }
 };
 
 
 /** Implements variant type traits for the built-in FName type. */
 template<> struct TVariantTraits<FName>
 {
-	static CONSTEXPR EVariantTypes GetType() { return EVariantTypes::Name; }
+	static constexpr EVariantTypes GetType() { return EVariantTypes::Name; }
 };
 
 
 /** Implements variant type traits for the built-in FRandomStream type. */
 template<> struct TVariantTraits<FRandomStream>
 {
-	static CONSTEXPR EVariantTypes GetType() { return EVariantTypes::RandomStream; }
+	static constexpr EVariantTypes GetType() { return EVariantTypes::RandomStream; }
 };
 
 
 /** Implements variant type traits for the built-in FRotator type. */
 template<> struct TVariantTraits<FRotator>
 {
-	static CONSTEXPR EVariantTypes GetType() { return EVariantTypes::Rotator; }
+	static constexpr EVariantTypes GetType() { return EVariantTypes::Rotator; }
 };
 
 
 /** Implements variant type traits for the built-in FString type. */
 template<> struct TVariantTraits<FString>
 {
-	static CONSTEXPR EVariantTypes GetType() { return EVariantTypes::String; }
+	static constexpr EVariantTypes GetType() { return EVariantTypes::String; }
 };
 
 
 /** Implements variant type traits for the built-in WIDECHAR type. */
 template<> struct TVariantTraits<WIDECHAR>
 {
-	static CONSTEXPR EVariantTypes GetType() { return EVariantTypes::Widechar; }
+	static constexpr EVariantTypes GetType() { return EVariantTypes::Widechar; }
 };
 
 
 /** Implements variant type traits for the built-in FTimespan type. */
 template<> struct TVariantTraits<FTimespan>
 {
-	static CONSTEXPR EVariantTypes GetType() { return EVariantTypes::Timespan; }
+	static constexpr EVariantTypes GetType() { return EVariantTypes::Timespan; }
 };
 
 
 /** Implements variant type traits for the built-in FTransform type. */
 template<> struct TVariantTraits<FTransform>
 {
-	static CONSTEXPR EVariantTypes GetType() { return EVariantTypes::Transform; }
+	static constexpr EVariantTypes GetType() { return EVariantTypes::Transform; }
 };
 
 
 /** Implements variant type traits for the built-in FTwoVectors type. */
 template<> struct TVariantTraits<FTwoVectors>
 {
-	static CONSTEXPR EVariantTypes GetType() { return EVariantTypes::TwoVectors; }
+	static constexpr EVariantTypes GetType() { return EVariantTypes::TwoVectors; }
 };
 
 
 /** Implements variant type traits for the built-in uint8 type. */
 template<> struct TVariantTraits<uint8>
 {
-	static CONSTEXPR EVariantTypes GetType() { return EVariantTypes::UInt8; }
+	static constexpr EVariantTypes GetType() { return EVariantTypes::UInt8; }
 };
 
 
 /** Implements variant type traits for the built-in uint16 type. */
 template<> struct TVariantTraits<uint16>
 {
-	static CONSTEXPR EVariantTypes GetType() { return EVariantTypes::UInt16; }
+	static constexpr EVariantTypes GetType() { return EVariantTypes::UInt16; }
 };
 
 
 /** Implements variant type traits for the built-in uint32 type. */
 template<> struct TVariantTraits<uint32>
 {
-	static CONSTEXPR EVariantTypes GetType() { return EVariantTypes::UInt32; }
+	static constexpr EVariantTypes GetType() { return EVariantTypes::UInt32; }
 };
 
 
 /** Implements variant type traits for the built-in uint64 type. */
 template<> struct TVariantTraits<uint64>
 {
-	static CONSTEXPR EVariantTypes GetType() { return EVariantTypes::UInt64; }
+	static constexpr EVariantTypes GetType() { return EVariantTypes::UInt64; }
 };
 
 
 /** Implements variant type traits for the built-in FVector type. */
 template<> struct TVariantTraits<FVector>
 {
-	static CONSTEXPR EVariantTypes GetType() { return EVariantTypes::Vector; }
+	static constexpr EVariantTypes GetType() { return EVariantTypes::Vector; }
 };
 
 
 /** Implements variant type traits for the built-in FVector2D type. */
 template<> struct TVariantTraits<FVector2D>
 {
-	static CONSTEXPR EVariantTypes GetType() { return EVariantTypes::Vector2d; }
+	static constexpr EVariantTypes GetType() { return EVariantTypes::Vector2d; }
 };
 
 
 /** Implements variant type traits for the built-in FVector4 type. */
 template<> struct TVariantTraits<FVector4>
 {
-	static CONSTEXPR EVariantTypes GetType() { return EVariantTypes::Vector4; }
+	static constexpr EVariantTypes GetType() { return EVariantTypes::Vector4; }
 };
 
 
 /** Implements variant type traits for the built-in NetworkGUID type. */
 template<> struct TVariantTraits<FNetworkGUID>
 {
-	static CONSTEXPR EVariantTypes GetType() { return EVariantTypes::NetworkGUID; }
+	static constexpr EVariantTypes GetType() { return EVariantTypes::NetworkGUID; }
 };
diff --git a/Engine/Source/Runtime/Core/Public/Templates/ResolveTypeAmbiguity.h b/Engine/Source/Runtime/Core/Public/Templates/ResolveTypeAmbiguity.h
--- a/Engine/Source/Runtime/Core/Public/Templates/ResolveTypeAmbiguity.h
+++ b/Engine/Source/Runtime/Core/Public/Templates/ResolveTypeAmbiguity.h
@@ -52,7 +52,7 @@
 	}
 
 #define MIX_SIGNED_INTS_2_ARGS(Func)				MIX_SIGNED_INTS_2_ARGS_ACTUAL(Func,)
-#define MIX_SIGNED_INTS_2_ARGS_CONSTEXPR(Func)		MIX_SIGNED_INTS_2_ARGS_ACTUAL(Func, CONSTEXPR)
+#define MIX_SIGNED_INTS_2_ARGS_CONSTEXPR(Func)		MIX_SIGNED_INTS_2_ARGS_ACTUAL(Func, constexpr)
 
 // Mixing any signed integral types with any other signed integral type results in same type as "X * Y", which is promoted to the type of the result of mixed arithmetic between the types
 #define MIX_SIGNED_TYPES_3_ARGS_ACTUAL(Func, OptionalMarkup) \
@@ -78,7 +78,7 @@
 	}
 
 #define MIX_SIGNED_INTS_3_ARGS(Func)				MIX_SIGNED_INTS_3_ARGS_ACTUAL(Func,)
-#define MIX_SIGNED_INTS_3_ARGS_CONSTEXPR(Func)		MIX_SIGNED_INTS_3_ARGS_ACTUAL(Func, CONSTEXPR)
+#define MIX_SIGNED_INTS_3_ARGS_CONSTEXPR(Func)		MIX_SIGNED_INTS_3_ARGS_ACTUAL(Func, constexpr)
 
 
 // Mixing float and double types with any other type results in same type as "X * Y", which is promoted to the floating point type with the highest precision of the argument types.
diff --git a/Engine/Source/Runtime/Engine/Classes/Components/ActorComponent.h b/Engine/Source/Runtime/Engine/Classes/Components/ActorComponent.h
--- a/Engine/Source/Runtime/Engine/Classes/Components/ActorComponent.h
+++ b/Engine/Source/Runtime/Engine/Classes/Components/ActorComponent.h
@@ -84,22 +84,22 @@ enum class EUpdateTransformFlags : int32
 	OnlyUpdateIfUsingSocket = 0x4	
 };
 
-CONSTEXPR inline EUpdateTransformFlags operator|(EUpdateTransformFlags Left, EUpdateTransformFlags Right)
+constexpr inline EUpdateTransformFlags operator|(EUpdateTransformFlags Left, EUpdateTransformFlags Right)
 {
 	return static_cast<EUpdateTransformFlags> ( static_cast<int32> (Left) | static_cast<int32> (Right) );
 }
 
-CONSTEXPR inline EUpdateTransformFlags operator&(EUpdateTransformFlags Left, EUpdateTransformFlags Right)
+constexpr inline EUpdateTransformFlags operator&(EUpdateTransformFlags Left, EUpdateTransformFlags Right)
 {
 	return static_cast<EUpdateTransformFlags> (static_cast<int32> (Left) & static_cast<int32> (Right));
 }
 
-CONSTEXPR inline bool operator !(EUpdateTransformFlags Value)
+constexpr inline bool operator !(EUpdateTransformFlags Value)
 {
 	return Value == EUpdateTransformFlags::None;
 }
 
-CONSTEXPR inline EUpdateTransformFlags operator ~(EUpdateTransformFlags Value)
+constexpr inline EUpdateTransformFlags operator ~(EUpdateTransformFlags Value)
 {
 	return static_cast<EUpdateTransformFlags>(~static_cast<int32>(Value));
 }
diff --git a/Engine/Source/Runtime/Engine/Classes/Engine/InstancedStaticMesh.h b/Engine/Source/Runtime/Engine/Classes/Engine/InstancedStaticMesh.h
--- a/Engine/Source/Runtime/Engine/Classes/Engine/InstancedStaticMesh.h
+++ b/Engine/Source/Runtime/Engine/Classes/Engine/InstancedStaticMesh.h
@@ -306,7 +306,7 @@ public:
 	ENGINE_API virtual void InitRHI(FRHICommandListBase& RHICmdList) override;
 
 	/** Make sure we account for changes in the signature of GetStaticBatchElementVisibility() */
-	static CONSTEXPR uint32 NumBitsForVisibilityMask()
+	static constexpr uint32 NumBitsForVisibilityMask()
 	{		
 		return 8 * sizeof(uint64);
 	}
diff --git a/Engine/Source/Runtime/Interchange/Core/Public/Types/AttributeStorage.h b/Engine/Source/Runtime/Interchange/Core/Public/Types/AttributeStorage.h
--- a/Engine/Source/Runtime/Interchange/Core/Public/Types/AttributeStorage.h
+++ b/Engine/Source/Runtime/Interchange/Core/Public/Types/AttributeStorage.h
@@ -287,7 +287,7 @@ namespace UE
 		 */
 		template<typename T, typename Enable = void> struct TAttributeTypeTraits
 		{
-			static CONSTEXPR EAttributeTypes GetType()
+			static constexpr EAttributeTypes GetType()
 			{
 				static_assert(!sizeof(T), "Attribute type trait must be specialized for this type.");
 				return EAttributeTypes::None;
@@ -1102,7 +1102,7 @@ namespace UE
 		/** Implements variant type traits for the built-in bool type. */
 		template<> struct TAttributeTypeTraits<bool>
 		{
-			static CONSTEXPR EAttributeTypes GetType() { return EAttributeTypes::Bool; }
+			static constexpr EAttributeTypes GetType() { return EAttributeTypes::Bool; }
 			static FString ToString(const bool& Value) { return Value ? TEXT("true") : TEXT("false"); }
 		};
 
@@ -1116,14 +1116,14 @@ namespace UE
 		/** Implements variant type traits for byte arrays. */
 		template<> struct TAttributeTypeTraits<TArray<uint8> >
 		{
-			static CONSTEXPR EAttributeTypes GetType() { return EAttributeTypes::ByteArray; }
+			static constexpr EAttributeTypes GetType() { return EAttributeTypes::ByteArray; }
 			static FString ToString(const TArray<uint8>& Value) {return TEXT("Array<uint8>"); }
 		};
 
 		/** Implements variant type traits for byte array64s. */
 		template<> struct TAttributeTypeTraits<TArray64<uint8> >
 		{
-			static CONSTEXPR EAttributeTypes GetType() { return EAttributeTypes::ByteArray64; }
+			static constexpr EAttributeTypes GetType() { return EAttributeTypes::ByteArray64; }
 			static FString ToString(const TArray64<uint8>& Value) { return TEXT("Array64<uint8>"); }
 		};
 
@@ -1131,7 +1131,7 @@ namespace UE
 		/** Implements variant type traits for the built-in FColor type. */
 		template<> struct TAttributeTypeTraits<FColor>
 		{
-			static CONSTEXPR EAttributeTypes GetType() { return EAttributeTypes::Color; }
+			static constexpr EAttributeTypes GetType() { return EAttributeTypes::Color; }
 			static FString ToString(const FColor& Value) { return Value.ToString(); }
 		};
 
@@ -1139,7 +1139,7 @@ namespace UE
 		/** Implements variant type traits for the built-in FDateTime type. */
 		template<> struct TAttributeTypeTraits<FDateTime>
 		{
-			static CONSTEXPR EAttributeTypes GetType() { return EAttributeTypes::DateTime; }
+			static constexpr EAttributeTypes GetType() { return EAttributeTypes::DateTime; }
 			static FString ToString(const FDateTime& Value) { return Value.ToString(); }
 		};
 
@@ -1147,7 +1147,7 @@ namespace UE
 		/** Implements variant type traits for the built-in double type. */
 		template<> struct TAttributeTypeTraits<double>
 		{
-			static CONSTEXPR EAttributeTypes GetType() { return EAttributeTypes::Double; }
+			static constexpr EAttributeTypes GetType() { return EAttributeTypes::Double; }
 			static FString ToString(const double& Value)
 			{
 				FStringFormatOrderedArguments OrderedArguments;
@@ -1160,7 +1160,7 @@ namespace UE
 		/** Implements variant type traits for enumeration types. */
 		template<typename EnumType> struct TAttributeTypeTraits<TEnumAsByte<EnumType> >
 		{
-			static CONSTEXPR EAttributeTypes GetType() { return EAttributeTypes::Enum; }
+			static constexpr EAttributeTypes GetType() { return EAttributeTypes::Enum; }
 			static FString ToString(const TEnumAsByte<EnumType>& Value)
 			{
 				uint32 ValueConv = static_cast<uint32>(static_cast<uint8>(Value));
@@ -1174,7 +1174,7 @@ namespace UE
 		/** Implements variant type traits for the built-in float type. */
 		template<> struct TAttributeTypeTraits<float>
 		{
-			static CONSTEXPR EAttributeTypes GetType() { return EAttributeTypes::Float; }
+			static constexpr EAttributeTypes GetType() { return EAttributeTypes::Float; }
 			static FString ToString(const float& Value)
 			{
 				FStringFormatOrderedArguments OrderedArguments;
@@ -1187,7 +1187,7 @@ namespace UE
 		/** Implements variant type traits for the built-in FGuid type. */
 		template<> struct TAttributeTypeTraits<FGuid>
 		{
-			static CONSTEXPR EAttributeTypes GetType() { return EAttributeTypes::Guid; }
+			static constexpr EAttributeTypes GetType() { return EAttributeTypes::Guid; }
 			static FString ToString(const FGuid& Value) { return Value.ToString(EGuidFormats::DigitsLower); }
 		};
 
@@ -1195,7 +1195,7 @@ namespace UE
 		/** Implements variant type traits for the built-in int8 type. */
 		template<> struct TAttributeTypeTraits<int8>
 		{
-			static CONSTEXPR EAttributeTypes GetType() { return EAttributeTypes::Int8; }
+			static constexpr EAttributeTypes GetType() { return EAttributeTypes::Int8; }
 			static FString ToString(const int8& Value) { return FString::FromInt(Value); }
 		};
 
@@ -1203,7 +1203,7 @@ namespace UE
 		/** Implements variant type traits for the built-in int16 type. */
 		template<> struct TAttributeTypeTraits<int16>
 		{
-			static CONSTEXPR EAttributeTypes GetType() { return EAttributeTypes::Int16; }
+			static constexpr EAttributeTypes GetType() { return EAttributeTypes::Int16; }
 			static FString ToString(const int16& Value) { return FString::FromInt(Value); }
 		};
 
@@ -1211,7 +1211,7 @@ namespace UE
 		/** Implements variant type traits for the built-in int32 type. */
 		template<> struct TAttributeTypeTraits<int32>
 		{
-			static CONSTEXPR EAttributeTypes GetType() { return EAttributeTypes::Int32; }
+			static constexpr EAttributeTypes GetType() { return EAttributeTypes::Int32; }
 			static FString ToString(const int32& Value) { return FString::FromInt(Value); }
 		};
 
@@ -1219,7 +1219,7 @@ namespace UE
 		/** Implements variant type traits for the built-in int64 type. */
 		template<> struct TAttributeTypeTraits<int64>
 		{
-			static CONSTEXPR EAttributeTypes GetType() { return EAttributeTypes::Int64; }
+			static constexpr EAttributeTypes GetType() { return EAttributeTypes::Int64; }
 			static FString ToString(const int64& Value)
 			{
 				FStringFormatOrderedArguments OrderedArguments;
@@ -1231,7 +1231,7 @@ namespace UE
 		/** Implements variant type traits for the built-in FIntRect type. */
 		template<> struct TAttributeTypeTraits<FIntRect>
 		{
-			static CONSTEXPR EAttributeTypes GetType() { return EAttributeTypes::IntRect; }
+			static constexpr EAttributeTypes GetType() { return EAttributeTypes::IntRect; }
 			static FString ToString(const FIntRect& Value) { return Value.ToString(); }
 		};
 
@@ -1239,7 +1239,7 @@ namespace UE
 		/** Implements variant type traits for the built-in FLinearColor type. */
 		template<> struct TAttributeTypeTraits<FLinearColor>
 		{
-			static CONSTEXPR EAttributeTypes GetType() { return EAttributeTypes::LinearColor; }
+			static constexpr EAttributeTypes GetType() { return EAttributeTypes::LinearColor; }
 			static FString ToString(const FLinearColor& Value) { return Value.ToString(); }
 		};
 
@@ -1249,7 +1249,7 @@ namespace UE
 		/** Implements variant type traits for the built-in FName type. */
 		template<> struct TAttributeTypeTraits<FName>
 		{
-			static CONSTEXPR EAttributeTypes GetType() { return EAttributeTypes::Name; }
+			static constexpr EAttributeTypes GetType() { return EAttributeTypes::Name; }
 			static FString ToString(const FName& Value) { return Value.ToString(); }
 		};
 
@@ -1261,7 +1261,7 @@ namespace UE
 		/** Implements variant type traits for the built-in FRandomStream type. */
 		template<> struct TAttributeTypeTraits<FRandomStream>
 		{
-			static CONSTEXPR EAttributeTypes GetType() { return EAttributeTypes::RandomStream; }
+			static constexpr EAttributeTypes GetType() { return EAttributeTypes::RandomStream; }
 			static FString ToString(const FRandomStream& Value) { return Value.ToString(); }
 		};
 
@@ -1272,7 +1272,7 @@ namespace UE
 		/** Implements variant type traits for the built-in FString type. */
 		template<> struct TAttributeTypeTraits<FString>
 		{
-			static CONSTEXPR EAttributeTypes GetType() { return EAttributeTypes::String; }
+			static constexpr EAttributeTypes GetType() { return EAttributeTypes::String; }
 			static FString ToString(const FString& Value) { return Value; }
 		};
 
@@ -1280,7 +1280,7 @@ namespace UE
 		/** Implements variant type traits for the built-in FTimespan type. */
 		template<> struct TAttributeTypeTraits<FTimespan>
 		{
-			static CONSTEXPR EAttributeTypes GetType() { return EAttributeTypes::Timespan; }
+			static constexpr EAttributeTypes GetType() { return EAttributeTypes::Timespan; }
 			static FString ToString(const FTimespan& Value) { return Value.ToString(); }
 		};
 
@@ -1291,7 +1291,7 @@ namespace UE
 		/** Implements variant type traits for the built-in FTwoVectors type. */
 		template<> struct TAttributeTypeTraits<FTwoVectors>
 		{
-			static CONSTEXPR EAttributeTypes GetType() { return EAttributeTypes::TwoVectors; }
+			static constexpr EAttributeTypes GetType() { return EAttributeTypes::TwoVectors; }
 			static FString ToString(const FTwoVectors& Value) { return Value.ToString(); }
 		};
 
@@ -1299,7 +1299,7 @@ namespace UE
 		/** Implements variant type traits for the built-in uint8 type. */
 		template<> struct TAttributeTypeTraits<uint8>
 		{
-			static CONSTEXPR EAttributeTypes GetType() { return EAttributeTypes::UInt8; }
+			static constexpr EAttributeTypes GetType() { return EAttributeTypes::UInt8; }
 			static FString ToString(const uint8& Value)
 			{
 				uint32 ValueConv = Value;
@@ -1321,7 +1321,7 @@ namespace UE
 				>::type
 			>
 		{
-			static CONSTEXPR EAttributeTypes GetType() { return EAttributeTypes::UInt8; }
+			static constexpr EAttributeTypes GetType() { return EAttributeTypes::UInt8; }
 			static FString ToString(const EnumClassType& Value)
 			{
 				uint32 ValueConv = (uint8)Value;
@@ -1334,7 +1334,7 @@ namespace UE
 		/** Implements variant type traits for the built-in uint16 type. */
 		template<> struct TAttributeTypeTraits<uint16>
 		{
-			static CONSTEXPR EAttributeTypes GetType() { return EAttributeTypes::UInt16; }
+			static constexpr EAttributeTypes GetType() { return EAttributeTypes::UInt16; }
 			static FString ToString(const uint16& Value)
 			{
 				uint32 ValueConv = Value;
@@ -1348,7 +1348,7 @@ namespace UE
 		/** Implements variant type traits for the built-in uint32 type. */
 		template<> struct TAttributeTypeTraits<uint32>
 		{
-			static CONSTEXPR EAttributeTypes GetType() { return EAttributeTypes::UInt32; }
+			static constexpr EAttributeTypes GetType() { return EAttributeTypes::UInt32; }
 			static FString ToString(const uint32& Value)
 			{
 				FStringFormatOrderedArguments OrderedArguments;
@@ -1361,7 +1361,7 @@ namespace UE
 		/** Implements variant type traits for the built-in uint64 type. */
 		template<> struct TAttributeTypeTraits<uint64>
 		{
-			static CONSTEXPR EAttributeTypes GetType() { return EAttributeTypes::UInt64; }
+			static constexpr EAttributeTypes GetType() { return EAttributeTypes::UInt64; }
 			static FString ToString(const uint64& Value)
 			{
 				FStringFormatOrderedArguments OrderedArguments;
@@ -1377,7 +1377,7 @@ namespace UE
 		/** Implements variant type traits for the built-in FVector2D type. */
 		template<> struct TAttributeTypeTraits<FVector2D>
 		{
-			static CONSTEXPR EAttributeTypes GetType() { return EAttributeTypes::Vector2d; }
+			static constexpr EAttributeTypes GetType() { return EAttributeTypes::Vector2d; }
 			static FString ToString(const FVector2D& Value) { return Value.ToString(); }
 		};
 
@@ -1386,28 +1386,28 @@ namespace UE
 		/** Implements variant type traits for the built-in Vector2DHalf type. */
 		template<> struct TAttributeTypeTraits<FIntPoint>
 		{
-			static CONSTEXPR EAttributeTypes GetType() { return EAttributeTypes::IntPoint; }
+			static constexpr EAttributeTypes GetType() { return EAttributeTypes::IntPoint; }
 			static FString ToString(const FIntPoint& Value) { return Value.ToString(); }
 		};
 
 		/** Implements variant type traits for the built-in Vector2DHalf type. */
 		template<> struct TAttributeTypeTraits<FIntVector>
 		{
-			static CONSTEXPR EAttributeTypes GetType() { return EAttributeTypes::IntVector; }
+			static constexpr EAttributeTypes GetType() { return EAttributeTypes::IntVector; }
 			static FString ToString(const FIntVector& Value) { return Value.ToString(); }
 		};
 
 		/** Implements variant type traits for the built-in Vector2DHalf type. */
 		template<> struct TAttributeTypeTraits<FVector2DHalf>
 		{
-			static CONSTEXPR EAttributeTypes GetType() { return EAttributeTypes::Vector2DHalf; }
+			static constexpr EAttributeTypes GetType() { return EAttributeTypes::Vector2DHalf; }
 			static FString ToString(const FVector2DHalf& Value) { return Value.ToString(); }
 		};
 
 		/** Implements variant type traits for the built-in Float16 type. */
 		template<> struct TAttributeTypeTraits<FFloat16>
 		{
-			static CONSTEXPR EAttributeTypes GetType() { return EAttributeTypes::Float16; }
+			static constexpr EAttributeTypes GetType() { return EAttributeTypes::Float16; }
 			static FString ToString(const FFloat16& Value)
 			{
 				float ValueConvert = Value.GetFloat();
@@ -1420,7 +1420,7 @@ namespace UE
 		/** Implements variant type traits for the built-in OrientedBox type. */
 		template<> struct TAttributeTypeTraits<FOrientedBox>
 		{
-			static CONSTEXPR EAttributeTypes GetType() { return EAttributeTypes::OrientedBox; }
+			static constexpr EAttributeTypes GetType() { return EAttributeTypes::OrientedBox; }
 			static FString ToString(const FOrientedBox& Value) { return TEXT("FOrientedBox"); }
 		};
 
@@ -1429,7 +1429,7 @@ namespace UE
 		/** Implements variant type traits for the built-in FrameNumber type. */
 		template<> struct TAttributeTypeTraits<FFrameNumber>
 		{
-			static CONSTEXPR EAttributeTypes GetType() { return EAttributeTypes::FrameNumber; }
+			static constexpr EAttributeTypes GetType() { return EAttributeTypes::FrameNumber; }
 			static FString ToString(const FFrameNumber& Value)
 			{
 				int32 ValueConvert = Value.Value;
@@ -1442,181 +1442,181 @@ namespace UE
 		/** Implements variant type traits for the built-in FrameRate type. */
 		template<> struct TAttributeTypeTraits<FFrameRate>
 		{
-			static CONSTEXPR EAttributeTypes GetType() { return EAttributeTypes::FrameRate; }
+			static constexpr EAttributeTypes GetType() { return EAttributeTypes::FrameRate; }
 			static FString ToString(const FFrameRate& Value) { return Value.ToPrettyText().ToString(); }
 		};
 
 		/** Implements variant type traits for the built-in FrameTime type. */
 		template<> struct TAttributeTypeTraits<FFrameTime>
 		{
-			static CONSTEXPR EAttributeTypes GetType() { return EAttributeTypes::FrameTime; }
+			static constexpr EAttributeTypes GetType() { return EAttributeTypes::FrameTime; }
 		};
 
 		/** Implements variant type traits for the built-in SoftObjectPath type. */
 		template<> struct TAttributeTypeTraits<FSoftObjectPath>
 		{
-			static CONSTEXPR EAttributeTypes GetType() { return EAttributeTypes::SoftObjectPath; }
+			static constexpr EAttributeTypes GetType() { return EAttributeTypes::SoftObjectPath; }
 			static FString ToString(const FSoftObjectPath& Value) { return Value.ToString(); }
 		};
 
 		/** Implements variant type traits for the built-in Matrix44f type. */
 		template<> struct TAttributeTypeTraits<FMatrix44f>
 		{
-			static CONSTEXPR EAttributeTypes GetType() { return EAttributeTypes::Matrix44f; }
+			static constexpr EAttributeTypes GetType() { return EAttributeTypes::Matrix44f; }
 			static FString ToString(const FMatrix44f& Value) { return Value.ToString(); }
 		};
 
 		/** Implements variant type traits for the built-in Matrix44f type. */
 		template<> struct TAttributeTypeTraits<FMatrix44d>
 		{
-			static CONSTEXPR EAttributeTypes GetType() { return EAttributeTypes::Matrix44d; }
+			static constexpr EAttributeTypes GetType() { return EAttributeTypes::Matrix44d; }
 			static FString ToString(const FMatrix44d& Value) { return Value.ToString(); }
 		};
 
 		/** Implements variant type traits for the built-in Plane4f type. */
 		template<> struct TAttributeTypeTraits<FPlane4f>
 		{
-			static CONSTEXPR EAttributeTypes GetType() { return EAttributeTypes::Plane4f; }
+			static constexpr EAttributeTypes GetType() { return EAttributeTypes::Plane4f; }
 			static FString ToString(const FPlane4f& Value) { return Value.ToString(); }
 		};
 
 		/** Implements variant type traits for the built-in Plane4d type. */
 		template<> struct TAttributeTypeTraits<FPlane4d>
 		{
-			static CONSTEXPR EAttributeTypes GetType() { return EAttributeTypes::Plane4d; }
+			static constexpr EAttributeTypes GetType() { return EAttributeTypes::Plane4d; }
 			static FString ToString(const FPlane4d& Value) { return Value.ToString(); }
 		};
 
 		/** Implements variant type traits for the built-in Quat4f type. */
 		template<> struct TAttributeTypeTraits<FQuat4f>
 		{
-			static CONSTEXPR EAttributeTypes GetType() { return EAttributeTypes::Quat4f; }
+			static constexpr EAttributeTypes GetType() { return EAttributeTypes::Quat4f; }
 			static FString ToString(const FQuat4f& Value) { return Value.ToString(); }
 		};
 
 		/** Implements variant type traits for the built-in Quat4d type. */
 		template<> struct TAttributeTypeTraits<FQuat4d>
 		{
-			static CONSTEXPR EAttributeTypes GetType() { return EAttributeTypes::Quat4d; }
+			static constexpr EAttributeTypes GetType() { return EAttributeTypes::Quat4d; }
 			static FString ToString(const FQuat4d& Value) { return Value.ToString(); }
 		};
 
 		/** Implements variant type traits for the built-in Rotator3f type. */
 		template<> struct TAttributeTypeTraits<FRotator3f>
 		{
-			static CONSTEXPR EAttributeTypes GetType() { return EAttributeTypes::Rotator3f; }
+			static constexpr EAttributeTypes GetType() { return EAttributeTypes::Rotator3f; }
 			static FString ToString(const FRotator3f& Value) { return Value.ToString(); }
 		};
 
 		/** Implements variant type traits for the built-in Rotator3d type. */
 		template<> struct TAttributeTypeTraits<FRotator3d>
 		{
-			static CONSTEXPR EAttributeTypes GetType() { return EAttributeTypes::Rotator3d; }
+			static constexpr EAttributeTypes GetType() { return EAttributeTypes::Rotator3d; }
 			static FString ToString(const FRotator3d& Value) { return Value.ToString(); }
 		};
 
 		/** Implements variant type traits for the built-in Rotator3f type. */
 		template<> struct TAttributeTypeTraits<FTransform3f>
 		{
-			static CONSTEXPR EAttributeTypes GetType() { return EAttributeTypes::Transform3f; }
+			static constexpr EAttributeTypes GetType() { return EAttributeTypes::Transform3f; }
 			static FString ToString(const FTransform3f& Value) { return Value.ToString(); }
 		};
 
 		/** Implements variant type traits for the built-in Transform3d type. */
 		template<> struct TAttributeTypeTraits<FTransform3d>
 		{
-			static CONSTEXPR EAttributeTypes GetType() { return EAttributeTypes::Transform3d; }
+			static constexpr EAttributeTypes GetType() { return EAttributeTypes::Transform3d; }
 			static FString ToString(const FTransform3d& Value) { return Value.ToString(); }
 		};
 
 		/** Implements variant type traits for the built-in Vector3f type. */
 		template<> struct TAttributeTypeTraits<FVector3f>
 		{
-			static CONSTEXPR EAttributeTypes GetType() { return EAttributeTypes::Vector3f; }
+			static constexpr EAttributeTypes GetType() { return EAttributeTypes::Vector3f; }
 			static FString ToString(const FVector3f& Value) { return Value.ToString(); }
 		};
 
 		/** Implements variant type traits for the built-in Vector3d type. */
 		template<> struct TAttributeTypeTraits<FVector3d>
 		{
-			static CONSTEXPR EAttributeTypes GetType() { return EAttributeTypes::Vector3d; }
+			static constexpr EAttributeTypes GetType() { return EAttributeTypes::Vector3d; }
 			static FString ToString(const FVector3d& Value) { return Value.ToString(); }
 		};
 
 		/** Implements variant type traits for the built-in Vector2f type. */
 		template<> struct TAttributeTypeTraits<FVector2f>
 		{
-			static CONSTEXPR EAttributeTypes GetType() { return EAttributeTypes::Vector2f; }
+			static constexpr EAttributeTypes GetType() { return EAttributeTypes::Vector2f; }
 			static FString ToString(const FVector2f& Value) { return Value.ToString(); }
 		};
 
 		/** Implements variant type traits for the built-in Vector4f type. */
 		template<> struct TAttributeTypeTraits<FVector4f>
 		{
-			static CONSTEXPR EAttributeTypes GetType() { return EAttributeTypes::Vector4f; }
+			static constexpr EAttributeTypes GetType() { return EAttributeTypes::Vector4f; }
 			static FString ToString(const FVector4f& Value) { return Value.ToString(); }
 		};
 
 		/** Implements variant type traits for the built-in Vector4d type. */
 		template<> struct TAttributeTypeTraits<FVector4d>
 		{
-			static CONSTEXPR EAttributeTypes GetType() { return EAttributeTypes::Vector4d; }
+			static constexpr EAttributeTypes GetType() { return EAttributeTypes::Vector4d; }
 			static FString ToString(const FVector4d& Value) { return Value.ToString(); }
 		};
 
 		/** Implements variant type traits for the built-in Box2f type. */
 		template<> struct TAttributeTypeTraits<FBox2f>
 		{
-			static CONSTEXPR EAttributeTypes GetType() { return EAttributeTypes::Box2f; }
+			static constexpr EAttributeTypes GetType() { return EAttributeTypes::Box2f; }
 			static FString ToString(const FBox2f& Value) { return Value.ToString(); }
 		};
 
 		/** Implements variant type traits for the built-in Box2D type. */
 		template<> struct TAttributeTypeTraits<FBox2D>
 		{
-			static CONSTEXPR EAttributeTypes GetType() { return EAttributeTypes::Box2D; }
+			static constexpr EAttributeTypes GetType() { return EAttributeTypes::Box2D; }
 			static FString ToString(const FBox2D& Value) { return Value.ToString(); }
 		};
 
 		/** Implements variant type traits for the built-in Box3f type. */
 		template<> struct TAttributeTypeTraits<FBox3f>
 		{
-			static CONSTEXPR EAttributeTypes GetType() { return EAttributeTypes::Box3f; }
+			static constexpr EAttributeTypes GetType() { return EAttributeTypes::Box3f; }
 			static FString ToString(const FBox3f& Value) { return Value.ToString(); }
 		};
 
 		/** Implements variant type traits for the built-in Box3d type. */
 		template<> struct TAttributeTypeTraits<FBox3d>
 		{
-			static CONSTEXPR EAttributeTypes GetType() { return EAttributeTypes::Box3d; }
+			static constexpr EAttributeTypes GetType() { return EAttributeTypes::Box3d; }
 			static FString ToString(const FBox3d& Value) { return Value.ToString(); }
 		};
 
 		/** Implements variant type traits for the built-in BoxSphereBounds3f type. */
 		template<> struct TAttributeTypeTraits<FBoxSphereBounds3f>
 		{
-			static CONSTEXPR EAttributeTypes GetType() { return EAttributeTypes::BoxSphereBounds3f; }
+			static constexpr EAttributeTypes GetType() { return EAttributeTypes::BoxSphereBounds3f; }
 			static FString ToString(const FBoxSphereBounds3f& Value) { return Value.ToString(); }
 		};
 
 		/** Implements variant type traits for the built-in BoxSphereBounds3d type. */
 		template<> struct TAttributeTypeTraits<FBoxSphereBounds3d>
 		{
-			static CONSTEXPR EAttributeTypes GetType() { return EAttributeTypes::BoxSphereBounds3d; }
+			static constexpr EAttributeTypes GetType() { return EAttributeTypes::BoxSphereBounds3d; }
 			static FString ToString(const FBoxSphereBounds3d& Value) { return Value.ToString(); }
 		};
 
 		/** Implements variant type traits for the built-in Sphere3f type. */
 		template<> struct TAttributeTypeTraits<FSphere3f>
 		{
-			static CONSTEXPR EAttributeTypes GetType() { return EAttributeTypes::Sphere3f; }
+			static constexpr EAttributeTypes GetType() { return EAttributeTypes::Sphere3f; }
 			static FString ToString(const FSphere3f& Value) { return TEXT("FSphere3f"); }
 		};
 
 		/** Implements variant type traits for the built-in Sphere3d type. */
 		template<> struct TAttributeTypeTraits<FSphere3d>
 		{
-			static CONSTEXPR EAttributeTypes GetType() { return EAttributeTypes::Sphere3d; }
+			static constexpr EAttributeTypes GetType() { return EAttributeTypes::Sphere3d; }
 			static FString ToString(const FSphere3d& Value) { return TEXT("FSphere3d"); }
 		};
 	} //ns interchange
diff --git a/Engine/Source/Runtime/PhysicsCore/Public/CollisionShape.h b/Engine/Source/Runtime/PhysicsCore/Public/CollisionShape.h
--- a/Engine/Source/Runtime/PhysicsCore/Public/CollisionShape.h
+++ b/Engine/Source/Runtime/PhysicsCore/Public/CollisionShape.h
@@ -21,10 +21,10 @@ struct FCollisionShape
 {
 	ECollisionShape::Type ShapeType;
 
-	static FORCEINLINE CONSTEXPR float MinBoxExtent() { return UE_KINDA_SMALL_NUMBER; }
-	static FORCEINLINE CONSTEXPR float MinSphereRadius() { return UE_KINDA_SMALL_NUMBER; }
-	static FORCEINLINE CONSTEXPR float MinCapsuleRadius() { return UE_KINDA_SMALL_NUMBER; }
-	static FORCEINLINE CONSTEXPR float MinCapsuleAxisHalfHeight() { return UE_KINDA_SMALL_NUMBER; }
+	static FORCEINLINE constexpr float MinBoxExtent() { return UE_KINDA_SMALL_NUMBER; }
+	static FORCEINLINE constexpr float MinSphereRadius() { return UE_KINDA_SMALL_NUMBER; }
+	static FORCEINLINE constexpr float MinCapsuleRadius() { return UE_KINDA_SMALL_NUMBER; }
+	static FORCEINLINE constexpr float MinCapsuleAxisHalfHeight() { return UE_KINDA_SMALL_NUMBER; }
 
 	/** Union that supports up to 3 floats **/
 	union
diff --git a/Engine/Source/Runtime/RHI/Public/RHIDefinitions.h b/Engine/Source/Runtime/RHI/Public/RHIDefinitions.h
--- a/Engine/Source/Runtime/RHI/Public/RHIDefinitions.h
+++ b/Engine/Source/Runtime/RHI/Public/RHIDefinitions.h
@@ -588,7 +588,7 @@ inline bool IsUniformBufferStaticSlotValid(const FUniformBufferStaticSlot Slot)
 struct FRHIResourceTableEntry
 {
 public:
-	static CONSTEXPR uint32 GetEndOfStreamToken()
+	static constexpr uint32 GetEndOfStreamToken()
 	{
 		return 0xffffffff;
 	}
