Interchange: Prevent user if they are not using the UInterchangeManager singleton.

#rnx

diff --git a/Engine/Source/Runtime/Interchange/Engine/Private/InterchangeManager.cpp b/Engine/Source/Runtime/Interchange/Engine/Private/InterchangeManager.cpp
--- a/Engine/Source/Runtime/Interchange/Engine/Private/InterchangeManager.cpp
+++ b/Engine/Source/Runtime/Interchange/Engine/Private/InterchangeManager.cpp
@@ -57,6 +57,8 @@ static FAutoConsoleVariableRef CCvarInterchangeImportEnable(
 	TEXT("Whether Interchange import is enabled."),
 	ECVF_Default);
 
+bool UInterchangeManager::bIsCreatingSingleton = false;
+
 namespace UE::Interchange::Private
 {
 	const FLogCategoryBase* GetLogInterchangePtr()
@@ -873,6 +875,16 @@ void UInterchangePipelineStackOverride::AddPipeline(UInterchangePipelineBase* Pi
 	OverridePipelines.Add(PipelineBase);
 }
 
+UInterchangeManager::UInterchangeManager(const FObjectInitializer& ObjectInitializer)
+	: Super(ObjectInitializer)
+{
+	//Client must use the singleton API
+	if (!bIsCreatingSingleton && !HasAnyFlags(RF_ClassDefaultObject | RF_ArchetypeObject))
+	{
+		UE_LOG(LogInterchangeEngine, Error, TEXT("Interchange manager is a singleton you must call GetInterchangeManager() or GetInterchangeManagerScripted() to access it."));
+	}
+}
+
 UInterchangeManager& UInterchangeManager::GetInterchangeManager()
 {
 	static TStrongObjectPtr<UInterchangeManager> InterchangeManager = nullptr;
@@ -907,8 +919,12 @@ UInterchangeManager& UInterchangeManager::GetInterchangeManager()
 		//We cannot create a TStrongObjectPtr outside of the main thread, we also need a valid Transient package
 		check(IsInGameThread() && GetTransientPackage());
 
+		bIsCreatingSingleton = true;
+
 		InterchangeManager = TStrongObjectPtr<UInterchangeManager>(NewObject<UInterchangeManager>(GetTransientPackage(), NAME_None, EObjectFlags::RF_NoFlags));
 
+		bIsCreatingSingleton = false;
+
 		InterchangeManager->GCEndDelegate = FCoreUObjectDelegates::GetPostGarbageCollect().AddLambda([]()
 			{
 				if (IsInterchangeImportEnabled())
@@ -1483,8 +1499,14 @@ UInterchangeManager::ImportInternal(const FString& ContentPath, const UInterchan
 	
 	LLM_SCOPE_BYNAME(TEXT("Interchange"));
 
+	if (!ensure(IsInGameThread()))
+	{
+		UE_LOG(LogInterchangeEngine, Error, TEXT("Cannot import file, the import process can be started only in the game thread."));
+		return TTuple<UE::Interchange::FAssetImportResultRef, UE::Interchange::FSceneImportResultRef>{ MakeShared< UE::Interchange::FImportResult, ESPMode::ThreadSafe >(), MakeShared< UE::Interchange::FImportResult, ESPMode::ThreadSafe >() };
+	}
+
 	ensure(IsInterchangeImportEnabled());
-	check(IsInGameThread());
+
 	static int32 GeneratedUniqueID = 0;
 	int32 UniqueId = ++GeneratedUniqueID;
 
@@ -1524,9 +1546,9 @@ UInterchangeManager::ImportInternal(const FString& ContentPath, const UInterchan
 		}
 	}
 
-	if (!ensure(IsInGameThread()))
+	if (this != &GetInterchangeManager())
 	{
-		//Import process can be started only in the game thread
+		UE_LOG(LogInterchangeEngine, Error, TEXT("Cannot import file, the interchange manager use to import this file is not the singleton, use GetInterchangeManager() or GetInterchangeManagerScripted() to acces the interchange manager singleton."));
 		return EarlyExit();
 	}
 
diff --git a/Engine/Source/Runtime/Interchange/Engine/Public/InterchangeManager.h b/Engine/Source/Runtime/Interchange/Engine/Public/InterchangeManager.h
--- a/Engine/Source/Runtime/Interchange/Engine/Public/InterchangeManager.h
+++ b/Engine/Source/Runtime/Interchange/Engine/Public/InterchangeManager.h
@@ -410,12 +410,14 @@ struct FImportAssetParameters
 	FOnImportDoneNative OnSceneImportDoneNative;
 };
 
-UCLASS(Transient, BlueprintType, MinimalAPI)
+UCLASS(Transient, BlueprintType, CustomConstructor, MinimalAPI)
 class UInterchangeManager : public UObject
 {
 	GENERATED_BODY()
 public:
 
+	UInterchangeManager(const FObjectInitializer& ObjectInitializer);
+
 	/**
 	 * Return the interchange manager singleton pointer.
 	 *
@@ -690,6 +692,9 @@ protected:
 	ImportInternal(const FString& ContentPath, const UInterchangeSourceData* SourceData, const FImportAssetParameters& ImportAssetParameters, const UE::Interchange::EImportType ImportType);
 
 private:
+
+	static bool bIsCreatingSingleton;
+
 	struct FQueuedTaskData
 	{
 		TSharedPtr<UE::Interchange::FImportAsyncHelper, ESPMode::ThreadSafe> AsyncHelper;
