Interchange: fix scene re-import

#rnx

diff --git a/Engine/Source/Runtime/Interchange/Engine/Private/Tasks/InterchangeTaskImportObject.cpp b/Engine/Source/Runtime/Interchange/Engine/Private/Tasks/InterchangeTaskImportObject.cpp
--- a/Engine/Source/Runtime/Interchange/Engine/Private/Tasks/InterchangeTaskImportObject.cpp
+++ b/Engine/Source/Runtime/Interchange/Engine/Private/Tasks/InterchangeTaskImportObject.cpp
@@ -342,7 +342,8 @@ void UE::Interchange::FTaskImportObject_GameThread::DoTask(ENamedThreads::Type C
 		{
 			if (!FactoryNode->ShouldForceNodeReimport())
 			{
-				return;
+				//Skip this object, simply assign it to the node, so we can retrieve it later.
+				bSkipObjectNoReplace = true;
 			}
 		}
 
@@ -402,7 +403,7 @@ void UE::Interchange::FTaskImportObject_GameThread::DoTask(ENamedThreads::Type C
 			}
 		}
 		ExistingAsset = StaticFindObject(nullptr, Pkg, *AssetName);
-		if (ExistingAsset && !AsyncHelper->TaskData.bReplaceExisting)
+		if (!bSkipObjectNoReplace && ExistingAsset && !AsyncHelper->TaskData.bReplaceExisting)
 		{
 			const FString AssetFullName = ExistingAsset->GetFullName();
 			//If the bReplaceExistingAllDialogAnswer was set do not show again the message dialog, simply reuse the previous answer.
@@ -435,7 +436,7 @@ void UE::Interchange::FTaskImportObject_GameThread::DoTask(ENamedThreads::Type C
 
 			if (bSkipObjectNoReplace)
 			{
-				//Do not replace existing asset, the option tell us to not override it. Skip this asset witha display message
+				//Do not replace existing asset, the option tell us to not override it. Skip this asset with a display message
 				UInterchangeResultWarning_Generic* Message = Factory->AddMessage<UInterchangeResultWarning_Generic>();
 				Message->SourceAssetName = AsyncHelper->SourceDatas[SourceIndex]->GetFilename();
 				Message->DestinationAssetName = AssetFullName;
