Added missing change from CL 31192191

diff --git a/Engine/Source/Runtime/Interchange/Engine/Private/Tasks/InterchangeTaskCreateSceneObjects.cpp b/Engine/Source/Runtime/Interchange/Engine/Private/Tasks/InterchangeTaskCreateSceneObjects.cpp
--- a/Engine/Source/Runtime/Interchange/Engine/Private/Tasks/InterchangeTaskCreateSceneObjects.cpp
+++ b/Engine/Source/Runtime/Interchange/Engine/Private/Tasks/InterchangeTaskCreateSceneObjects.cpp
@@ -87,6 +87,7 @@ void UE::Interchange::FTaskCreateSceneObjects::DoTask(ENamedThreads::Type Curren
 		if (AsyncHelper->BaseNodeContainers.IsValidIndex(SourceIndex))
 		{
 			CreateSceneObjectsParams.NodeContainer = AsyncHelper->BaseNodeContainers[SourceIndex].Get();
+			CreateSceneObjectsParams.SourceData = AsyncHelper->SourceDatas[SourceIndex];
 		}
 
 		UObject* SceneObject = Factory->ImportSceneObject_GameThread(CreateSceneObjectsParams);
