Clean up typos and phrasing in Interchange modules - contribution from Robb Surridge

diff --git a/Engine/Plugins/Interchange/Runtime/Source/Pipelines/Public/InterchangePipelineMeshesUtilities.h b/Engine/Plugins/Interchange/Runtime/Source/Pipelines/Public/InterchangePipelineMeshesUtilities.h
--- a/Engine/Plugins/Interchange/Runtime/Source/Pipelines/Public/InterchangePipelineMeshesUtilities.h
+++ b/Engine/Plugins/Interchange/Runtime/Source/Pipelines/Public/InterchangePipelineMeshesUtilities.h
@@ -17,7 +17,7 @@ class UInterchangeMeshNode;
 class UInterchangeSceneNode;
 
 /*
-* This container exist only because UPROPERTY cannot support nested container. See FInterchangeMeshInstance
+* This container exists only because UPROPERTY cannot support nested container. See FInterchangeMeshInstance.
 */
 USTRUCT(BlueprintType)
 struct FInterchangeLodSceneNodeContainer
@@ -25,16 +25,16 @@ struct FInterchangeLodSceneNodeContainer
 	GENERATED_BODY()
 
 	/**
-	 * Each scene node here represent a mesh scene node. Only if we represent a lod group we can have more then 1 mesh scene node for a specific lod index.
+	 * Each scene node here represents a mesh scene node. If it represents a LOD group, there may be more then one mesh scene node for a specific LOD index.
 	 */
 	UPROPERTY(EditAnywhere, Category = "Interchange | Pipeline | MeshesInstance")
 	TArray<TObjectPtr<const UInterchangeSceneNode>> SceneNodes;
 };
 
 /*
-* A mesh instance is a description of a translated scene node that point on a translated mesh asset.
-* A mesh instance pointing on a lod group can have many lods and many scene mesh nodes per lod index.
-* A mesh instance pointing on a mesh node will have only the lod 0 and will point on one scene mesh node.
+* A mesh instance is a description of a translated scene node that points to a translated mesh asset.
+* A mesh instance that points to an LOD group can have many LODs and many scene mesh nodes per LOD index.
+* A mesh instance that points to a mesh node will have only LOD 0 and will point to one scene mesh node.
 */
 USTRUCT(BlueprintType)
 struct FInterchangeMeshInstance
@@ -49,13 +49,13 @@ struct FInterchangeMeshInstance
 		bHasMorphTargets = false;
 	}
 	/**
-	 * This ID represent either 1: a lod group scene node uid or 2: a mesh scene node uid.
+	 * This ID represents either a LOD group scene node UID or a mesh scene node UID.
 	 */
 	UPROPERTY(EditAnywhere, Category = "Interchange | Pipeline | MeshesInstance")
 	FString MeshInstanceUid;
 
 	/**
-	 * If this mesh instance represent a LodGroup this member will not be null, but will be null if the mesh instance do not represent a lod group
+	 * This member is null unless the mesh instance represents a LOD group.
 	 */
 	UPROPERTY(EditAnywhere, Category = "Interchange | Pipeline | MeshesInstance")
 	TObjectPtr<const UInterchangeSceneNode> LodGroupNode;
@@ -70,7 +70,7 @@ struct FInterchangeMeshInstance
 	bool bHasMorphTargets;
 
 	/**
-	 * Each scene node here represent a mesh scene node. Only if we represent a lod group we can have more then 1 mesh scene node for a specific lod index.
+	 * Each scene node here represents a mesh scene node. If it represents a LOD group, there may be more then one mesh scene node for a specific LOD index.
 	 */
 	UPROPERTY(EditAnywhere, Category = "Interchange | Pipeline | MeshesInstance")
 	TMap<int32, FInterchangeLodSceneNodeContainer> SceneNodePerLodIndex;
@@ -83,7 +83,7 @@ struct FInterchangeMeshInstance
 };
 
 /*
-* A mesh geometry is a description of a translated mesh asset node that define a geometry.
+* A mesh geometry is a description of a translated mesh asset node that defines a geometry.
 */
 USTRUCT(BlueprintType)
 struct FInterchangeMeshGeometry
@@ -96,32 +96,32 @@ struct FInterchangeMeshGeometry
 	}
 
 	/**
-	 * Represent the unique id of the UInterchangeMeshNode represent by this structure.
+	 * The unique ID of the UInterchangeMeshNode represented by this structure.
 	 */
 	UPROPERTY(EditAnywhere, Category = "Interchange | Pipeline | MeshesGeometry")
 	FString MeshUid;
 
 	/**
-	 * The UInterchangeMeshNode pointer represent by this structure.
+	 * The UInterchangeMeshNode pointer represented by this structure.
 	 */
 	UPROPERTY(EditAnywhere, Category = "Interchange | Pipeline | MeshesGeometry")
 	TObjectPtr<const UInterchangeMeshNode> MeshNode = nullptr;
 
 	/**
-	 * All mesh instance referencing this UInterchangeMeshNode pointer.
+	 * All mesh instances that refer to this UInterchangeMeshNode pointer.
 	 */
 	UPROPERTY(EditAnywhere, Category = "Interchange | Pipeline | MeshesGeometry")
 	TArray<FString> ReferencingMeshInstanceUids;
 
 	/**
-	 * A list of all scene nodes representing sockets attached to this mesh
+	 * A list of all scene nodes that represent sockets attached to this mesh.
 	 */
 	UPROPERTY(EditAnywhere, Category = "Interchange | Pipeline | MeshesGeometry")
 	TArray<FString> AttachedSocketUids;
 };
 
 /*
-* Represent the context UInterchangePipelineMeshesUtilities will use when client query the data
+* Represents the context UInterchangePipelineMeshesUtilities will use when the client queries data.
 */
 USTRUCT(BlueprintType)
 struct FInterchangePipelineMeshesUtilitiesContext
@@ -129,32 +129,32 @@ struct FInterchangePipelineMeshesUtilitiesContext
 	GENERATED_BODY()
 
 	/**
-	 * Convert static mesh to skeletal mesh
+	 * If enabled, all static meshes are converted to skeletal meshes.
 	 */
 	UPROPERTY(EditAnywhere, BlueprintReadWrite, Category = "Interchange | Pipeline | MeshesContext")
 	bool bConvertStaticMeshToSkeletalMesh = false;
 
 	/**
-	 * Convert static mesh to skeletal mesh
+	 * If enabled, all skeletal meshes are converted to static meshes.
 	 */
 	UPROPERTY(EditAnywhere, BlueprintReadWrite, Category = "Interchange | Pipeline | MeshesContext")
 	bool bConvertSkeletalMeshToStaticMesh = false;
 
 	/**
-	 * Convert static mesh that has morph target to skeletal mesh
+	 * If enabled, all static meshes that have morph targets will be imported as skeletal meshes instead.
 	 */
-	UPROPERTY(EditAnywhere, BlueprintReadWrite, Category = "Interchange | Pipeline | MeshesContext")
+	UPROPERTY(EditAnywhere, BlueprintReadWrite, Category = "Interchange | Pipeline | MeshesContext", meta = (DisplayName = "Convert Static Meshes with Morph Targets to Skeletal Meshes"))
 	bool bConvertStaticsWithMorphTargetsToSkeletals = false;
 
 	/**
-	 * If checked, meshes nested in bone hierarchies will be imported instead of being converted to bones. If the mesh are not skinned they will
-	 * be added to skeletal mesh and remove from the static meshes.
+	 * If enabled, meshes nested in bone hierarchies are imported as meshes instead of being converted to bones. If the meshes are not skinned, they are
+	 * added to the skeletal mesh and removed from the list of static meshes.
 	 */
 	UPROPERTY(EditAnywhere, BlueprintReadWrite, Category = "Interchange | Pipeline | MeshesContext")
 	bool bImportMeshesInBoneHierarchy = true;
 
 	/**
-	 * When querying geometry, this flag will not add MeshGeometry if there is a scene node pointing on a geometry.
+	 * When querying geometry, this flag will not add MeshGeometry if there is a scene node that points to a geometry.
 	 */
 	UPROPERTY(EditAnywhere, BlueprintReadWrite, Category = "Interchange | Pipeline | MeshesContext")
 	bool bQueryGeometryOnlyIfNoInstance = true;
@@ -179,126 +179,126 @@ public:
 	static UInterchangePipelineMeshesUtilities* CreateInterchangePipelineMeshesUtilities(UInterchangeBaseNodeContainer* BaseNodeContainer);
 
 	/**
-	* Get all mesh instance unique ids.
+	* Get the unique IDs of all mesh instances.
 	*/
 	UFUNCTION(BlueprintCallable, Category = "Interchange | Pipeline | Meshes")
 	void GetAllMeshInstanceUids(TArray<FString>& MeshInstanceUids) const;
 
 	/**
-	* Iterate all mesh instance.
+	* Iterate over all mesh instances.
 	*/
 	void IterateAllMeshInstance(TFunctionRef<void(const FInterchangeMeshInstance&)> IterationLambda) const;
 
 	/**
-	* Get all skinned mesh instance unique ids.
+	* Get the unique IDs of all skinned mesh instances.
 	*/
 	UFUNCTION(BlueprintCallable, Category = "Interchange | Pipeline | Meshes")
 	void GetAllSkinnedMeshInstance(TArray<FString>& MeshInstanceUids) const;
 
 	/**
-	* Iterate all skinned mesh instance.
+	* Iterate over all skinned mesh instances.
 	*/
 	void IterateAllSkinnedMeshInstance(TFunctionRef<void(const FInterchangeMeshInstance&)> IterationLambda) const;
 
 	/**
-	* Get all static mesh instance unique ids.
+	* Get the unique IDs of all static mesh instances.
 	*/
 	UFUNCTION(BlueprintCallable, Category = "Interchange | Pipeline | Meshes")
 	void GetAllStaticMeshInstance(TArray<FString>& MeshInstanceUids) const;
 
 	/**
-	* Iterate all static mesh instance.
+	* Iterate over all static mesh instances.
 	*/
 	void IterateAllStaticMeshInstance(TFunctionRef<void(const FInterchangeMeshInstance&)> IterationLambda) const;
 
 	/**
-	* Get all mesh geometry unique ids.
+	* Get the unique IDs of all mesh geometry.
 	*/
 	UFUNCTION(BlueprintCallable, Category = "Interchange | Pipeline | Meshes")
 	void GetAllMeshGeometry(TArray<FString>& MeshGeometryUids) const;
 		
 	/**
-	* Iterate all mesh geometry.
+	* Iterate over all mesh geometry.
 	*/
 	void IterateAllMeshGeometry(TFunctionRef<void(const FInterchangeMeshGeometry&)> IterationLambda) const;
 
 	/**
-	* Get all skinned mesh geometry unique ids.
+	* Get the unique IDs of all skinned mesh geometry.
 	*/
 	UFUNCTION(BlueprintCallable, Category = "Interchange | Pipeline | Meshes")
 	void GetAllSkinnedMeshGeometry(TArray<FString>& MeshGeometryUids) const;
 
 	/**
-	* Iterate all skinned mesh geometry.
+	* Iterate over all skinned mesh geometry.
 	*/
 	void IterateAllSkinnedMeshGeometry(TFunctionRef<void(const FInterchangeMeshGeometry&)> IterationLambda) const;
 
 	/**
-	* Get all static mesh geometry unique ids.
+	* Get the unique IDs of all static mesh geometry.
 	*/
 	UFUNCTION(BlueprintCallable, Category = "Interchange | Pipeline | Meshes")
 	void GetAllStaticMeshGeometry(TArray<FString>& MeshGeometryUids) const;
 
 	/**
-	* Iterate all static mesh geometry.
+	* Iterate over all static mesh geometry.
 	*/
 	void IterateAllStaticMeshGeometry(TFunctionRef<void(const FInterchangeMeshGeometry&)> IterationLambda) const;
 
 	/**
-	* Get all not instanced mesh geometry unique ids.
+	* Get the unique IDs of all non-instanced mesh geometry.
 	*/
 	UFUNCTION(BlueprintCallable, Category = "Interchange | Pipeline | Meshes")
 	void GetAllMeshGeometryNotInstanced(TArray<FString>& MeshGeometryUids) const;
 
 	/**
-	* Iterate all mesh geometry not instanced.
+	* Iterate over all non-instanced mesh geometry.
 	*/
 	void IterateAllMeshGeometryNotIntanced(TFunctionRef<void(const FInterchangeMeshGeometry&)> IterationLambda) const;
 
 	/**
-	* Return true if there is an existing FInterchangeMeshInstance matching the MeshInstanceUid key.
+	* Return true if there is an existing FInterchangeMeshInstance that matches the MeshInstanceUid key.
 	*/
 	UFUNCTION(BlueprintCallable, Category = "Interchange | Pipeline | Meshes")
 	bool IsValidMeshInstanceUid(const FString& MeshInstanceUid) const;
 
 	/**
-	* Get the instanced mesh from the unique ids.
+	* Get the instanced mesh from the unique ID.
 	*/
 	UFUNCTION(BlueprintCallable, Category = "Interchange | Pipeline | Meshes")
 	const FInterchangeMeshInstance& GetMeshInstanceByUid(const FString& MeshInstanceUid) const;
 
 	/**
-	* Return true if there is an existing FInterchangeMeshGeometry matching the MeshInstanceUid key.
+	* Return true if there is an existing FInterchangeMeshGeometry that matches the MeshInstanceUid key.
 	*/
 	UFUNCTION(BlueprintCallable, Category = "Interchange | Pipeline | Meshes")
 	bool IsValidMeshGeometryUid(const FString& MeshGeometryUid) const;
 
 	/**
-	* Get the geometry mesh from the unique ids.
+	* Get the geometry mesh from the unique ID.
 	*/
 	UFUNCTION(BlueprintCallable, Category = "Interchange | Pipeline | Meshes")
 	const FInterchangeMeshGeometry& GetMeshGeometryByUid(const FString& MeshGeometryUid) const;
 
 	/**
-	* Get all instanced mesh uids using the mesh geometry unique ids.
+	* Get all instanced mesh UIDs that use the mesh geometry unique ID.
 	*/
 	UFUNCTION(BlueprintCallable, Category = "Interchange | Pipeline | Meshes")
 	void GetAllMeshInstanceUidsUsingMeshGeometryUid(const FString& MeshGeometryUid, TArray<FString>& MeshInstanceUids) const;
 
 	/**
-	* Iterate all instanced mesh uids using the mesh geometry unique ids.
+	* Iterate over all instanced mesh UIDs that use the mesh geometry unique ID.
 	*/
 	void IterateAllMeshInstanceUsingMeshGeometry(const FString& MeshGeometryUid, TFunctionRef<void(const FInterchangeMeshInstance&)> IterationLambda) const;
 
 	/**
-	* Return a list of skinned FInterchangeMeshInstance uid that can be combined together.
-	* We cannot create a skinned mesh with multiple skeleton root node, This function return combined MeshInstance per skeleton roots
+	* Return a list of skinned FInterchangeMeshInstance UIDs that can be combined together.
+	* We cannot create a skinned mesh that has multiple skeleton root nodes. This function returns a combined MeshInstance per skeleton roots.
 	*/
 	void GetCombinedSkinnedMeshInstances(TMap<FString, TArray<FString>>& OutMeshInstanceUidsPerSkeletonRootUid) const;
 	
 	/**
-	* Return the skeleton root node Uid, this is the uid for a UInterchangeSceneNode that has a "Joint" specialized type.
-	* Return an empty string if the MeshInstanceUid parameter point on nothing.
+	* Return the skeleton root node UID. This is the UID for a UInterchangeSceneNode that has a "Joint" specialized type.
+	* Return an empty string if the MeshInstanceUid parameter points to nothing.
 	*/
 	UFUNCTION(BlueprintCallable, Category = "Interchange | Pipeline | Meshes")
 	FString GetMeshInstanceSkeletonRootUid(const FString& MeshInstanceUid) const;
@@ -306,8 +306,8 @@ public:
 	FString GetMeshInstanceSkeletonRootUid(const FInterchangeMeshInstance& MeshInstance) const;
 
 	/**
-	* Return the skeleton root node Uid, this is the uid for a UInterchangeSceneNode that has a "Joint" specialized type.
-	* Return an empty string if the MeshGeometryUid parameter point on nothing.
+	* Return the skeleton root node UID. This is the UID for a UInterchangeSceneNode that has a "Joint" specialized type.
+	* Return an empty string if the MeshGeometryUid parameter points to nothing.
 	*/
 	UFUNCTION(BlueprintCallable, Category = "Interchange | Pipeline | Meshes")
 	FString GetMeshGeometrySkeletonRootUid(const FString& MeshGeometryUid) const;
diff --git a/Engine/Plugins/Interchange/Runtime/Source/Pipelines/Public/InterchangeglTFPipeline.h b/Engine/Plugins/Interchange/Runtime/Source/Pipelines/Public/InterchangeglTFPipeline.h
--- a/Engine/Plugins/Interchange/Runtime/Source/Pipelines/Public/InterchangeglTFPipeline.h
+++ b/Engine/Plugins/Interchange/Runtime/Source/Pipelines/Public/InterchangeglTFPipeline.h
@@ -16,7 +16,7 @@ class UInterchangeBaseNodeContainer;
 class UInterchangeMaterialInstanceFactoryNode;
 class UInterchangeShaderGraphNode;
 
-UCLASS(config = Interchange, meta = (DisplayName = "Interchange glTF"))
+UCLASS(config = Interchange, meta = (DisplayName = "Interchange glTF", ToolTip = "Interchange settings for glTF conversions."))
 class INTERCHANGEPIPELINES_API UGLTFPipelineSettings : public UDeveloperSettings
 {
 	GENERATED_BODY()
@@ -36,8 +36,8 @@ private:
 
 	static const TArray<FString> ExpectedMaterialInstanceIdentifiers; //Default MaterialInstance' identifiers
 
-	//Helper for the Settings Customizer to decide if the ParentMaterials should be editible or not
-	//Should be editible from Project Settings, and should NOT be editible from the Import
+	//Helper for the Settings Customizer to decide if the ParentMaterials should be editable or not.
+	//Should be editable from Project Settings, and should NOT be editable from the Import.
 	//It is set by the Pipeline's Customizer.
 	//Equals to Pipeline->CanEditPropertiesStates()
 	bool bMaterialParentsEditible = true;
diff --git a/Engine/Plugins/Tests/InterchangeTests/Source/InterchangeTests/Public/InterchangeImportTestStepReimport.h b/Engine/Plugins/Tests/InterchangeTests/Source/InterchangeTests/Public/InterchangeImportTestStepReimport.h
--- a/Engine/Plugins/Tests/InterchangeTests/Source/InterchangeTests/Public/InterchangeImportTestStepReimport.h
+++ b/Engine/Plugins/Tests/InterchangeTests/Source/InterchangeTests/Public/InterchangeImportTestStepReimport.h
@@ -17,19 +17,19 @@ class INTERCHANGETESTS_API UInterchangeImportTestStepReimport : public UIntercha
 
 public:
 
-	/** The source file to import (path relative to the json script) */
+	/** The source file to import (path relative to the json script). */
 	UPROPERTY(EditAnywhere, Category = General)
 	FFilePath SourceFileToReimport;
 
-	/** The type of the asset to reimport. In the case that only one such asset were imported, this is unambiguous */
+	/** The type of the asset to reimport. If only one such asset was imported, this is unambiguous. */
 	UPROPERTY(EditAnywhere, Category = General)
 	TSubclassOf<UObject> AssetTypeToReimport;
 
-	/** If there were multiple assets of the above type imported, specify the concrete name here */
+	/** If there were multiple assets of the above type imported, specify the concrete name here. */
 	UPROPERTY(EditAnywhere, Category = General)
 	FString AssetNameToReimport;
 
-	/** Whether imported assets should be saved and freshly reloaded after import */
+	/** Determines whether imported assets should be saved and freshly reloaded after import. */
 	UPROPERTY(EditAnywhere, Category = General)
 	bool bSaveThenReloadImportedAssets = true;
 
diff --git a/Engine/Source/Runtime/Interchange/Core/Private/Types/AttributeStorage.cpp b/Engine/Source/Runtime/Interchange/Core/Private/Types/AttributeStorage.cpp
--- a/Engine/Source/Runtime/Interchange/Core/Private/Types/AttributeStorage.cpp
+++ b/Engine/Source/Runtime/Interchange/Core/Private/Types/AttributeStorage.cpp
@@ -127,12 +127,12 @@ namespace UE
 
 			if (HasAttributeStorageResult(Result, EAttributeStorageResult::Operation_Error_WrongSize))
 			{
-				UE_LOG(LogInterchangeCore, Warning, TEXT("Attribute storage operation [%s] Key[%s]: Stored attribute value size do not match parameter value size."), *OperationName, *(AttributeKey.ToString()));
+				UE_LOG(LogInterchangeCore, Warning, TEXT("Attribute storage operation [%s] Key[%s]: Stored attribute value size does not match parameter value size."), *OperationName, *(AttributeKey.ToString()));
 			}
 
 			if (HasAttributeStorageResult(Result, EAttributeStorageResult::Operation_Error_WrongType))
 			{
-				UE_LOG(LogInterchangeCore, Warning, TEXT("Attribute storage operation [%s] Key[%s]: Stored attribute value type do not match parameter value type."), *OperationName, *(AttributeKey.ToString()));
+				UE_LOG(LogInterchangeCore, Warning, TEXT("Attribute storage operation [%s] Key[%s]: Stored attribute value type does not match parameter value type."), *OperationName, *(AttributeKey.ToString()));
 			}
 
 			if (HasAttributeStorageResult(Result, EAttributeStorageResult::Operation_Error_InvalidStorage))
@@ -418,7 +418,7 @@ namespace UE
 
 		EAttributeStorageResult FAttributeStorage::GetAttribute(const FAttributeKey& ElementAttributeKey, FString& OutValue, TSpecializeType<FString >) const
 		{
-			static_assert(TAttributeTypeTraits<FString>::GetType() != EAttributeTypes::None, "Not a supported type for the attributes. Check EAttributeTypes for the supported types");
+			static_assert(TAttributeTypeTraits<FString>::GetType() != EAttributeTypes::None, "Not a supported type for the attributes. Check EAttributeTypes for the supported types.");
 
 			//Lock the storage
 			FScopeLock ScopeLock(&StorageMutex);
@@ -455,7 +455,7 @@ namespace UE
 
 		EAttributeStorageResult FAttributeStorage::GetAttribute(const FAttributeKey& ElementAttributeKey, FName& OutValue, TSpecializeType<FName >) const
 		{
-			static_assert(TAttributeTypeTraits<FName>::GetType() != EAttributeTypes::None, "Not a supported type for the attributes. Check EAttributeTypes for supported types");
+			static_assert(TAttributeTypeTraits<FName>::GetType() != EAttributeTypes::None, "Not a supported type for the attributes. Check EAttributeTypes for supported types.");
 
 			//Lock the storage
 			FScopeLock ScopeLock(&StorageMutex);
@@ -501,7 +501,7 @@ namespace UE
 
 		EAttributeStorageResult FAttributeStorage::GetAttribute(const FAttributeKey& ElementAttributeKey, FSoftObjectPath& OutValue, TSpecializeType<FSoftObjectPath >) const
 		{
-			static_assert(TAttributeTypeTraits<FSoftObjectPath>::GetType() != EAttributeTypes::None, "Not a supported type for the attributes. Check EAttributeTypes for supported types");
+			static_assert(TAttributeTypeTraits<FSoftObjectPath>::GetType() != EAttributeTypes::None, "Not a supported type for the attributes. Check EAttributeTypes for supported types.");
 
 			//Lock the storage
 			FScopeLock ScopeLock(&StorageMutex);
diff --git a/Engine/Source/Runtime/Interchange/Core/Public/InterchangeFactoryBase.h b/Engine/Source/Runtime/Interchange/Core/Public/InterchangeFactoryBase.h
--- a/Engine/Source/Runtime/Interchange/Core/Public/InterchangeFactoryBase.h
+++ b/Engine/Source/Runtime/Interchange/Core/Public/InterchangeFactoryBase.h
@@ -43,19 +43,18 @@ public:
 /**
  * Asset factory implementation:
  * 
- * The first 3 steps use the interchange factory node to import/re-import the UObject
+ * The first three steps use the Interchange factory node to import or reimport the UObject:
  * 
- * 1. BeginImportAsset_GameThread - Create the asset UObject. You can also import source data (retrieve payloads) and setup properties on the game thread.
- * 2. ImportAsset_Async - Import source data (retrieve payloads) and setup properties asynchronously on any thread.
- * 3. EndImportAsset_GameThread - Anything you need to do on the game thread to finalize the import source data and setup properties. (Conflict resolution that need UI)
+ * 1. BeginImportAsset_GameThread - Create the asset UObject. You can also import source data (retrieve payloads) and set up properties on the game thread.
+ * 2. ImportAsset_Async - Import source data (retrieve payloads) and set up properties asynchronously on any thread.
+ * 3. EndImportAsset_GameThread - Anything you need to do on the game thread to finalize the import source data and set up properties. For example, conflict resolution that needs UI.
  * 
- * The last two steps are helpful to change the imported/re-imported UObject before and after the PostEditChange (render data build) is call on the asset.
+ * The last two steps are helpful to change the imported or reimported UObject before and after the PostEditChange (render data build) is called on the asset.
  * 
  * 4. SetupObject_GameThread - Do any UObject setup required before the build (before PostEditChange), the UObject dependencies should exist and have all the source data and properties imported.
  * 5. FinalizeObject_GameThread - Do any final UObject setup after the build (after PostEditChange). Note that the build of an asset can be asynchronous and this function will be call after the async build is done.
  * 
- * 
- * Scene factory implementation
+ * Scene factory implementation:
  * 
  * 1. ImportSceneObject_GameThread - Create an actor in a level.
  */
@@ -67,7 +66,7 @@ class UInterchangeFactoryBase : public UObject
 public:
 
 	/**
-	 * return the UClass this factory can create.
+	 * Return the UClass this factory can create.
 	 */
 	UFUNCTION(BlueprintCallable, Category = "Interchange | Factory")
 	virtual UClass* GetFactoryClass() const
@@ -76,47 +75,47 @@ public:
 	}
 
 	/**
-	 * return the asset type this factory can create.
+	 * Return the asset type this factory can create.
 	 */
 	virtual EInterchangeFactoryAssetType GetFactoryAssetType() { return EInterchangeFactoryAssetType::None; }
 
 	/**
-	 * Parameters to pass to CreateAsset function
+	 * Parameters to pass to the CreateAsset function.
 	 */
 	struct FImportAssetObjectParams
 	{
-		/** The package where to create the asset, if null it will put it in the transient package */
+		/** The package where the asset should be created. If null, it will be put in the transient package. */
 		UObject* Parent = nullptr;
 
-		/** The name we want to give to the asset we will create */
+		/** The name to give to the created asset. */
 		FString AssetName = FString();
 
-		/** The base node that describe how to create the asset */
+		/** The base node that describes how to create the asset. */
 		UInterchangeFactoryBaseNode* AssetNode = nullptr;
 
-		/** The translator is use to retrieve the PayLoad data in case the factory need it */
+		/** The translator is used to retrieve the payLoad data in case the factory needs it. */
 		const UInterchangeTranslatorBase* Translator = nullptr;
 
 		/** The source data. TODO: we have to refactor UAssetImportData, the source data should be the base class for this now */
 		const UInterchangeSourceData* SourceData = nullptr;
 
-		/** The node container associate with the current source index */
+		/** The node container associated with the current source index. */
 		const UInterchangeBaseNodeContainer* NodeContainer = nullptr;
 
 		/**
-		 * If when we try to create the package we found out the asset already exist, this field will contain
-		 * the asset we want to re-import. The re-import should just change the source data and not any asset settings.
+		 * If when we try to create the package we find out the asset already exists, this field will contain
+		 * the asset we want to reimport. The reimport should just change the source data and not any asset settings.
 		 */
 		UObject* ReimportObject = nullptr;
 	};
 
 	struct FImportAssetResult
 	{
-		//If the factory set this to true, the interchange task import object should skip this asset
-		//Asset can be skip if it already exist and its not a re-import. Example, if we import a static mesh and we find an existing material.
-		//We don't want to override the material in this case since the UE material often contains game play graph.
+		//If the factory sets this to true, the Interchange task import object should skip this asset.
+		//An asset can be skipped if it already exists and the factory isn't doing a reimport. For example, if we import a static mesh and we find an existing material.
+		//We don't want to override the material in this case because the UE material often contains content related to gameplay.
 		bool bIsFactorySkipAsset = false;
-		//Return the UObject import or re-import by the factory, leave it to nullptr if there was an error.
+		//Return the UObject imported or reimported by the factory, or leave it set to nullptr if there was an error.
 		UObject* ImportedObject = nullptr;
 	};
 
@@ -127,13 +126,13 @@ public:
 	}
 
 	/**
-	 * Override it to import/re-import source data and configure the properties synchronously.
-	 * Create the asset package on the game thread since its not thread safe.
+	 * Override this function to import/reimport source data and configure the properties synchronously.
+	 * Create the asset package on the game thread because it's not thread-safe.
 	 *
-	 * @param Arguments - The structure containing all necessary arguments to import the UObject and the package, see the structure definition for the documentation.
-	 * @return the FImportAssetResult, see the structure to access the documentation.
+	 * @param Arguments - The structure that contains all necessary arguments to import the UObject and the package. See the structure definition for its documentation.
+	 * @return the FImportAssetResult. See the structure to access its documentation.
 	 *
-	 * @Note Mandatory to override this function to create the asset UObject package. Not mandatory for level actor, use CreateSceneObject instead.
+	 * @Note Mandatory to override this function to create the asset UObject package. Not mandatory for level actor; use CreateSceneObject instead.
 	 */
 	virtual FImportAssetResult BeginImportAsset_GameThread(const FImportAssetObjectParams& Arguments)
 	{
@@ -157,11 +156,11 @@ public:
 	}
 
 	/**
-	 * Override it to import/re-import the UObject source data and configure the properties asynchronously. Do not create the asset UObject asynchronously.
-	 * Helpful to get all the payload in parallel or do any long import task
+	 * Override this function to import/reimport the UObject source data and configure the properties asynchronously. Do not create the asset UObject asynchronously.
+	 * Helpful to get all the payloads in parallel or do any long import tasks.
 	 *
-	 * @param Arguments - The structure containing all necessary data to imported the UObject, see the structure definition for the documentation.
-	 * @return the FImportAssetResult, see the structure to access the documentation.
+	 * @param Arguments - The structure that contains all necessary data to import the UObject. See the structure definition for its documentation.
+	 * @return the FImportAssetResult. See the structure to access its documentation.
 	 */
 	virtual FImportAssetResult ImportAsset_Async(const FImportAssetObjectParams& Arguments)
 	{
@@ -173,11 +172,11 @@ public:
 	}
 
 	/**
-	 * Override it to end import/re-import on the game thread. Helpful if you need to display ui (i.e. re-import material conflict resolution) or if you
-	 * need to do anything not thread safe to complete the import/re-import.
+	 * Override this function to end import/reimport on the game thread. Helpful if you need to display UI (for example, reimport material conflict resolution) or if you
+	 * need to do anything not thread-safe to complete the import/reimport.
 	 *
-	 * @param Arguments - The structure containing all necessary data to imported the UObject, see the structure definition for the documentation.
-	 * @return the FImportAssetResult, see the structure to access the documentation.
+	 * @param Arguments - The structure that contains all necessary data to import the UObject. See the structure definition for its documentation.
+	 * @return the FImportAssetResult. See the structure to access its documentation.
 	 */
 	virtual FImportAssetResult EndImportAsset_GameThread(const FImportAssetObjectParams& Arguments)
 	{
@@ -192,32 +191,32 @@ public:
 	}
 
 	/**
-	 * Parameters to pass to SpawnActor function
+	 * Parameters to pass to the SpawnActor function.
 	 */
 	struct FImportSceneObjectsParams
 	{
-		/** The level in which to create the scene objects */
+		/** The level in which to create the scene objects. */
 		ULevel* Level = nullptr;
 
-		/** The name we want to give to the actor that we will create */
+		/** The name to give to the created actor. */
 		FString ObjectName;
 
-		/** The base node that describe how to create the asset */
+		/** The base node that describes how to create the asset. */
 		UInterchangeFactoryBaseNode* FactoryNode = nullptr;
 
-		/** The node container associated with the current source index */
+		/** The node container associated with the current source index. */
 		const UInterchangeBaseNodeContainer* NodeContainer = nullptr;
 
 		/** The source data. TODO: we have to refactor UAssetImportData, the source data should be the base class for this now */
 		const UInterchangeSourceData* SourceData = nullptr;
 
 		/**
-		 * If not null, the factory must perform a reimport of the scene node
+		 * If not null, the factory must perform a reimport of the scene node.
 		 */
 		UObject* ReimportObject = nullptr;
 	
 		/**
-		 * Factory base node associated with the reimported scene node
+		 * Factory base node associated with the reimported scene node.
 		 */
 		const UInterchangeFactoryBaseNode* ReimportFactoryNode = nullptr;
 	};
@@ -225,28 +224,28 @@ public:
 	/**
 	 * Creates the scene object from a Scene Node data.
 	 *
-	 * @param Arguments - The structure containing all necessary arguments, see the structure definition for the documentation.
-	 * @return The scene object or nullptr if the operation was unsuccessful.
+	 * @param Arguments - The structure that contains all necessary arguments. See the structure definition for its documentation.
+	 * @return The scene object, or nullptr if the operation was unsuccessful.
 	 */
 	virtual UObject* ImportSceneObject_GameThread(const FImportSceneObjectsParams& Arguments)
 	{
 		return nullptr;
 	}
 
-	/** Call when the user cancel the operation. */
+	/** Call when the user cancels the operation. */
 	virtual void Cancel() {}
 
 	/**
-	 * Parameters to pass to CreateAsset function
+	 * Parameters to pass to the CreateAsset function.
 	 */
 	struct FSetupObjectParams
 	{
-		/** The source data, mainly use to set the asset import data file. TODO: we have to refactor UAssetImportData, the source data should be the base class for this now */
+		/** The source data, mainly used to set the asset import data file. TODO: we have to refactor UAssetImportData. The source data should be the base class for this now. */
 		const UInterchangeSourceData* SourceData = nullptr;
 		UInterchangeFactoryBaseNode* FactoryNode = nullptr;
 
 
-		/** The UObject  we want to execute code on*/
+		/** The UObject we want to execute code on. */
 		UObject* ImportedObject = nullptr;
 		FString NodeUniqueID;
 		UInterchangeBaseNodeContainer* NodeContainer = nullptr;
@@ -259,8 +258,8 @@ public:
 	};
 
 	/*
-	 * Do any UObject setup required before the build (before PostEditChange) and after all dependency UObject have been imported.
-	 * @note - This function is called when starting the pre completion task (before PostEditChange is called for the asset).
+	 * Do any UObject setup required before the build (before PostEditChange) and after all dependency UObjects have been imported.
+	 * @note - This function is called when starting the pre-completion task (before PostEditChange is called for the asset).
 	 */
 	virtual void SetupObject_GameThread(const FSetupObjectParams& Arguments)
 	{
@@ -268,8 +267,8 @@ public:
 	}
 	
 	/*
-	 * Do any final UObject setup after the build (after PostEditChange)
-	 * @note - This function is called at the end of the pre completion task (after PostEditChange is called for the asset).
+	 * Do any final UObject setup after the build (after PostEditChange).
+	 * @note - This function is called at the end of the pre-completion task (after PostEditChange is called for the asset).
 	 */
 	virtual void FinalizeObject_GameThread(const FSetupObjectParams& Arguments)
 	{
@@ -286,7 +285,7 @@ public:
 	}
 
 	/**
-	 * Sets the object's source at the specified index to the given SourceFileName
+	 * Sets the object's source at the specified index to the given SourceFileName.
 	 */
 	virtual bool SetSourceFilename(const UObject* Object, const FString& SourceFilename, int32 SourceIndex) const
 	{
diff --git a/Engine/Source/Runtime/Interchange/Core/Public/InterchangePipelineBase.h b/Engine/Source/Runtime/Interchange/Core/Public/InterchangePipelineBase.h
--- a/Engine/Source/Runtime/Interchange/Core/Public/InterchangePipelineBase.h
+++ b/Engine/Source/Runtime/Interchange/Core/Public/InterchangePipelineBase.h
@@ -115,7 +115,7 @@ struct FInterchangePipelinePropertyStates
 	UPROPERTY(EditAnywhere, BlueprintReadWrite, Category = "Property States")
 	bool bLocked = false;
 
-	/** The property states for the import context */
+	/** The property states for the import context. */
 	UPROPERTY(EditAnywhere, BlueprintReadWrite, Category = "Context Properties States")
 	FInterchangePipelinePropertyStatePerContext BasicLayoutStates;
 
@@ -123,7 +123,7 @@ struct FInterchangePipelinePropertyStates
 	UPROPERTY(EditAnywhere, BlueprintReadWrite, Category = "Context Properties States")
 	FInterchangePipelinePropertyStatePerContext ImportStates;
 
-	/** The property states for the reimport context */
+	/** The property states for the reimport context. */
 	UPROPERTY(EditAnywhere, BlueprintReadWrite, Category = "Context Properties States")
 	FInterchangePipelinePropertyStatePerContext ReimportStates;
 };
@@ -198,7 +198,7 @@ public:
 	 */
 	UFUNCTION(BlueprintNativeEvent, BlueprintCallable, Category = "Interchange | Pipeline")
 	INTERCHANGECORE_API void ScriptedExecutePipeline(UInterchangeBaseNodeContainer* BaseNodeContainer, const TArray<UInterchangeSourceData*>& SourceDatas, const FString& ContentBasePath);
-	/** The default implementation (call if the blueprint do not have any implementation) will call the virtual ExecutePipeline */
+	/** The default implementation, which is called if the Blueprint does not have any implementation, calls the virtual ExecutePipeline(). */
 	void ScriptedExecutePipeline_Implementation(UInterchangeBaseNodeContainer* BaseNodeContainer, const TArray<UInterchangeSourceData*>& SourceDatas, const FString& ContentBasePath)
 	{
 		//By default we call the virtual import pipeline execution
@@ -206,26 +206,26 @@ public:
 	}
 
 	/**
-	 * ScriptedExecutePostFactoryPipeline is call after the factory create the unreal asset but before calling PostEditChange.
-	 * @note - the FTaskPreCompletion task is calling this function not the virtual one that is call by the default implementation.
+	 * ScriptedExecutePostFactoryPipeline is called after the factory creates an Unreal asset, but before it calls PostEditChange.
+	 * @note - The FTaskPreCompletion task calls this function, not the virtual one that is called by the default implementation.
 	 */
 	UFUNCTION(BlueprintNativeEvent, BlueprintCallable, Category = "Interchange | Pipeline")
 	INTERCHANGECORE_API void ScriptedExecutePostFactoryPipeline(const UInterchangeBaseNodeContainer* BaseNodeContainer, const FString& FactoryNodeKey, UObject* CreatedAsset, bool bIsAReimport);
-	/** The default implementation (call if the blueprint do not have any implementation) will call the virtual ExecutePostFactoryPipeline */
+	/** The default implementation, which is called if the Blueprint does not have any implementation, calls the virtual ExecutePostFactoryPipeline(). */
 	void ScriptedExecutePostFactoryPipeline_Implementation(const UInterchangeBaseNodeContainer* BaseNodeContainer, const FString& FactoryNodeKey, UObject* CreatedAsset, bool bIsAReimport)
 	{
 		//By default we call the virtual import pipeline execution
 		ExecutePostFactoryPipeline(BaseNodeContainer, FactoryNodeKey, CreatedAsset, bIsAReimport);
 	}
 	/**
-	 * ScriptedExecutePostImportPipeline is call after we completely import an asset, PostEditChange is already called.
-	 * This can be useful if you need an asset build data to finish the setup of another asset.
-	 * @example - PhysicsAsset need skeletal mesh render data to be build properly.
-	 * @note - the FTaskPostImport is calling this function not the virtual one that is call by the default implementation.
+	 * ScriptedExecutePostImportPipeline is called after an asset is completely imported, after PostEditChange has already been called.
+	 * This can be useful if you need build data for one asset to finish setting up another asset.
+	 * @example - A PhysicsAsset needs skeletal mesh render data to be built properly.
+	 * @note - the FTaskPostImport calls this function not the virtual one that is call by the default implementation.
 	 */
 	UFUNCTION(BlueprintNativeEvent, BlueprintCallable, Category = "Interchange | Pipeline")
 	INTERCHANGECORE_API void ScriptedExecutePostImportPipeline(const UInterchangeBaseNodeContainer* BaseNodeContainer, const FString& FactoryNodeKey, UObject* CreatedAsset, bool bIsAReimport);
-	/** The default implementation (call if the blueprint do not have any implementation) will call the virtual ExecutePostImportPipeline */
+	/** The default implementation, which is called if the Blueprint does not have any implementation, calls the virtual ExecutePostImportPipeline(). */
 	void ScriptedExecutePostImportPipeline_Implementation(const UInterchangeBaseNodeContainer* BaseNodeContainer, const FString& FactoryNodeKey, UObject* CreatedAsset, bool bIsAReimport)
 	{
 		//By default we call the virtual import pipeline execution
@@ -233,13 +233,13 @@ public:
 	}
 
 	/**
-	 * Non virtual helper to allow blueprint to implement event base function let the interchange know if it can run asynchronously.
-	 * the Interchange manager is calling this function not the virtual one that is call by the default implementation.
+	 * Non-virtual helper that allows Blueprint to implement an event-based function.
+	 * The Interchange manager calls this function, not the virtual one that is called by the default implementation.
 	 */
 	UFUNCTION(BlueprintCallable, BlueprintNativeEvent, Category = "Interchange | Pipeline")
 	INTERCHANGECORE_API void ScriptedExecuteExportPipeline(UInterchangeBaseNodeContainer* BaseNodeContainer);
 
-	/** The default implementation (call if the blueprint do not have any implementation) will call the virtual ExecuteExportPipeline */
+	/** The default implementation, which is called if the Blueprint does not have any implementation, calls the virtual ExecuteExportPipeline(). */
 	void ScriptedExecuteExportPipeline_Implementation(UInterchangeBaseNodeContainer* BaseNodeContainer)
 	{
 		ExecuteExportPipeline(BaseNodeContainer);
@@ -252,9 +252,9 @@ public:
 	}
 
 	/**
-	 * This function tell the interchange manager if we can execute this pipeline in async mode. If it return false, the ScriptedExecuteImportPipeline
-	 * will be call on the main thread (GameThread), if true it will be run in a background thread and possibly in parallel. If there is multiple
-	 * import process in same time.
+	 * This function tells the Interchange manager if this pipeline can be executed in async mode. If it returns false, the ScriptedExecuteImportPipeline
+	 * will be called on the main thread (GameThread). If it returns true, the pipeline will be run in a background thread and possibly in parallel if there are multiple
+	 * import processes at the same time.
 	 *
 	 */
 	INTERCHANGECORE_API virtual bool CanExecuteOnAnyThread(EInterchangePipelineTask PipelineTask)
@@ -263,24 +263,24 @@ public:
 	}
 
 	/**
-	 * Non virtual helper to allow blueprint to implement event base function.
-	 * the Interchange framework is calling this function not the virtual one that is called by the default implementation.
+	 * Non-virtual helper that allows Blueprint to implement an event-based function.
+	 * the Interchange framework calls this function, not the virtual one that is called by the default implementation.
 	 */
 	UFUNCTION(BlueprintNativeEvent, BlueprintCallable, Category = "Interchange | Pipeline")
 	INTERCHANGECORE_API void ScriptedSetReimportSourceIndex(UClass* ReimportObjectClass, const int32 SourceFileIndex);
 
-	/** The default implementation (call if the blueprint do not have any implementation) will call the virtual SetReimportContentFromSourceIndex */
+	/** The default implementation, which is called if the Blueprint does not have any implementation, calls the virtual SetReimportContentFromSourceIndex(). */
 	void ScriptedSetReimportSourceIndex_Implementation(UClass* ReimportObjectClass, const int32 SourceFileIndex)
 	{
 		SetReimportSourceIndex(ReimportObjectClass, SourceFileIndex);
 	}
 
 	/**
-	 * Non scripted class should return false here, we have the default to true because scripted class cannot override
-	 * this function since it can be call in a asynchronous thread, which python cannot be executed.
+	 * A non-scripted class should return false here. The default is set to true because scripted classes cannot override
+	 * this function since it can be called in an asynchronous thread, which is not possible for Python.
 	 *
-	 * We cannot call ScriptedCanExecuteOnAnyThread for a scripted python pipeline from the task parsing async thread.
-	 * This function allow us to not call it and force the ScriptedExecutePostImportPipeline to execute on the game thread.
+	 * We cannot call ScriptedCanExecuteOnAnyThread for a scripted Python pipeline from the task parsing async thread.
+	 * This function allows us to not call it and force the ScriptedExecutePostImportPipeline to execute on the game thread.
 	 */
 	virtual bool IsScripted()
 	{
@@ -292,12 +292,12 @@ public:
 	INTERCHANGECORE_API void SaveSettings(const FName PipelineStackName);
 
 	/**
-	 * This function is call before we show the pipeline dialog. Pipeline that override it can change the existing settings according to the re-import type.
-	 * The function is also call when we import or re-import custom LOD and alternate skinning.
+	 * This function is called before we show the pipeline dialog. Pipelines that override it can change the existing settings according to the reimport type.
+	 * The function is also called when we import or reimport custom LOD and alternate skinning.
 	 *
-	 * @Note - The function will set the context of the pipeline
-	 * @Param ReimportType - Tell pipeline what re-import type the user want to achieve.
-	 * @Param ReimportAsset - This is an optional parameter which is set when re-importing an asset.
+	 * @Note - The function will set the context of the pipeline.
+	 * @Param ReimportType - Tells the pipeline what reimport type the user wants to achieve.
+	 * @Param ReimportAsset - This is an optional parameter which is set when reimporting an asset.
 	 */
 	INTERCHANGECORE_API virtual void AdjustSettingsForContext(EInterchangePipelineContext ReimportType, TObjectPtr<UObject> ReimportAsset);
 	INTERCHANGECORE_API virtual void AdjustSettingsFromCache();
@@ -311,13 +311,13 @@ public:
 	}
 
 	/**
-	 * This function is called before showing the import dialog it is not called doing a re-import.
+	 * This function is called before showing the import dialog. It is not called when doing a reimport.
 	 */
 	virtual void PreDialogCleanup(const FName PipelineStackName) {}
 
 	/**
 	 * This function should return true if all the pipeline settings are in a valid state to start the import.
-	 * The Interchange Pipeline Configuration Dialog will call this to know if the Import button can be enable.
+	 * The Interchange Pipeline Configuration Dialog will call this to know if the Import button can be enabled.
 	 */
 	virtual bool IsSettingsAreValid(TOptional<FText>& OutInvalidReason) const
 	{
@@ -326,8 +326,8 @@ public:
 
 #if WITH_EDITOR
 	/**
-	 * Filter the pipeline properties from the translated data. This function is call by the import dialog after having duplicate and load the settings of the pipeline.
-	 * If some specific options change, the UI must refresh the filter, see function IsPropertyChangeNeedRefresh.
+	 * Filter the pipeline properties from the translated data. This function is called by the import dialog after having duplicated and loaded the settings of the pipeline.
+	 * If some specific options change, the UI must refresh the filter. See the IsPropertyChangeNeedRefresh function.
 	 */
 	INTERCHANGECORE_API virtual void FilterPropertiesFromTranslatedData(UInterchangeBaseNodeContainer* InBaseNodeContainer)
 	{
@@ -335,8 +335,8 @@ public:
 	}
 
 	/**
-	 * The import dialog will call this function when the user change a specific property, return true if the pipeline UI should be refresh.
-	 * A refresh will call the function FilterPropertiesFromTranslatedData.
+	 * The import dialog will call this function when the user changes a specific property. Return true if the pipeline UI should be refreshed.
+	 * A refresh will call the FilterPropertiesFromTranslatedData function.
 	 */
 	INTERCHANGECORE_API virtual bool IsPropertyChangeNeedRefresh(const FPropertyChangedEvent& PropertyChangedEvent)
 	{
@@ -384,29 +384,29 @@ public:
 	}
 
 	/**
-	 * Return a const property states pointer. Return nullptr if the property do not exist
+	 * Return a const property states pointer. Return nullptr if the property does not exist.
 	 */
 	INTERCHANGECORE_API const FInterchangePipelinePropertyStates* GetPropertyStates(const FName PropertyPath) const;
 
 	/**
-	 * Return a mutable property states pointer. Return nullptr if the property do not exist
+	 * Return a mutable property states pointer. Return nullptr if the property does not exist.
 	 */
 	INTERCHANGECORE_API FInterchangePipelinePropertyStates* GetMutablePropertyStates(const FName PropertyPath);
 
 	/**
-	 * Return true if the property has valid states, false if no states was set for the property
+	 * Return true if the property has valid states, or false if no states were set for the property.
 	 */
 	UFUNCTION(BlueprintCallable, Category = "Interchange | Pipeline")
 	INTERCHANGECORE_API bool DoesPropertyStatesExist(const FName PropertyPath) const;
 
 	/**
-	 * Return a mutable property states Reference. Add the property states if it doesnt exist.
+	 * Return a mutable property states reference. Add the property states if it doesn't exist.
 	 */
 	UFUNCTION(BlueprintCallable, Category = "Interchange | Pipeline")
 	INTERCHANGECORE_API FInterchangePipelinePropertyStates& FindOrAddPropertyStates(const FName PropertyPath);
 
 	/**
-	 * Return the property name of the properties states map
+	 * Return the property name of the properties states map.
 	 */
 	static INTERCHANGECORE_API FName GetPropertiesStatesPropertyName();
 
@@ -429,10 +429,10 @@ public:
 #endif //WITH_EDITOR
 
 	/*
-	 * Return true if the pipeline is not a sub pipeline. Stand Alone Pipeline can have some extra properties.
-	 * The customize detail of the pipeline will call this function to show or hide those property.
-	 * The pipeline itself can call the function to know if a property must be use or not.
-	 * Example: the override asset name property
+	 * Return true if the pipeline is not a sub-pipeline. Standalone pipelines can have some extra properties.
+	 * The customize detail of the pipeline will call this function to show or hide those properties.
+	 * The pipeline itself can call the function to know whether a property must be used.
+	 * Example: the override asset name property.
 	 */
 	bool IsStandAlonePipeline()
 	{
@@ -440,7 +440,7 @@ public:
 	};
 
 	/*
-	* Update/Re-set the TWeakObjectPtrs to point to their original TObjectPtr
+	* Update/reset the TWeakObjectPtrs to point to their original TObjectPtr.
 	* Works based on Variable naming.
 	*/
 	INTERCHANGECORE_API void UpdateWeakObjectPtrs();
@@ -450,15 +450,15 @@ public:
 	/** End UObject overrides */
 
 	/*
-	 * If it return true it mean we want to save this pipeline in the asset import data.
-	 * We will reuse this pipeline when re-importing the asset.
-	 * If false its probably a debug helper pipeline we do not want to save into assets
+	 * If this returns true, this pipeline will be saved in the asset import data.
+	 * We will reuse this pipeline when reimporting the asset.
+	 * If false, it's probably a debug helper pipeline that we do not want to save into assets.
 	 */
 	INTERCHANGECORE_API virtual bool SupportReimport() const { return true; }
 
 	/*
-	 * The Unreal import system have an option to force a name if we import only one main asset (one texture, one mesh or one animation).
-	 * The generic asset pipeline is using this information to behave like expected.
+	 * The Unreal import system has an option to force a name if we import only one main asset (one texture, one mesh or one animation).
+	 * The generic asset pipeline uses this information to behave as expected.
 	 */
 	FString DestinationName;
 
@@ -471,17 +471,17 @@ protected:
 	}
 
 	/**
-	 * This function is called after the factory create an asset, but before PostEditChange is called.
-	 * It is always call on the game thread, code inside this function should not wait on other thread. (do not stall the game thread)
-	 * This is the place to do any change to the unreal asset before PostEditChange is call.
+	 * This function is called after the factory creates an asset, but before PostEditChange is called.
+	 * It is always called on the game thread. Code inside this function should not wait on other threads. Do not stall the game thread.
+	 * This is the place to make any changes to the Unreal asset before PostEditChange is called.
 	 */
 	virtual void ExecutePostFactoryPipeline(const UInterchangeBaseNodeContainer* BaseNodeContainer, const FString& NodeKey, UObject* CreatedAsset, bool bIsAReimport)
 	{
 	}
 
 	/**
-	 * This function is called after the unreal asset is completely import and PostEditChange was called.
-	 * @Note: Some unreal assets have asynchronous build, its possible they are still compiling.
+	 * This function is called after the Unreal asset is completely imported and PostEditChange was called.
+	 * @Note: Some Unreal assets have asynchronous build operations. It's possible they are still compiling.
 	 */
 	virtual void ExecutePostImportPipeline(const UInterchangeBaseNodeContainer* BaseNodeContainer, const FString& NodeKey, UObject* CreatedAsset, bool bIsAReimport)
 	{
@@ -510,19 +510,19 @@ protected:
 #endif //WITH_EDITOR
 	
 	/**
-	 * If true, the property editor for this pipeline instance will allow properties states edition.
-	 * If false, the property editor for this pipeline instance will apply the properties states.
+	 * If true, the property editor for this pipeline instance will allow editing property states.
+	 * If false, the property editor for this pipeline instance will apply the property states.
 	 *
-	 * Note: If you open in the content browser a pipeline asset you will be able to edit the properties states.
-	 *       If you import a file with interchange, the import dialog will apply properties states.
+	 * Note: If you open a pipeline asset in the Content Browser, you will be able to edit the property states.
+	 *       If you import a file with Interchange, the import dialog will apply property states.
 	 */
 	bool bAllowPropertyStatesEdition = true;
 
 	/**
-	 * If true, this pipeline instance is use for reimport.
-	 * If false, this pipeline instance is use for import.
+	 * If true, this pipeline instance is used for reimport.
+	 * If false, this pipeline instance is used for import.
 	 *
-	 * Note: This context must be set by the owner instancing this pipeline. This context will be use to hide or not some properties.
+	 * Note: This context must be set by the owner instancing this pipeline. This context will be used to determine whether to hide some properties.
 	 */
 	bool bIsReimportContext = false;
 
@@ -539,7 +539,7 @@ protected:
 
 	/**
 	 * Map of property path and lock status. Any properties that have a true lock status will be readonly when showing the import dialog.
-	 * Use the API to Get and Set the properties states
+	 * Use the API to Get and Set the properties states.
 	*/
 	UPROPERTY()
 	TMap<FName, FInterchangePipelinePropertyStates> PropertiesStates;
diff --git a/Engine/Source/Runtime/Interchange/Core/Public/Nodes/InterchangeBaseNode.h b/Engine/Source/Runtime/Interchange/Core/Public/Nodes/InterchangeBaseNode.h
--- a/Engine/Source/Runtime/Interchange/Core/Public/Nodes/InterchangeBaseNode.h
+++ b/Engine/Source/Runtime/Interchange/Core/Public/Nodes/InterchangeBaseNode.h
@@ -35,20 +35,20 @@ class UInterchangeBaseNode;
 struct FFrame;
 
 /**
- * Internal Helper to get set custom property for class that derive from UInterchangeBaseNode. This is use by the macro IMPLEMENT_UOD_ATTRIBUTE.
+ * Internal Helper to get set custom property for class that derive from UInterchangeBaseNode. This is used by the macro IMPLEMENT_UOD_ATTRIBUTE.
  */
 namespace InterchangePrivateNodeBase
 {
 	/**
-	 * Retrieve a custom attribute if the attribute exist
+	 * Retrieve a custom attribute if the attribute exists.
 	 *
-	 * @param Attributes - The attribute storage you want to query the custom attribute
-	 * @param AttributeKey - The storage key for the attribute
-	 * @param OperationName - The name of the operation in case there is an error
-	 * @param OutAttributeValue - This is where we store the value we retrieve from the storage
+	 * @param Attributes - The attribute storage you want to query for the custom attribute.
+	 * @param AttributeKey - The storage key for the attribute.
+	 * @param OperationName - The name of the operation, in case there is an error.
+	 * @param OutAttributeValue - This is where we store the value we retrieve from the storage.
 	 *
-	 * @return - return true if the attribute exist in the storage and was query without error.
-	 *           return false if the attribute do not exist or there is an error retriving it from the Storage
+	 * @return - returns true if the attribute exists in the storage and was queried without error.
+	 *           returns false if the attribute does not exist or there is an error retrieving it from the storage.
 	 */
 	template<typename ValueType>
 	bool GetCustomAttribute(const UE::Interchange::FAttributeStorage& Attributes, const UE::Interchange::FAttributeKey& AttributeKey, const FString& OperationName, ValueType& OutAttributeValue)
@@ -72,12 +72,12 @@ namespace InterchangePrivateNodeBase
 	}
 
 	/**
-	 * Add or update a custom attribute value in the specified storage
+	 * Add or update a custom attribute value in the specified storage.
 	 *
-	 * @param Attributes - The attribute storage you want to add or update the custom attribute
-	 * @param AttributeKey - The storage key for the attribute
-	 * @param OperationName - The name of the operation in case there is an error
-	 * @param AttributeValue - The value we want to add or update in the storage
+	 * @param Attributes - The attribute storage in which you want to add or update the custom attribute.
+	 * @param AttributeKey - The storage key for the attribute.
+	 * @param OperationName - The name of the operation, in case there is an error.
+	 * @param AttributeValue - The value to add or update in the storage.
 	 */
 	template<typename ValueType>
 	bool SetCustomAttribute(UE::Interchange::FAttributeStorage& Attributes, const UE::Interchange::FAttributeKey& AttributeKey, const FString& OperationName, const ValueType& AttributeValue)
@@ -103,18 +103,18 @@ namespace InterchangePrivateNodeBase
 }
 
 /**
- * Use those macro to create Get/Set/ApplyToAsset for an attribute you want to support in your custom UOD node that derive from UInterchangeBaseNode.
- * The attribute key will be declare under private access modifier and the functions are declare under public access modifier.
- * So the class access modifier is public after calling this macro.
+ * Use these macros to create Get/Set/ApplyToAsset for an attribute you want to support in your custom UOD node that derive from UInterchangeBaseNode.
+ * The attribute key will be declared under the private access modifier, and the functions are declared under the public access modifier.
+ * The class access modifier is public after calling this macro.
  *
  * @note - The Get will return false if the attribute was never set.
  * @note - The Set will add the attribute to the storage or update the value if the storage already has this attribute.
- * @note - The Apply will set a member variable of the AssetType instance. It will apply the value only if the storage contain the attribute key
+ * @note - The Apply will set a member variable of the AssetType instance. It will apply the value only if the storage contains the attribute key
  *
- * @param AttributeName - The name of the Get/Set functions. Example if you pass Foo you will end up with GetFoo and SetFoo function
- * @param AttributeType - This is to specify the type of the attribute. bool, float, FString... anything supported by the FAttributeStorage
- * @param AssetType - This is the asset you want to apply the storage value"
- * @param EnumType - Optional, specify it only if the AssetType member is an enum so we can type cast it in the apply function (we use uint8 to store the enum value)"
+ * @param AttributeName - The name of the Get/Set functions. For example, if you pass Foo you will end up with GetFoo and SetFoo functions.
+ * @param AttributeType - This is to specify the type of the attribute: bool, float, FString, or anything else supported by the FAttributeStorage.
+ * @param AssetType - This is the asset you want to apply the storage value.
+ * @param EnumType - Optional. Specify only if the AssetType member is an enum so we can type cast it in the apply function (we use uint8 to store the enum value).
  */
 #define IMPLEMENT_NODE_ATTRIBUTE_KEY(AttributeName)																				\
 const UE::Interchange::FAttributeKey Macro_Custom##AttributeName##Key = UE::Interchange::FAttributeKey(TEXT(#AttributeName));
@@ -148,8 +148,8 @@ const UE::Interchange::FAttributeKey Macro_Custom##AttributeName##Key = UE::Inte
 namespace UE::Interchange
 {
 	/**
-		* Helper struct use to declare static const data we use in the UInterchangeBaseNode
-		* Node that derive from UInterchangeBaseNode can also add a struct that derive from this one to add there static data
+		* Helper struct used to declare static const data we use in the UInterchangeBaseNode.
+		* Nodes that derive from UInterchangeBaseNode can also add a struct that derives from this one to add static data.
 		* @note: The static data are mainly for holding Attribute keys. All attributes that are always available for a node should be in this class or a derived class.
 		*/
 	struct FBaseNodeStaticData
@@ -179,13 +179,13 @@ UENUM(BlueprintType)
 enum class EInterchangeNodeUserInterfaceContext : uint8
 {
 	None,
-	Preview //When we want to preview the node data, a preview is normaly read only and should not show internal data
+	Preview //When we want to preview the node data. A preview is normally read-only and should not show internal data.
 };
 
 /**
- * This struct is used to store and retrieve key value attributes. The attributes are store in a generic FAttributeStorage which serialize the value in a TArray64<uint8>
- * See UE::Interchange::EAttributeTypes to know the supported template types
- * This is an abstract class. This is the base class of the interchange node graph format, all class in this format should derive from this class
+ * This struct is used to store and retrieve key-value attributes. The attributes are stored in a generic FAttributeStorage that serializes the values in a TArray64<uint8>.
+ * See UE::Interchange::EAttributeTypes to know the supported template types.
+ * This is an abstract class. This is the base class of the interchange node graph format; all classes in this format should derive from this class.
  */
 UCLASS(BlueprintType, MinimalAPI)
 class UInterchangeBaseNode : public UObject
@@ -198,37 +198,37 @@ public:
 	INTERCHANGECORE_API UInterchangeBaseNode();
 
 	/**
-	 * Initialize the base data of the node
-	 * @param UniqueID - The uniqueId for this node
-	 * @param DisplayLabel - The name of the node
+	 * Initialize the base data of the node.
+	 * @param UniqueID - The unique ID for this node.
+	 * @param DisplayLabel - The name of the node.
 	 *
 	 */
 	UFUNCTION(BlueprintCallable, Category = "Interchange | Node")
 	INTERCHANGECORE_API void InitializeNode(const FString& UniqueID, const FString& DisplayLabel, const EInterchangeNodeContainerType NodeContainerType);
 
 	/**
-	 * Return the node type name of the class, we use this when reporting error
+	 * Return the node type name of the class. This is used when reporting errors.
 	 */
 	INTERCHANGECORE_API virtual FString GetTypeName() const;
 
 	/**
-	 * Icon name, use to retrieve the brush when we display the node in any UI
+	 * Icon name, used to retrieve the brush when we display the node in any UI.
 	 */
 	INTERCHANGECORE_API virtual FName GetIconName() const;
 
 #if WITH_EDITOR
 	/**
-	 * UI that inspect node attribute call this to give a readable name to attribute key
+	 * UI that inspects node attributes calls this to give a readable name to an attribute key.
 	 */
 	INTERCHANGECORE_API virtual FString GetKeyDisplayName(const UE::Interchange::FAttributeKey& NodeAttributeKey) const;
 
 	/**
-	 * UI that inspect node attribute call this to display or hide an attribute
+	 * UI that inspects node attributes calls this to display or hide an attribute.
 	 */
 	INTERCHANGECORE_API virtual bool ShouldHideAttribute(const UE::Interchange::FAttributeKey& NodeAttributeKey) const;
 
 	/**
-	 * UI that inspect node attribute call this to display the attribute under the returned category
+	 * UI that inspects node attributes calls this to display the attribute under the returned category.
 	 */
 	INTERCHANGECORE_API virtual FString GetAttributeCategory(const UE::Interchange::FAttributeKey& NodeAttributeKey) const;
 
@@ -236,15 +236,15 @@ public:
 
 #if WITH_EDITORONLY_DATA
 	/**
-	 * Temporary property set by UI to have the context for ShouldHideAttribute. We use this because property editor cannot add this custom context when it call CustomizeDetails
+	 * Temporary property set by UI to have the context for ShouldHideAttribute. We use this because the property editor cannot add this custom context when it calls CustomizeDetails.
 	 */
 	UPROPERTY(Transient, DuplicateTransient)
 	EInterchangeNodeUserInterfaceContext UserInterfaceContext = EInterchangeNodeUserInterfaceContext::None;
 #endif
 
 	/**
-	 * Add an attribute to the node
-	 * @param NodeAttributeKey - The key of the attribute
+	 * Add an attribute to the node.
+	 * @param NodeAttributeKey - The key of the attribute.
 	 * @param Value - The attribute value.
 	 *
 	 */
@@ -262,24 +262,24 @@ public:
 	}
 
 	/**
-	 * Return true if the node contain an attribute with the specified Key
-	 * @param nodeAttributeKey - The key of the searched attribute
+	 * Return true if the node contains an attribute with the specified Key.
+	 * @param nodeAttributeKey - The key of the searched attribute.
 	 *
 	 */
 	INTERCHANGECORE_API virtual bool HasAttribute(const UE::Interchange::FAttributeKey& NodeAttributeKey) const;
 
 	/**
-	 * This function return an attribute type for the specified Key. Return type None if the key is invalid
+	 * This function returns an attribute type for the specified Key. Returns type None if the key is invalid.
 	 *
-	 * @param NodeAttributeKey - The key of the attribute
+	 * @param NodeAttributeKey - The key of the attribute.
 	 */
 	INTERCHANGECORE_API virtual UE::Interchange::EAttributeTypes GetAttributeType(const UE::Interchange::FAttributeKey& NodeAttributeKey) const;
 
 	/**
-	 * This function return an  attribute handle for the specified Key.
-	 * If there is an issue with the KEY or storage the method will trip a check, always make sure you have a valid key before calling this
+	 * This function returns an attribute handle for the specified Key.
+	 * If there is an issue with the KEY or storage, the method will trip a check. Always make sure you have a valid key before calling this.
 	 * 
-	 * @param NodeAttributeKey - The key of the attribute
+	 * @param NodeAttributeKey - The key of the attribute.
 	 */
 	template<typename T>
 	UE::Interchange::FAttributeStorage::TAttributeHandle<T> GetAttributeHandle(const UE::Interchange::FAttributeKey& NodeAttributeKey) const
@@ -288,108 +288,108 @@ public:
 	}
 
 	/**
-	 * Use this function to query all the node attribute keys
+	 * Query all the node attribute keys.
 	 */
 	INTERCHANGECORE_API void GetAttributeKeys(TArray<UE::Interchange::FAttributeKey>& AttributeKeys) const;
 
 	/**
-	 * Remove any attribute from this node. Return false if we cannot remove it. If the attribute do not exist it will return true.
+	 * Remove the specified attribute from this node. Returns false if it could not be removed. If the attribute does not exist, returns true.
 	 */
 	UFUNCTION(BlueprintCallable, Category = "Interchange | Node")
 	INTERCHANGECORE_API bool RemoveAttribute(const FString& NodeAttributeKey);
 
 	/**
-	 * Add a boolean attribute to this node. Return false if the attribute do not exist or if we cannot add it
+	 * Add a Boolean attribute to this node. Returns false if the attribute does not exist or if it cannot be added.
 	 */
 	UFUNCTION(BlueprintCallable, Category = "Interchange | Node")
 	INTERCHANGECORE_API bool AddBooleanAttribute(const FString& NodeAttributeKey, const bool& Value);
 	
 	/**
-	 * Get a boolean attribute from this node. Return false if the attribute do not exist
+	 * Get a Boolean attribute from this node. Returns false if the attribute does not exist.
 	 */
 	UFUNCTION(BlueprintCallable, Category = "Interchange | Node")
 	INTERCHANGECORE_API bool GetBooleanAttribute(const FString& NodeAttributeKey, bool& OutValue) const;
 
 	/**
-	 * Add a int32 attribute to this node. Return false if the attribute do not exist or if we cannot add it
+	 * Add a int32 attribute to this node. Returns false if the attribute does not exist or if it cannot be added.
 	 */
 	UFUNCTION(BlueprintCallable, Category = "Interchange | Node")
 	INTERCHANGECORE_API bool AddInt32Attribute(const FString& NodeAttributeKey, const int32& Value);
 	
 	/**
-	 * Get a int32 attribute from this node. Return false if the attribute do not exist
+	 * Get a int32 attribute from this node. Returns false if the attribute does not exist.
 	 */
 	UFUNCTION(BlueprintCallable, Category = "Interchange | Node")
 	INTERCHANGECORE_API bool GetInt32Attribute(const FString& NodeAttributeKey, int32& OutValue) const;
 
 	/**
-	 * Add a float attribute to this node. Return false if the attribute do not exist or if we cannot add it
+	 * Add a float attribute to this node. Returns false if the attribute does not exist or if it cannot be added.
 	 */
 	UFUNCTION(BlueprintCallable, Category = "Interchange | Node")
 	INTERCHANGECORE_API bool AddFloatAttribute(const FString& NodeAttributeKey, const float& Value);
 	
 	/**
-	 * Get a float attribute from this node. Return false if the attribute do not exist
+	 * Get a float attribute from this node. Returns false if the attribute does not exist.
 	 */
 	UFUNCTION(BlueprintCallable, Category = "Interchange | Node")
 	INTERCHANGECORE_API bool GetFloatAttribute(const FString& NodeAttributeKey, float& OutValue) const;
 
 	/**
-	 * Add a double attribute to this node. Return false if the attribute do not exist or if we cannot add it
+	 * Add a double attribute to this node. Returns false if the attribute does not exist or if it cannot be added.
 	 */
 	UFUNCTION(BlueprintCallable, Category = "Interchange | Node")
 	INTERCHANGECORE_API bool AddDoubleAttribute(const FString& NodeAttributeKey, const double& Value);
 	
 	/**
-	 * Get a double attribute from this node. Return false if the attribute do not exist
+	 * Get a double attribute from this node. Returns false if the attribute does not exist.
 	 */
 	UFUNCTION(BlueprintCallable, Category = "Interchange | Node")
 	INTERCHANGECORE_API bool GetDoubleAttribute(const FString& NodeAttributeKey, double& OutValue) const;
 
 	/**
-	 * Add a string attribute to this node. Return false if the attribute do not exist or if we cannot add it
+	 * Add a string attribute to this node. Returns false if the attribute does not exist or if it cannot be added.
 	 */
 	UFUNCTION(BlueprintCallable, Category = "Interchange | Node")
 	INTERCHANGECORE_API bool AddStringAttribute(const FString& NodeAttributeKey, const FString& Value);
 	
 	/**
-	 * Get a string attribute from this node. Return false if the attribute do not exist
+	 * Get a string attribute from this node. Returns false if the attribute does not exist.
 	 */
 	UFUNCTION(BlueprintCallable, Category = "Interchange | Node")
 	INTERCHANGECORE_API bool GetStringAttribute(const FString& NodeAttributeKey, FString& OutValue) const;
 
 	/**
-	 * Add a guid attribute to this node. Return false if the attribute do not exist or if we cannot add it
+	 * Add a GUID attribute to this node. Returns false if the attribute does not exist or if it cannot be added.
 	 */
 	UFUNCTION(BlueprintCallable, Category = "Interchange | Node")
 	INTERCHANGECORE_API bool AddGuidAttribute(const FString& NodeAttributeKey, const FGuid& Value);
 
 	/**
-	 * Get a guid attribute from this node. Return false if the attribute do not exist
+	 * Get a GUID attribute from this node. Returns false if the attribute does not exist.
 	 */
 	UFUNCTION(BlueprintCallable, Category = "Interchange | Node")
 	INTERCHANGECORE_API bool GetGuidAttribute(const FString& NodeAttributeKey, FGuid& OutValue) const;
 
 	/**
-	 * Add a FLinearColor attribute to this node. Return false if the attribute do not exist or if we cannot add it
+	 * Add an FLinearColor attribute to this node. Returns false if the attribute does not exist or if it cannot be added.
 	 */
 	UFUNCTION(BlueprintCallable, Category = "Interchange | Node")
 	INTERCHANGECORE_API bool AddLinearColorAttribute(const FString& NodeAttributeKey, const FLinearColor& Value);
 	
 	/**
-	 * Get a FLinearColor attribute from this node. Return false if the attribute do not exist
+	 * Get an FLinearColor attribute from this node. Returns false if the attribute does not exist.
 	 */
 	UFUNCTION(BlueprintCallable, Category = "Interchange | Node")
 	INTERCHANGECORE_API bool GetLinearColorAttribute(const FString& NodeAttributeKey, FLinearColor& OutValue) const;
 
 	/**
-	 * Add a Vector2 attribute to this node. Return false if the attribute do not exist or if we cannot add it
+	 * Add a Vector2 attribute to this node. Returns false if the attribute does not exist or if it cannot be added.
 	 */
 	UFUNCTION(BlueprintCallable, Category = "Interchange | Node")
 	INTERCHANGECORE_API bool AddVector2Attribute(const FString& NodeAttributeKey, const FVector2f& Value);
 
 	/**
-	 * Get a Vector2 attribute from this node. Return false if the attribute do not exist
+	 * Get a Vector2 attribute from this node. Returns false if the attribute does not exist.
 	 */
 	UFUNCTION(BlueprintCallable, Category = "Interchange | Node")
 	INTERCHANGECORE_API bool GetVector2Attribute(const FString& NodeAttributeKey, FVector2f& OutValue) const;
@@ -417,7 +417,7 @@ public:
 	}
 
 	/**
-	 * Return the unique id pass in the constructor.
+	 * Return the unique ID passed in the constructor.
 	 */
 	UFUNCTION(BlueprintCallable, Category = "Interchange | Node")
 	INTERCHANGECORE_API FString GetUniqueID() const;
@@ -435,90 +435,90 @@ public:
 	INTERCHANGECORE_API bool SetDisplayLabel(const FString& DisplayName);
 
 	/**
-	 * Return the parent unique id. In case the attribute does not exist it will return InvalidNodeUid()
+	 * Return the parent unique ID. If the attribute does not exist, returns InvalidNodeUid().
 	 */
 	UFUNCTION(BlueprintCallable, Category = "Interchange | Node")
 	INTERCHANGECORE_API FString GetParentUid() const;
 
 	/**
-	 * Set the parent unique id.
+	 * Set the parent unique ID.
 	 */
 	UE_DEPRECATED(5.1, "This function should not be use - Use UInterchangeBaseNodeContainer::SetNodeParentUid function to create a parent hierarchy. This will ensure the container cache is always in sync")
 	UFUNCTION(BlueprintCallable, Category = "Interchange | Node")
 	INTERCHANGECORE_API bool SetParentUid(const FString& ParentUid);
 
 	/**
-	 * Get number of target assets relating to this object.
+	 * Get the number of target assets relating to this object.
 	 */
 	UFUNCTION(BlueprintCallable, Category = "Interchange | Node")
 	INTERCHANGECORE_API int32 GetTargetNodeCount() const;
 
 	/**
-	 * Get target assets relating to this object.
+	 * Get the target assets relating to this object.
 	 */
 	UFUNCTION(BlueprintCallable, Category = "Interchange | Node")
 	INTERCHANGECORE_API void GetTargetNodeUids(TArray<FString>& OutTargetAssets) const;
 
 	/**
-	 * Add asset node UID relating to this object.
+	 * Add an asset node UID relating to this object.
 	 */
 	UFUNCTION(BlueprintCallable, Category = "Interchange | Node")
 	INTERCHANGECORE_API bool AddTargetNodeUid(const FString& AssetUid) const;
 
 	/**
-	 * Remove asset node UID relating to this object.
+	 * Remove an asset node UID relating to this object.
 	 */
 	UFUNCTION(BlueprintCallable, Category = "Interchange | Node")
 	INTERCHANGECORE_API bool RemoveTargetNodeUid(const FString& AssetUid) const;
 
 	/**
-	 * IsEnable true mean that the node will be import/export, if false it will be discarded.
-	 * Return false if this node was disabled. Return true if the attribute is not there or if it was enabled.
+	 * If true, the node is imported or exported. If false, it is discarded.
+	 * Returns false if the node was disabled. Returns true if the attribute is not there or if it was enabled.
 	 */
 	UFUNCTION(BlueprintCallable, Category = "Interchange | Node")
 	INTERCHANGECORE_API bool IsEnabled() const;
 
 	/**
-	 * Set the IsEnable attribute to determine if this node should be part of the import/export process
-	 * @param bIsEnabled - The enabled state we want to set this node. True will import/export the node, fals will not.
-	 * @return true if it was able to set the attribute, false otherwise.
+	 * Determine whether this node should be part of the import or export process.
+	 * @param bIsEnabled - If true, this node is imported or exported. If false, it is discarded.
+	 * @return true if the attribute was set, or false otherwise.
 	 */
 	UFUNCTION(BlueprintCallable, Category = "Interchange | Node")
 	INTERCHANGECORE_API bool SetEnabled(const bool bIsEnabled);
 
 	/**
-	 * Return the node container type which define the purpose of the node (Factory node, translated scene node or translated asset node).
+	 * Return the node container type that defines the purpose of the node (factory node, translated scene node, or translated asset node).
 	 */
 	UFUNCTION(BlueprintCallable, Category = "Interchange | Node")
 	INTERCHANGECORE_API EInterchangeNodeContainerType GetNodeContainerType() const;
 
 	/**
-	 * Return a FGuid build from the FSHA1 of all the attribute data contain in the node.
+	 * Return an FGuid built from the FSHA1 of all the attribute data contained in the node.
 	 *
 	 * @note the attribute are sorted by key when building the FSHA1 data. The hash will be deterministic for the same data whatever
 	 * the order we add the attributes.
-	 * This function interface is pure virtual
+	 * This function interface is pure virtual.
 	 */
 	INTERCHANGECORE_API virtual FGuid GetHash() const;
 
 	/**
-	 * Optional, Any node that can import/export an asset should set the proper name we will give to the asset.
-	 * If the attribute was never set, it will return GetDisplayLabel.
+	 * Optional. Any node that can import or export an asset should set the desired name for the asset.
+	 * If the attribute was never set, returns GetDisplayLabel().
 	 */
 	UFUNCTION(BlueprintCallable, Category = "Interchange | Node")
 	INTERCHANGECORE_API virtual FString GetAssetName() const;
 
 	/**
-	 * Set the name we want for the imported asset this node represent. The asset factory will call GetAssetName()
+	 * Set the name for the imported asset this node represents. The asset factory will call GetAssetName().
 	 */
 	UFUNCTION(BlueprintCallable, Category = "Interchange | Node")
 	INTERCHANGECORE_API virtual bool SetAssetName(const FString& AssetName);
 
-	/** Return the invalid unique ID */
+	/** Return the invalid unique ID. */
 	static INTERCHANGECORE_API FString InvalidNodeUid();
 
 	/**
-	 * Serialize the node, by default only the attribute storage is serialize for a node
+	 * Serialize the node. By default, only the attribute storage is serialized for a node.
 	 */
 	INTERCHANGECORE_API virtual void Serialize(FArchive& Ar) override;
 
diff --git a/Engine/Source/Runtime/Interchange/Core/Public/Nodes/InterchangeBaseNodeContainer.h b/Engine/Source/Runtime/Interchange/Core/Public/Nodes/InterchangeBaseNodeContainer.h
--- a/Engine/Source/Runtime/Interchange/Core/Public/Nodes/InterchangeBaseNodeContainer.h
+++ b/Engine/Source/Runtime/Interchange/Core/Public/Nodes/InterchangeBaseNodeContainer.h
@@ -28,10 +28,11 @@ template <typename T> struct TObjectPtr;
 
 
 /**
- * Interchange UInterchangeBaseNode graph is a format used to feed asset/scene import/reimport/export factories/writer.
- * This container hold a flat list of all nodes that have been translated from the source data.
- *
- * Translators are filling this container and the Import/Export managers are reading it to execute the import/export process
+ * The Interchange UInterchangeBaseNode graph is a format used to feed factories and writers when they import, reimport, and
+ * export an asset or scene.
+ * 
+ * This container holds a flat list of all nodes that have been translated from the source data.
+ * Translators fill this container, and the import/export managers read it to execute the import/export process.
  */
  UCLASS(BlueprintType, MinimalAPI)
 class UInterchangeBaseNodeContainer : public UObject
@@ -42,10 +43,6 @@ public:
 
 	/**
 	 * Empty the container.
-	 *
-	 * @param Node - a pointer on the node you want to add
-	 * @return: return the node unique ID of the added item. If the node already exist it will return the existing ID. Return InvalidNodeUid if the node cannot be added.
-	 *
 	 */
 	UFUNCTION(BlueprintCallable, Category = "Interchange | Node Container")
 	INTERCHANGECORE_API void Reset()
@@ -55,7 +52,7 @@ public:
 	}
 
 	/**
-	 * Add a node in the container, the node will be add into a TMap.
+	 * Add a node to the container. The node is added into a TMap.
 	 *
 	 * @param Node - a pointer on the node you want to add
 	 * @return: return the node unique ID of the added item. If the node already exist it will return the existing ID. Return InvalidNodeUid if the node cannot be added.
@@ -67,11 +64,11 @@ public:
 	UFUNCTION(BlueprintCallable, Category = "Interchange | Node Container")
 	INTERCHANGECORE_API void ReplaceNode(const FString& NodeUniqueID, UInterchangeFactoryBaseNode* NewNode);
 
-	/** Return true if the node unique ID exist in the container */
+	/** Return true if the node unique ID exists in the container. */
 	UFUNCTION(BlueprintCallable, Category = "Interchange | Node Container")
 	INTERCHANGECORE_API bool IsNodeUidValid(const FString& NodeUniqueID) const;
 
-	/** Unordered iteration of the all nodes */
+	/** Unordered iteration of the all nodes. */
 	INTERCHANGECORE_API void IterateNodes(TFunctionRef<void(const FString&, UInterchangeBaseNode*)> IterationLambda) const;
 
 	template <typename T>
@@ -88,7 +85,7 @@ public:
 
 	/**
 	 * Recursively traverse the hierarchy starting with the specified node unique ID.
-	 * If the iteration lambda return true the iteration will stop. If the iteration lambda return false the iteration will continue.
+	 * If the iteration lambda returns true, the iteration will stop. If the iteration lambda returns false, the iteration will continue.
 	 */
 	void IterateNodeChildren(const FString& NodeUniqueID, TFunctionRef<void(const UInterchangeBaseNode*)> IterationLambda) const
 	{
@@ -106,9 +103,9 @@ public:
 
 	/**
 	 * Recursively traverse the hierarchy starting with the specified node unique ID.
-	 * If the iteration lambda return true the iteration will stop. If the iteration lambda return false the iteration will continue.
+	 * If the iteration lambda returns true, the iteration will stop. If the iteration lambda returns false, the iteration will continue.
 	 * 
-	 * @return - Return true if the iteration was break, false otherwise
+	 * @return - Return true if the iteration was broken, or false otherwise.
 	 */
 	bool BreakableIterateNodeChildren(const FString& NodeUniqueID, TFunctionRef<bool(const UInterchangeBaseNode*)> IterationLambda) const
 	{
@@ -130,7 +127,7 @@ public:
 		return false;
 	}
 
-	/** Unordered iteration of the all nodes, but I can be stop early by returning true */
+	/** Unordered iteration of the all nodes, but can be stopped early by returning true. */
 	INTERCHANGECORE_API void BreakableIterateNodes(TFunctionRef<bool(const FString&, UInterchangeBaseNode*)> IterationLambda) const;
 
 	template <typename T>
@@ -148,61 +145,61 @@ public:
 		}
 	}
 
-	/** Return all nodes that do not have any parent */
+	/** Return all nodes that do not have any parent. */
 	UFUNCTION(BlueprintCallable, Category = "Interchange | Node Container")
 	INTERCHANGECORE_API void GetRoots(TArray<FString>& RootNodes) const;
 
-	/** Return all nodes that are of the ClassNode type*/
+	/** Return all nodes that are of the ClassNode type. */
 	UFUNCTION(BlueprintCallable, Category = "Interchange | Node Container")
 	INTERCHANGECORE_API void GetNodes(const UClass* ClassNode, TArray<FString>& OutNodes) const;
 
-	/** Get a node pointer. Once added to the container, nodes are considered const */
+	/** Get a node pointer. Once added to the container, nodes are considered const. */
 	UFUNCTION(BlueprintCallable, Category = "Interchange | Node Container")
 	INTERCHANGECORE_API const UInterchangeBaseNode* GetNode(const FString& NodeUniqueID) const;
 
-	/** Get a factory node pointer */
+	/** Get a factory node pointer. */
 	UFUNCTION(BlueprintCallable, Category = "Interchange | Node Container")
 	INTERCHANGECORE_API UInterchangeFactoryBaseNode* GetFactoryNode(const FString& NodeUniqueID) const;
 
-	/** Set node ParentUid */
+	/** Set the ParentUid of the node. */
 	UFUNCTION(BlueprintCallable, Category = "Interchange | Node Container")
 	INTERCHANGECORE_API bool SetNodeParentUid(const FString& NodeUniqueID, const FString& NewParentNodeUid);
 
-	/** Get the node children count */
+	/** Get the number of children the node has. */
 	UFUNCTION(BlueprintCallable, Category = "Interchange | Node Container")
 	INTERCHANGECORE_API int32 GetNodeChildrenCount(const FString& NodeUniqueID) const;
 
-	/** Get all children Uid */
+	/** Get the UIDs of all the node's children. */
 	UFUNCTION(BlueprintCallable, Category = "Interchange | Node Container")
 	INTERCHANGECORE_API TArray<FString> GetNodeChildrenUids(const FString& NodeUniqueID) const;
 
-	/** Get the node nth const children */
+	/** Get the nth const child of the node */
 	UFUNCTION(BlueprintCallable, Category = "Interchange | Node Container")
 	INTERCHANGECORE_API UInterchangeBaseNode* GetNodeChildren(const FString& NodeUniqueID, int32 ChildIndex);
 
-	/** Get the node nth const children. Const version */
+	/** Get the nth const child of the node. Const version. */
 	INTERCHANGECORE_API const UInterchangeBaseNode* GetNodeChildren(const FString& NodeUniqueID, int32 ChildIndex) const;
 
 	/**
-	 * This function serialize the node container and all node sub-objects point by it.
-	 * Out of process translator like fbx will dump a file containing this data and the editor
+	 * This function serializes the node container and all node sub-objects it points .
+	 * Out-of-process translators like FBX will dump a file containing this data, and the editor
 	 * will read the file and regenerate the container from the saved data.
 	 */
 	INTERCHANGECORE_API void SerializeNodeContainerData(FArchive& Ar);
 
-	/* Serialize the node container into the specified file.*/
+	/* Serialize the node container into the specified file. */
 	UFUNCTION(BlueprintCallable, Category = "Interchange | Node Container")
 	INTERCHANGECORE_API void SaveToFile(const FString& Filename);
 
-	/* Serialize the node container from the specified file.*/
+	/* Serialize the node container from the specified file. */
 	UFUNCTION(BlueprintCallable, Category = "Interchange | Node Container")
 	INTERCHANGECORE_API void LoadFromFile(const FString& Filename);
 
-	/* Fill the children uids cache to optimize the GetNodeChildrenUids call.*/
+	/* Fill the cache of children UIDs to optimize the GetNodeChildrenUids call. */
 	UFUNCTION(BlueprintCallable, Category = "Interchange | Node Container")
 	INTERCHANGECORE_API void ComputeChildrenCache();
 
-	/* Reset the children uids cache. */
+	/* Reset the cache of children UIDs. */
 	UFUNCTION(BlueprintCallable, Category = "Interchange | Node Container")
 	void ResetChildrenCache()
 	{
