Clean up typos and phrasing in Interchange modules - contribution from Robb Surridge

diff --git a/Engine/Source/Runtime/Interchange/Core/Public/Nodes/InterchangeFactoryBaseNode.h b/Engine/Source/Runtime/Interchange/Core/Public/Nodes/InterchangeFactoryBaseNode.h
--- a/Engine/Source/Runtime/Interchange/Core/Public/Nodes/InterchangeFactoryBaseNode.h
+++ b/Engine/Source/Runtime/Interchange/Core/Public/Nodes/InterchangeFactoryBaseNode.h
@@ -45,13 +45,13 @@ struct FFrame;
 	return false;
 
 /*
- * Use this macro to implement a custom node attribute setter that will add delegate to apply the node attribute value to any UObject derive from
- * the specified type class. Use this macro if the type class is different from the Node GetObjectCLass()
+ * Use this macro to implement a custom node attribute setter that will add a delegate to apply the node attribute value to any UObject derived from
+ * the specified type class. Use this macro if the type class is different from the node's GetObjectCLass().
  *
- * The node should derive from UInterchangeFactoryBaseNode to apply a delegate
- * The node need to implement 2 functions for the attribute
- *  - ApplyCustom##AttributeName##ToAsset(UObject*) , see UE::Interchange::FApplyAttributeToAsset delegate signature
- *  - FillCustom##AttributeName##FromAsset(UObject*) , see UE::Interchange::FFillAttributeToAsset delegate signature
+ * The node should derive from UInterchangeFactoryBaseNode to apply a delegate.
+ * The node needs to implement 2 functions for the attribute:
+ *  - ApplyCustom##AttributeName##ToAsset(UObject*): see the UE::Interchange::FApplyAttributeToAsset delegate signature.
+ *  - FillCustom##AttributeName##FromAsset(UObject*): see the UE::Interchange::FFillAttributeToAsset delegate signature.
  */
 #define IMPLEMENT_NODE_ATTRIBUTE_SETTER_WITH_CUSTOM_DELEGATE_WITH_CUSTOM_CLASS(NodeClassName, AttributeName, AttributeType, AssetType)				\
 	FString OperationName = GetTypeName() + TEXT(".Set" #AttributeName);																			\
@@ -71,8 +71,8 @@ struct FFrame;
 	return false;
 
 /*
- * This macro instead of receiving a TypeClass use the factory node GetObjectClass() to register the delegate
- * See documentation of IMPLEMENT_NODE_ATTRIBUTE_SETTER_WITH_CUSTOM_DELEGATE_WITH_CUSTOM_CLASS
+ * This macro uses the factory node GetObjectClass() to register the delegate instead of receiving a TypeClass.
+ * See the documentation of IMPLEMENT_NODE_ATTRIBUTE_SETTER_WITH_CUSTOM_DELEGATE_WITH_CUSTOM_CLASS.
  */
 #define IMPLEMENT_NODE_ATTRIBUTE_SETTER_WITH_CUSTOM_DELEGATE(NodeClassName, AttributeName, AttributeType)											\
 	FString OperationName = GetTypeName() + TEXT(".Set" #AttributeName);																			\
@@ -92,11 +92,11 @@ struct FFrame;
 
 
 /*
- * This macro force the addition of the Fill and Apply delegates on 'this'
- * if the designated attribute has been set on the source node
+ * This macro forces the addition of the Fill and Apply delegates on 'this'
+ * if the designated attribute has been set on the source node.
  * Copying the arrays of delegates from the source node to 'this' is not sufficient
  * because the delegates of the source node have a reference on the source node not 'this'.
- * Consequently a call to those delegates would always use the source node.
+ * Consequently, a call to those delegates would always use the source node.
  */
 #define COPY_NODE_DELEGATES(SourceNode, AttributeName, AttributeType, AssetType)															\
 {																																			\
@@ -143,18 +143,18 @@ namespace UE::Interchange
 UENUM()
 enum class EReimportStrategyFlags : uint8
 {
-	//Do not apply any property when re-importing, simply change the source data
+	//Do not apply any properties when reimporting. Simply change the source data
 	ApplyNoProperties,
-	//Always apply all pipeline specified properties
+	//Always apply all properties specified by the pipeline.
 	ApplyPipelineProperties,
-	//Always apply all pipeline properties, but leave the properties modified in editor since the last import
+	//Always apply all properties specified by the pipeline, but leave the properties that were modified in the editor since the last import.
 	ApplyEditorChangedProperties
 };
 
 /**
- * This struct is used to store and retrieve key value attributes. The attributes are store in a generic FAttributeStorage which serialize the value in a TArray64<uint8>
- * See UE::Interchange::EAttributeTypes to know the supported template types
- * This is an abstract class. This is the base class of the interchange node graph format, all class in this format should derive from this class
+ * This struct is used to store and retrieve key-value attributes. The attributes are stored in a generic FAttributeStorage that serializes the values in a TArray64<uint8>.
+ * See UE::Interchange::EAttributeTypes to know the supported template types.
+ * This is an abstract class. This is the base class of the Interchange node graph format; all classes in this format should derive from this class.
  */
 UCLASS(BlueprintType, MinimalAPI)
 class UInterchangeFactoryBaseNode : public UInterchangeBaseNode
@@ -185,17 +185,17 @@ public:
 	INTERCHANGECORE_API bool SetReimportStrategyFlags(const EReimportStrategyFlags& ReimportStrategyFlags);
 
 	/**
-	 * Return true if this node should skip the factory import process, false otherwise.
-	 * Node can be in a situation where we have to skip the import process because we cannot import the associated asset for multiple reason.
-	 * - An asset can already exist and is representing a different type (UClass)
-	 * - An asset can already exist and is being compiled
-	 * - An asset can already exist and is being imported by another concurrent import task (user importing multiple file in same time in the same content folder)
+	 * Return true if this node should skip the factory import process, or false otherwise.
+	 * Nodes can be in a situation where we have to skip the import process because we cannot import the associated asset for multiple reasons. For example:
+	 * - An asset can already exist and represents a different type (UClass).
+	 * - An asset can already exist and is being compiled.
+	 * - An asset can already exist and is being imported by another concurrent import task (such as a user importing multiple files at the same time in the same content folder).
 	 */
 	UFUNCTION(BlueprintCallable, Category = "Interchange | Node")
 	INTERCHANGECORE_API bool ShouldSkipNodeImport() const;
 
 	/**
-	 * Add the skip node attribute. Use this function to cancel the creation of the unreal asset, see ShouldSkipNodeImport for more documentation.
+	 * Add the skip node attribute. Use this function to cancel the creation of the Unreal asset. See ShouldSkipNodeImport for more documentation.
 	 */
 	UFUNCTION(BlueprintCallable, Category = "Interchange | Node")
 	INTERCHANGECORE_API bool SetSkipNodeImport();
@@ -207,12 +207,12 @@ public:
 	INTERCHANGECORE_API bool UnsetSkipNodeImport();
 
 	/**
-	 * Return the UClass of the object we represent so we can find factory/writer
+	 * Return the UClass of the object we represent, so we can find the factory/writer.
 	 */
 	INTERCHANGECORE_API virtual class UClass* GetObjectClass() const;
 
 	/**
-	 * Return the custom sub-path under PackageBasePath, where the assets will be created.
+	 * Return the custom sub-path under PackageBasePath where the assets will be created.
 	 */
 	UFUNCTION(BlueprintCallable, Category = "Interchange | Node")
 	INTERCHANGECORE_API bool GetCustomSubPath(FString& AttributeValue) const;
@@ -224,19 +224,19 @@ public:
 	INTERCHANGECORE_API bool SetCustomSubPath(const FString& AttributeValue);
 
 	/**
-	 * This function allow to retrieve the number of factory dependencies for this object.
+	 * Retrieve the number of factory dependencies for this object.
 	 */
 	UFUNCTION(BlueprintCallable, Category = "Interchange | Node")
 	INTERCHANGECORE_API int32 GetFactoryDependenciesCount() const;
 
 	/**
-	 * This function allow to retrieve the dependency for this object.
+	 * Retrieve the dependencies for this object.
 	 */
 	UFUNCTION(BlueprintCallable, Category = "Interchange | Node")
 	INTERCHANGECORE_API void GetFactoryDependencies(TArray<FString>& OutDependencies ) const;
 
 	/**
-	 * This function allow to retrieve one dependency for this object.
+	 * Retrieve one dependency for this object.
 	 */
 	UFUNCTION(BlueprintCallable, Category = "Interchange | Node")
 	INTERCHANGECORE_API void GetFactoryDependency(const int32 Index, FString& OutDependency) const;
@@ -254,13 +254,13 @@ public:
 	INTERCHANGECORE_API bool RemoveFactoryDependencyUid(const FString& DependencyUid);
 
 	/**
-	 * Return the custom ReferenceObject. The UObject this factory node has created.
+	 * Return the custom ReferenceObject: the UObject this factory node has created.
 	 */
 	UFUNCTION(BlueprintCallable, Category = "Interchange | Node")
 	INTERCHANGECORE_API bool GetCustomReferenceObject(FSoftObjectPath& AttributeValue) const;
 
 	/**
-	 * Set the custom ReferenceObject. The UObject this factory node has created.
+	 * Set the custom ReferenceObject: the UObject this factory node has created.
 	 */
 	UFUNCTION(BlueprintCallable, Category = "Interchange | Node")
 	INTERCHANGECORE_API bool SetCustomReferenceObject(const FSoftObjectPath& AttributeValue);
@@ -281,17 +281,17 @@ public:
 	 * @param NodeAttributeKey	The key for the attribute to write.
 	 * @param Object			The object to write to.
 	 * @param PropertyName		The name of the property to write to on the Object.
-	 * @result					True if we succeeded.
+	 * @result					True if the operation succeeded.
 	 */
 	template<typename AttributeType>
 	inline bool ApplyAttributeToObject(const FString& NodeAttributeKey, UObject* Object, const FName PropertyName) const;
 
 	/**
 	 * Reads an attribute value from a UObject.
-	 * @param NodeAttributeKey	The key for the attribute to update.
+	 * @param NodeAttributeKey	The key for the attribute to read.
 	 * @param Object			The object to read from.
 	 * @param PropertyName		The name of the property to read from on the Object.
-	 * @result					True if we succeeded.
+	 * @result					True if the operation succeeded.
 	 */
 	template<typename AttributeType>
 	inline bool FillAttributeFromObject(const FString& NodeAttributeKey, UObject* Object, const FName PropertyName);
@@ -308,16 +308,16 @@ public:
 	}
 
 	/**
-	 * Each Attribute that was set and have a delegate set for the specified UObject->UClass will
-	 * get the delegate execute so it apply the attribute to the UObject property.
-	 * See the macros IMPLEMENT_NODE_ATTRIBUTE_SETTER at the top of the file to know how delegates are setup for property.
+	 * Each Attribute that was set and has a delegate set for the specified UObject->UClass will
+	 * have the delegate executed so it applies the attribute to the UObject property.
+	 * See the macros IMPLEMENT_NODE_ATTRIBUTE_SETTER at the top of the file to know how delegates are setup for properties.
 	 */
 	INTERCHANGECORE_API void ApplyAllCustomAttributeToObject(UObject* Object) const;
 
 	INTERCHANGECORE_API void FillAllCustomAttributeFromObject(UObject* Object) const;
 
 	/**
-	 * Copies all the custom attributes from SourceNode to this node and
+	 * Copies all the custom attributes from SourceNode to this node, and
 	 * gets the appropriate values from Object.
 	 * @param SourceNode		The source factory node to copy from.
 	 * @param Object			The object to fill custom attributes' values from, if applicable.
@@ -326,32 +326,31 @@ public:
 
 	/**
 	 * Return whether or not an object should be created even if it has been deleted in the editor.
-	 * Return false by default
+	 * Return false by default.
 	 */
 	UFUNCTION(BlueprintCallable, Category = "Interchange | Node")
 	INTERCHANGECORE_API bool ShouldForceNodeReimport() const;
 
 	/**
-	 * Allow the creation of the unreal object even if it has been previously deleted in the editor.
+	 * Allow the creation of the Unreal object even if it has been previously deleted in the editor.
 	 */
 	UFUNCTION(BlueprintCallable, Category = "Interchange | Node")
 	INTERCHANGECORE_API bool SetForceNodeReimport();
 
 	/**
-	 * Disallow the creation of the unreal object if it has been previously deleted in the editor.
+	 * Disallow the creation of the Unreal object if it has been previously deleted in the editor.
 	 */
 	UFUNCTION(BlueprintCallable, Category = "Interchange | Node")
 	INTERCHANGECORE_API bool UnsetForceNodeReimport();
 
 protected:
 	/**
-	 * Those dependencies are use by the interchange parsing task to make sure the asset are created in the correct order.
-	 * Example: Mesh factory node will have dependencies on material factory node
-	 *          Material factory node will have dependencies on texture factory node
+	 * Those dependencies are used by the Interchange parsing task to make sure the assets are created in the correct order.
+	 * Examples: The Mesh factory node will have dependencies on the material factory node, and the Material factory node will have dependencies on the texture factory node.
 	 */
 	UE::Interchange::TArrayAttributeHelper<FString> FactoryDependencies;
 
-	/* This array hold the delegate to apply the attribute that has to be set on an UObject */
+	/* This array holds the delegate to apply the attribute that has to be set on an UObject. */
 	TMap<UClass*, TArray<UE::Interchange::FApplyAttributeToAsset>> ApplyCustomAttributeDelegates;
 
 	TMap<UClass*, TArray<UE::Interchange::FFillAttributeToAsset>> FillCustomAttributeDelegates;
@@ -418,7 +417,7 @@ inline bool UInterchangeFactoryBaseNode::ApplyAttributeToObject(const FString& N
 
 /**
  * Specialized version of ApplyAttributeToObject for strings.
- * If the target property is a FObjectPropertyBase, treat the string as an object path.
+ * If the target property is an FObjectPropertyBase, treat the string as an object path.
  */
 template<>
 inline bool UInterchangeFactoryBaseNode::ApplyAttributeToObject<FString>(const FString& NodeAttributeKey, UObject* Object, const FName PropertyName) const
@@ -461,7 +460,7 @@ inline bool UInterchangeFactoryBaseNode::ApplyAttributeToObject<FString>(const F
 
 /**
  * Specialized version of ApplyAttributeToObject for bools.
- * If the target property is a FBoolProperty, treat the propertyg as a bitfield.
+ * If the target property is an FBoolProperty, treat the property as a bitfield.
  */
 template<>
 inline bool UInterchangeFactoryBaseNode::ApplyAttributeToObject<bool>(const FString& NodeAttributeKey, UObject* Object, const FName PropertyName) const
@@ -537,7 +536,7 @@ inline bool UInterchangeFactoryBaseNode::FillAttributeFromObject(const FString&
 
 /**
  * Specialized version of FillAttributeFromObject for strings.
- * If the target property is a FObjectPropertyBase, treat the string as an object path.
+ * If the target property is an FObjectPropertyBase, treat the string as an object path.
  */
 template<>
 inline bool UInterchangeFactoryBaseNode::FillAttributeFromObject<FString>(const FString& NodeAttributeKey, UObject* Object, const FName PropertyName)
@@ -572,7 +571,7 @@ inline bool UInterchangeFactoryBaseNode::FillAttributeFromObject<FString>(const
 
 /**
  * Specialized version of FillAttributeFromObject for bools.
- * If the target property is a FBoolProperty, treat the property as a bitfield.
+ * If the target property is an FBoolProperty, treat the property as a bitfield.
  */
 template<>
 inline bool UInterchangeFactoryBaseNode::FillAttributeFromObject<bool>(const FString& NodeAttributeKey, UObject* Object, const FName PropertyName)
diff --git a/Engine/Source/Runtime/Interchange/Core/Public/Nodes/InterchangeSourceNode.h b/Engine/Source/Runtime/Interchange/Core/Public/Nodes/InterchangeSourceNode.h
--- a/Engine/Source/Runtime/Interchange/Core/Public/Nodes/InterchangeSourceNode.h
+++ b/Engine/Source/Runtime/Interchange/Core/Public/Nodes/InterchangeSourceNode.h
@@ -20,7 +20,7 @@ class UObject;
 struct FFrame;
 
 /**
- * This class allow a translator to add general source data that describe the whole source. Pipeline can use this information.
+ * This class allows a translator to add general source data that describes the whole source. Pipelines can use this information.
  */
 UCLASS(BlueprintType, MinimalAPI)
 class UInterchangeSourceNode : public UInterchangeBaseNode
@@ -31,30 +31,30 @@ class UInterchangeSourceNode : public UInterchangeBaseNode
 
 public:
 	/**
-	 * Initialize the base data of the node
-	 * @param UniqueID - The uniqueId for this node
-	 * @param DisplayLabel - The name of the node
+	 * Initialize the base data of the node.
+	 * @param UniqueID - The unique ID for this node.
+	 * @param DisplayLabel - The name of the node.
 	 *
 	 */
 	UFUNCTION(BlueprintCallable, Category = "Interchange | Node | source")
 	INTERCHANGECORE_API void InitializeSourceNode(const FString& UniqueID, const FString& DisplayLabel);
 
 	/**
-	 * Return the node type name of the class, we use this when reporting error
+	 * Return the node type name of the class. This is used when reporting errors.
 	 */
 	INTERCHANGECORE_API virtual FString GetTypeName() const override;
 
-	/* The translators that want to modify the common data should ensure they create the unique common pipeline node. */
+	/* Translators that want to modify the common data should ensure they create the unique common pipeline node. */
 	static INTERCHANGECORE_API UInterchangeSourceNode* FindOrCreateUniqueInstance(UInterchangeBaseNodeContainer* NodeContainer);
 
-	/* If the unique instance doesn't exist it will return nullptr. This function should be use by the pipelines, to avoid creating a node. */
+	/* This function should be use by pipelines to avoid creating a node. If the unique instance doesn't exist, returns nullptr. */
 	static INTERCHANGECORE_API const UInterchangeSourceNode* GetUniqueInstance(const UInterchangeBaseNodeContainer* NodeContainer);
 
 	/** Query the source frame rate numerator. */
 	UFUNCTION(BlueprintCallable, Category = "Interchange | Node | Source")
 	INTERCHANGECORE_API bool GetCustomSourceFrameRateNumerator(int32& AttributeValue) const;
 
-	/** Store the source frame rate numerator. */
+	/** Set the source frame rate numerator. */
 	UFUNCTION(BlueprintCallable, Category = "Interchange | Node | Source")
 	INTERCHANGECORE_API bool SetCustomSourceFrameRateNumerator(const int32& AttributeValue);
 
@@ -62,47 +62,47 @@ public:
 	UFUNCTION(BlueprintCallable, Category = "Interchange | Node | Source")
 	INTERCHANGECORE_API bool GetCustomSourceFrameRateDenominator(int32& AttributeValue) const;
 
-	/** Store the source frame rate denominator. */
+	/** Set the source frame rate denominator. */
 	UFUNCTION(BlueprintCallable, Category = "Interchange | Node | Source")
 	INTERCHANGECORE_API bool SetCustomSourceFrameRateDenominator(const int32& AttributeValue);
 
-	/** Query the source time line start. */
+	/** Query the start of the source timeline. */
 	UFUNCTION(BlueprintCallable, Category = "Interchange | Node | Source")
 	INTERCHANGECORE_API bool GetCustomSourceTimelineStart(double& AttributeValue) const;
 
-	/** Store the source time line start. */
+	/** Set the start of the source timeline. */
 	UFUNCTION(BlueprintCallable, Category = "Interchange | Node | Source")
 	INTERCHANGECORE_API bool SetCustomSourceTimelineStart(const double& AttributeValue);
 
-	/** Query the source time line end. */
+	/** Query the end of the source timeline. */
 	UFUNCTION(BlueprintCallable, Category = "Interchange | Node | Source")
 	INTERCHANGECORE_API bool GetCustomSourceTimelineEnd(double& AttributeValue) const;
 
-	/** Store the source time line end. */
+	/** Set the end of the source timeline. */
 	UFUNCTION(BlueprintCallable, Category = "Interchange | Node | Source")
 	INTERCHANGECORE_API bool SetCustomSourceTimelineEnd(const double& AttributeValue);
 
-	/** Query the source animated time start. */
+	/** Query the start of the source animated time. */
 	UFUNCTION(BlueprintCallable, Category = "Interchange | Node | Source")
 	INTERCHANGECORE_API bool GetCustomAnimatedTimeStart(double& AttributeValue) const;
 
-	/** Store the source animated time start. */
+	/** Set the start of the source animated time. */
 	UFUNCTION(BlueprintCallable, Category = "Interchange | Node | Source")
 	INTERCHANGECORE_API bool SetCustomAnimatedTimeStart(const double& AttributeValue);
 
-	/** Query the source animated time end. */
+	/** Query the end of the source animated time. */
 	UFUNCTION(BlueprintCallable, Category = "Interchange | Node | Source")
 	INTERCHANGECORE_API bool GetCustomAnimatedTimeEnd(double& AttributeValue) const;
 
-	/** Store the source animated time end. */
+	/** Set the end of the source animated time. */
 	UFUNCTION(BlueprintCallable, Category = "Interchange | Node | Source")
 	INTERCHANGECORE_API bool SetCustomAnimatedTimeEnd(const double& AttributeValue);
 
-	/** Query the source animated time end. */
+	/** Query whether to import materials that aren't used. */
 	UFUNCTION(BlueprintCallable, Category = "Interchange | Node | Source")
 	INTERCHANGECORE_API bool GetCustomImportUnusedMaterial(bool& AttributeValue) const;
 
-	/** Store the source animated time end. */
+	/** Set whether to import materials that aren't used. */
 	UFUNCTION(BlueprintCallable, Category = "Interchange | Node | Source")
 	INTERCHANGECORE_API bool SetCustomImportUnusedMaterial(const bool& AttributeValue);
 
diff --git a/Engine/Source/Runtime/Interchange/Core/Public/Nodes/InterchangeUserDefinedAttribute.h b/Engine/Source/Runtime/Interchange/Core/Public/Nodes/InterchangeUserDefinedAttribute.h
--- a/Engine/Source/Runtime/Interchange/Core/Public/Nodes/InterchangeUserDefinedAttribute.h
+++ b/Engine/Source/Runtime/Interchange/Core/Public/Nodes/InterchangeUserDefinedAttribute.h
@@ -36,9 +36,9 @@ struct FInterchangeUserDefinedAttributeInfo
 };
 
 /**
- * UInterchangeUserDefinedAttributesAPI is used to store and retrieve user defined attributes (i.e. DCC node attributes, pipelines will have access to those attributes)
- * Any user defined attribute have: name, value and a optional AnimationPayloadKey (FRichCurve which is a float curve).
- * Value type must be supported by the UE::Interchange::EAttributeTypes enumeration.
+ * UInterchangeUserDefinedAttributesAPI is used to store and retrieve user-defined attributes such as DCC node attributes, so that pipelines have access to those attributes.
+ * Every user-defined attribute has a name, a value, and an optional AnimationPayloadKey: an FRichCurve that is a float curve.
+ * The value type must be supported by the UE::Interchange::EAttributeTypes enumeration.
  */
 UCLASS(BlueprintType, Experimental, MinimalAPI)
 class UInterchangeUserDefinedAttributesAPI : public UObject
@@ -48,12 +48,12 @@ class UInterchangeUserDefinedAttributesAPI : public UObject
 public:
 
 	/**
-	 * Create user defined attribute with a value and a optional payload key
-	 * param UserDefinedAttributeName - The name of the user defined attribute
-	 * param Value - The value of the user defined attribute
-	 * param PayloadKey - The translator payload key to retrieve the FRichCurve animation for this user defined attribute
-	 * note - User defined attributes are the DCC translated node user custom attributes (i.e. Maya extra attributes)
-	 *        Payload key will point on a FRichCurve payload.
+	 * Create a user-defined attribute with a value and a optional payload key.
+	 * param UserDefinedAttributeName - The name of the user-defined attribute.
+	 * param Value - The value of the user-defined attribute.
+	 * param PayloadKey - The translator payload key to retrieve the FRichCurve animation for this user-defined attribute.
+	 * Note - User-defined attributes are the user custom attributes from the DCC translated node (for example, extra attributes in Maya).
+	 *        The payload key points to an FRichCurve payload.
 	 */
 	template<typename ValueType>
 	static bool CreateUserDefinedAttribute(UInterchangeBaseNode* InterchangeNode, const FString& UserDefinedAttributeName, const ValueType& Value, const TOptional<FString>& PayloadKey, bool RequiresDelegate = false);
@@ -74,22 +74,22 @@ public:
 	static INTERCHANGECORE_API bool CreateUserDefinedAttribute_FString(UInterchangeBaseNode* InterchangeNode, const FString& UserDefinedAttributeName, const FString& Value, const FString& PayloadKey, bool RequiresDelegate = false);
 
 	/**
-	 * Remove the specified user defined attribute
-	 * param UserDefinedAttributeName - The name of the user defined attribute to remove
-	 * return - True if the attribute exist and was remove or if the attribute doesn't exist. Return false if the attribute exist but the attribute was not properly remove.
-	 * note - User defined attributes are the DCC translated node user custom attributes (i.e. Maya extra attributes)
-	 *        Payload key will point on a FRichCurve payload.
+	 * Remove the specified user-defined attribute.
+	 * @param UserDefinedAttributeName - The name of the user-defined attribute to remove.
+	 * @return - True if the attribute exists and was removed, or if the attribute doesn't exist. Returns false if the attribute exists but could not be removed.
+	 * Note - User-defined attributes are the user custom attributes from the DCC translated node (for example, extra attributes in Maya).
+	 *        The payload key points to an FRichCurve payload.
 	 */
 	UFUNCTION(BlueprintCallable, Category = "Interchange | Node | UserDefinedAttribute")
 	static INTERCHANGECORE_API bool RemoveUserDefinedAttribute(UInterchangeBaseNode* InterchangeNode, const FString& UserDefinedAttributeName);
 
 	/**
-	 * Get user defined attribute value and optional payload key
-	 * param UserDefinedAttributeName - The name of the user defined attribute
-	 * param OutValue - The value of the user defined attribute
-	 * param OutPayloadKey - The translator payload key to retrieve the FRichCurve animation for this user defined attribute
-	 * note - User defined attributes are the DCC translated node user custom attributes (i.e. Maya extra attributes)
-	 *        Payload key will point on a FRichCurve payload.
+	 * Get the value of a user-defined attribute and an optional payload key.
+	 * @param UserDefinedAttributeName - The name of the user-defined attribute.
+	 * @param OutValue - The value of the user-defined attribute.
+	 * @param OutPayloadKey - The translator payload key to retrieve the FRichCurve animation for this user-defined attribute.
+	 * Note - User-defined attributes are the user custom attributes from the DCC translated node (for example, extra attributes in Maya).
+	 *        The payload key points to an FRichCurve payload.
 	 */
 	template<typename ValueType>
 	static bool GetUserDefinedAttribute(const UInterchangeBaseNode* InterchangeNode, const FString& UserDefinedAttributeName, ValueType& OutValue, TOptional<FString>& OutPayloadKey);
diff --git a/Engine/Source/Runtime/Interchange/Core/Public/Types/AttributeStorage.h b/Engine/Source/Runtime/Interchange/Core/Public/Types/AttributeStorage.h
--- a/Engine/Source/Runtime/Interchange/Core/Public/Types/AttributeStorage.h
+++ b/Engine/Source/Runtime/Interchange/Core/Public/Types/AttributeStorage.h
@@ -194,7 +194,7 @@ namespace UE
 
 		/**
 		 * Enumerates the built-in types that can be stored in instances of FAttributeStorage.
-		 * We cannot change the value of a type to make sure the serialization of old asset is always working
+		 * We cannot change the value of a type to make sure the serialization of old assets is always working.
 		 */
 		enum class EAttributeTypes : int32
 		{
@@ -262,12 +262,12 @@ namespace UE
 		};
 
 		/**
-		 * Return the FString for the specified AttributeType
+		 * Return the FString for the specified AttributeType.
 		 */
 		INTERCHANGECORE_API FString AttributeTypeToString(EAttributeTypes AttributeType);
 
 		/**
-		 * Return the AttributeType for the specified FString, return EAttributeTypes::None if the string do not match any
+		 * Return the AttributeType for the specified FString, or return EAttributeTypes::None if the string does not match any
 		 * supported attribute type.
 		 */
 		INTERCHANGECORE_API EAttributeTypes StringToAttributeType(const FString& AttributeTypeString);
@@ -279,7 +279,7 @@ namespace UE
 		 * data types that are to be used internally by FAttributeStorage. Traits for the most commonly used built-in
 		 * types are declared below.
 		 *
-		 * Complex types, such as structures and classes can be serialized into a byte array
+		 * Complex types, such as structures and classes, can be serialized into a byte array
 		 * and then assigned to an attribute. Note that you will be responsible for ensuring
 		 * correct byte ordering when serializing those types.
 		 *
@@ -307,28 +307,28 @@ namespace UE
 		}
 
 		/**
-		 * Enum to return complete status of a storage operation
-		 * It support success with additional information
-		 * It support multiple errors
+		 * Enum to return complete status of a storage operation.
+		 * It supports success with additional information.
+		 * It supports multiple errors.
 		 */
 		enum class EAttributeStorageResult : uint64
 		{
 			None											= 0x0,
-			//success result
+			//Success result.
 			Operation_Success								= 0x1,
 		
 
-			//Operation error result from here
+			//Operation error results from here.
 
-			//The type of the value was not matching the existing type. We cannot change the type of an existing attribute
+			//The type of the value was not matching the existing type. We cannot change the type of an existing attribute.
 			Operation_Error_WrongType						= (0x1 << 20),
-			//The size of the value is different from the existing size. We cannot have a different size
+			//The size of the value is different from the existing size. We cannot have a different size.
 			Operation_Error_WrongSize						= (0x1 << 21),
-			//The AttributeAllocationTable has an attribute that the offset is no valid in the storage
+			//The AttributeAllocationTable has an attribute whose offset is not valid in the storage.
 			Operation_Error_AttributeAllocationCorrupted	= (0x1 << 22),
-			//We cannot find the specified key
+			//We cannot find the specified key.
 			Operation_Error_CannotFoundKey					= (0x1 << 23),
-			//There was an error when removing an attribute from the AttributeAllocationTable. The TArray remove has failed
+			//There was an error when removing an attribute from the AttributeAllocationTable. The TArray remove has failed.
 			Operation_Error_CannotRemoveAttribute			= (0x1 << 24),
 			//We try to override an attribute but the specified options do not allow override.
 			Operation_Error_CannotOverrideAttribute			= (0x1 << 25),
@@ -341,8 +341,8 @@ namespace UE
 		ENUM_CLASS_FLAGS(EAttributeStorageResult)
 
 		/**
-		 * Helper function to interpret storage result
-		 * @return true if the Result contain at least one of the RefResult flags
+		 * Helper function to interpret storage results.
+		 * @return true if the result contains at least one of the RefResult flags.
 		 */
 		FORCEINLINE bool HasAttributeStorageResult(const EAttributeStorageResult Result, const EAttributeStorageResult RefResult)
 		{
@@ -350,8 +350,8 @@ namespace UE
 		}
 	
 		/**
-		 * Helper function to determine if the storage result is success
-		 * @return true if result contain Operation_Success
+		 * Helper function to determine if the storage result is success.
+		 * @return true if the result contains Operation_Success.
 		 */
 		FORCEINLINE bool IsAttributeStorageResultSuccess(const EAttributeStorageResult Result)
 		{
@@ -360,27 +360,27 @@ namespace UE
 
 		/**
 		 * Helper function to transform an operation result into a LOG.
-		 * @param Result - The result we want to output a log for
-		 * @param OperationName - The operation name that end up with the speicified Result.
-		 * @param AttributeKey - The attribute we did apply the operation on.
+		 * @param Result - The result we want to output a log for.
+		 * @param OperationName - The operation name that ended up with the specified result.
+		 * @param AttributeKey - The attribute we applied the operation to.
 		 */
 		INTERCHANGECORE_API void LogAttributeStorageErrors(const EAttributeStorageResult Result, const FString OperationName, const FAttributeKey AttributeKey );
 
 
 		/**
-		 * Enum to pass options when we add an attribute
+		 * Enum to pass options when we add an attribute.
 		 */
 		enum class EAttributeStorageAddOptions : uint32
 		{
 			None					= 0x0,
 			Option_Override			= 0x1,
-			//allow the AddAttribute to override the value if it exist.
+			//allow the AddAttribute to override the value if it exists.
 		};
 
 		ENUM_CLASS_FLAGS(EAttributeStorageAddOptions)
 
 		/**
-		 * Helper function to interpret storage add options
+		 * Helper function to interpret storage add options.
 		 */
 		FORCEINLINE bool HasAttributeStorageAddOption(const EAttributeStorageAddOptions Options, const EAttributeStorageAddOptions RefOptions)
 		{
@@ -388,18 +388,18 @@ namespace UE
 		}
 
 		/**
-		 * Enumerates the attribute properties. Those properties affect how the attribute are stored or what they are use for
+		 * Enumerates the attribute properties. Those properties affect how the attribute are stored or what they are used for.
 		 */
 		enum class EAttributeProperty : uint32
 		{
 			None						= 0x0,
-			NoHash						= 0x1, /*No hash attribute will not be part of the hash result when calling GetStorageHash*/
+			NoHash						= 0x1, /* No hash attribute will not be part of the hash result when calling GetStorageHash. */
 		};
 
 		ENUM_CLASS_FLAGS(EAttributeProperty)
 
 		/**
-		 * Helper function to interpret attribute property
+		 * Helper function to interpret an attribute property.
 		 */
 		FORCEINLINE bool HasAttributeProperty(const EAttributeProperty PropertyA, const EAttributeProperty PropertyB)
 		{
@@ -407,7 +407,7 @@ namespace UE
 		}
 
 		/**
-		 * This is a helper class to specialize template function inside a class, it allow specialize template to not be static
+		 * This is a helper class to specialize template functions inside a class. It allows specialized templates to not be static.
 		 */
 		template<typename T>
 		struct TSpecializeType
@@ -417,19 +417,19 @@ namespace UE
 
 	
 		/**
-		 * This class is a Key/Value storage inside a TArray64<uint8>
-		 * The keys are of type FAttributeKey which is a FString, each key is unique and have only one value.
-		 * The value can be of any type contain in EAttributeTypes.
+		 * This class is a Key/Value storage inside a TArray64<uint8>.
+		 * The keys are of type FAttributeKey, which is an FString. Each key is unique and has only one value.
+		 * The value can be of any type contained in EAttributeTypes.
 		 * 
-		 * @notes
-		 * The storage is multi thread safe, it use a mutex to lock the storage for every read/write operation.
-		 * The hash of the storage is deterministic because it sort the attributes before calculating the hash.
+		 * @note
+		 * The storage is multi-thread safe. It uses a mutex to lock the storage for every read/write operation.
+		 * The hash of the storage is deterministic because it sorts the attributes before calculating the hash.
 		 */
 		class FAttributeStorage
 		{
 		public:
 			/**
-			 * Class to Get/Set an attribute of the storage 
+			 * Class to get/set an attribute of the storage.
 			 */
 			template<typename T>
 			class TAttributeHandle
@@ -443,7 +443,7 @@ namespace UE
 				{}
 
 				/**
-				* Return true if the storage contain a valid attribute key, false otherwise. 
+				* Return true if the storage contains a valid attribute key, or false otherwise. 
 				*/
 				FORCEINLINE bool IsValid() const
 				{
@@ -516,14 +516,14 @@ namespace UE
 
 			/**
 			 * Register an attribute in the storage. Return success if the attribute was properly added, or there is an existing
-			 * attribute of the same type, return error otherwise.
+			 * attribute of the same type. Return an error otherwise.
 			 *
-			 * @Param ElementAttributeKey - is the storage key (the path) of the attribute
-			 * @Param DefaultValue - is the default value for the registered attribute
+			 * @Param ElementAttributeKey - the storage key (the path) of the attribute.
+			 * @Param DefaultValue - the default value for the registered attribute.
 			 *
-			 * @note Possible errors
-			 * - Key exist with a different type
-			 * - Storage is corrupted
+			 * @note Possible errors:
+			 * - Key exists with a different type.
+			 * - Storage is corrupted.
 			 */
 			template<typename T>
 			EAttributeStorageResult RegisterAttribute(const FAttributeKey& ElementAttributeKey, const T& DefaultValue, EAttributeProperty AttributeProperty = EAttributeProperty::None)
@@ -573,19 +573,19 @@ namespace UE
 			/**
 			 * Remove an attribute from the storage.
 			 *
-			 * @param ElementAttributeKey - is the storage key (the path) of the attribute to remove.
+			 * @param ElementAttributeKey - the storage key (the path) of the attribute to remove.
 			 *
-			 * @note Possible errors
-			 * - Key do not exist
-			 * - Internal storage structure remove error
+			 * @note Possible errors:
+			 * - Key does not exist.
+			 * - Internal storage structure removal error.
 			 */
 			INTERCHANGECORE_API EAttributeStorageResult UnregisterAttribute(const FAttributeKey& ElementAttributeKey);
 
 			/**
-			 * Return an attribute handle for the specified attribute. This handle is compile type check and is use to Get and Set the attribute value type.
-			 * The function will assert if the key is missing or the type doesnt match the specified template type.
+			 * Return an attribute handle for the specified attribute. This handle is a compile type check and is use to get and set the attribute value type.
+			 * The function will assert if the key is missing or the type doesn't match the specified template type.
 			 *
-			 * @param ElementAttributeKey is the storage key (the path) of the attribute.
+			 * @param ElementAttributeKey - the storage key (the path) of the attribute.
 			 */
 			template<typename T>
 			TAttributeHandle<T> GetAttributeHandle(const FAttributeKey& ElementAttributeKey) const
@@ -596,87 +596,87 @@ namespace UE
 
 
 			/**
-			 * Return the attribute type if the key exist, return None if the key is missing
+			 * Return the attribute type if the key exists, or None if the key is missing.
 			 *
-			 * @param ElementAttributeKey - is the storage key (the path) of the attribute.
+			 * @param ElementAttributeKey - the storage key (the path) of the attribute.
 			 *
 			 */
 			INTERCHANGECORE_API EAttributeTypes GetAttributeType(const FAttributeKey& ElementAttributeKey) const;
 
 			/**
-			 * Return true if the attribute key point on an existing attribute in the storage. Return false otherwise
+			 * Return true if the attribute key points to an existing attribute in the storage. Return false otherwise.
 			 *
-			 * @param ElementAttributeKey - is the storage key (the path) of the attribute.
+			 * @param ElementAttributeKey - the storage key (the path) of the attribute.
 			 *
 			 */
 			INTERCHANGECORE_API bool ContainAttribute(const FAttributeKey& ElementAttributeKey) const;
 
 			/**
-			 * This function allow the user to retrieve the array of keys to iterate and do reflection on the storage content
+			 * Retrieve the array of keys that can be used to iterate and do reflection on the storage content.
 			 *
 			 */
 			INTERCHANGECORE_API void GetAttributeKeys(TArray<FAttributeKey>& AttributeKeys) const;
 	
 			/**
-			 * Return a FGuid build from the FSHA1 of the specified attribute data. If the Attribute do not exist it return an empty FGUID
+			 * Return an FGuid built from the FSHA1 of the specified attribute data. If the attribute does not exist, return an empty FGUID.
 			 *
-			 * @param ElementAttributeKey is the storage key (the path) of the attribute.
+			 * @param ElementAttributeKey - the storage key (the path) of the attribute.
 			 *
 			 */
 			INTERCHANGECORE_API FGuid GetAttributeHash(const FAttributeKey& ElementAttributeKey) const;
 	
 			/**
-			 * This function fill the OutGuid with the hash of the specified attribute. Return true if the attribute exist and the OutGuid was assign, return false otherwise without touching the OutGuid.
+			 * This function fills the OutGuid with the hash of the specified attribute. Return true if the attribute exists and the OutGuid was assigned, or false otherwise without touching the OutGuid.
 			 *
-			 * @param ElementAttributeKey - is the storage key (the path) of the attribute.
-			 * @param OutGuid - is where we put the attribute hash.
+			 * @param ElementAttributeKey - the storage key (the path) of the attribute.
+			 * @param OutGuid - where we put the attribute hash.
 			 *
 			 */
 			INTERCHANGECORE_API bool GetAttributeHash(const FAttributeKey& ElementAttributeKey, FGuid& OutGuid) const;
 
 			/**
-			 * Return a FGuid build from the FSHA1 of all the attribute data contain in the node.
-			 * The data include the UniqueID and the DisplayLabel.
+			 * Return an FGuid built from the FSHA1 of all the attribute data contained in the node.
+			 * The data includes the UniqueID and the DisplayLabel.
 			 *
-			 * @note the attribute are sorted by key when building the FSHA1 data. The hash will be deterministic for the same data whatever
+			 * @note the attributes are sorted by key when building the FSHA1 data. The hash will be deterministic for the same data whatever
 			 * the order we add the attributes.
 			 */
 			INTERCHANGECORE_API FGuid GetStorageHash() const;
 
 			/**
-			 * Compare two storage to know which properties was modified/added/removed.
+			 * Compare two storage objects to know which properties were modified/added/removed.
 			 *
-			 * @param BaseStorage - The storage reference
-			 * @param VersionStorage - The storage with the changes
-			 * @param RemovedAttribute - All attributes that are in base storage but not in version storage, contain key only valid for base storage
-			 * @param AddedAttributes - All attributes that are in version storage but not in base storage, contain key only valid for version storage
-			 * @param ModifiedAttributes - All attributes that are in both storage base storage but have a different hash (different value)
+			 * @param BaseStorage - The reference storage.
+			 * @param VersionStorage - The storage with the changes.
+			 * @param RemovedAttribute - All attributes that are in base storage but not in version storage. Contains keys that are only valid for the base storage.
+			 * @param AddedAttributes - All attributes that are in version storage but not in base storage. Contains keys that are only valid for the version storage.
+			 * @param ModifiedAttributes - All attributes that are in both storage but have a different hash (different value).
 			 *
 			 */
 			static INTERCHANGECORE_API void CompareStorage(const FAttributeStorage& BaseStorage, const FAttributeStorage& VersionStorage, TArray<FAttributeKey>& RemovedAttributes, TArray<FAttributeKey>& AddedAttributes, TArray<FAttributeKey>& ModifiedAttributes);
 
 			/**
-			 * Copy an array of attributes from the source storage to the destination storage. If the attribute already exist in the destination, the value will be updated.
-			 * If a key do not exist in the source it will not be copy/created in the destination
+			 * Copy an array of attributes from the source storage to the destination storage. If the attribute already exists in the destination, the value will be updated.
+			 * If a key does not exist in the source it will not be copied/created in the destination.
 			 *
-			 * @param SourceStorage - The storage source
-			 * @param DestinationStorage - The storage destination
-			 * @param AttributeKeys - All attributes that must be copy from the source to the destination.
+			 * @param SourceStorage - The storage source.
+			 * @param DestinationStorage - The storage destination.
+			 * @param AttributeKeys - All attributes that must be copied from the source to the destination.
 			 *
 			 */
 			static INTERCHANGECORE_API void CopyStorageAttributes(const FAttributeStorage& SourceStorage, FAttributeStorage& DestinationStorage, const TArray<FAttributeKey>& AttributeKeys);
 
 			/**
-			 * Return the defrag ratio. This ratio is use to know when we need to defrag the storage
+			 * Return the defrag ratio. This ratio is used to know when we need to defrag the storage.
 			 * @example - a ratio of 0.1f will defrag the storage if the memory lost is bigger then 10% of the storage allocation.
-			 * Defrag is call when we remove an attribute or when we set the defrag ratio
+			 * Defrag is called when we remove an attribute or when we set the defrag ratio.
 			 */
 			FORCEINLINE float GetDefragRatio() const
 			{
 				return DefragRatio;
 			}
 	
-			/** Set the defrag ratio. See the GetDefragRatio for the defrag documentation. */
+			/** Set the defrag ratio. See GetDefragRatio() for the defrag documentation. */
 			INTERCHANGECORE_API void SetDefragRatio(const float InDefragRatio);
 	
 			friend FArchive& operator<<(FArchive& Ar, FAttributeStorage& Storage)
@@ -691,13 +691,13 @@ namespace UE
 			/**
 			 * Reserve the allocation table and the storage data.
 			 * 
-			 * @param NewAttributeCount: The number of attribute we want to reserve. Zero value do not reserve attribute count
-			 * @param NewStorageSize: The size on the storage all those new attribute will need. Passing zero value do not reserve storage size
+			 * @param NewAttributeCount: The number of attributes we want to reserve. Passing a zero value does not reserve attribute count.
+			 * @param NewStorageSize: The size of the storage all the new attributes will need. Passing a zero value do not reserve storage size.
 			 */
 			INTERCHANGECORE_API void Reserve(int64 NewAttributeCount, int64 NewStorageSize);
 
 		protected:
-			/** Structure use to hold the attribute information store in the attribute allocation table. */
+			/** Structure used to hold the attribute information stored in the attribute allocation table. */
 			struct FAttributeAllocationInfo
 			{
         		//The offset in the storage
diff --git a/Engine/Source/Runtime/Interchange/Engine/Private/InterchangeImportCommon.cpp b/Engine/Source/Runtime/Interchange/Engine/Private/InterchangeImportCommon.cpp
--- a/Engine/Source/Runtime/Interchange/Engine/Private/InterchangeImportCommon.cpp
+++ b/Engine/Source/Runtime/Interchange/Engine/Private/InterchangeImportCommon.cpp
@@ -338,7 +338,7 @@ namespace UE::Interchange
 			{
 				if (!PreviousAssetNode)
 				{
-					UE_LOG(LogInterchangeEngine, Error, TEXT("Cannot apply the re-import strategy for asset [%s], because there is no previous asset node in the import data."), *Asset->GetName());
+					UE_LOG(LogInterchangeEngine, Error, TEXT("Cannot apply the reimport strategy for asset [%s], because there is no previous asset node in the import data."), *Asset->GetName());
 					return;
 				}
 				TArray<FAttributeKey> RemovedAttributes;
@@ -409,13 +409,13 @@ namespace UE::Interchange
 		// create a new static mesh or overwrite existing asset, if possible
 		if (!ExistingAsset)
 		{
-			UE_LOG(LogInterchangeEngine, Error, TEXT("Cannot import the %s asset [%s], because it was not create on the game thread."), *AssetClass->GetName(), *AssetName);
+			UE_LOG(LogInterchangeEngine, Error, TEXT("Cannot import the %s asset [%s] because it was not created on the game thread."), *AssetClass->GetName(), *AssetName);
 			return nullptr;
 		}
 
 		if (!ExistingAsset->GetClass()->IsChildOf(AssetClass))
 		{
-			UE_LOG(LogInterchangeEngine, Error, TEXT("Cannot import the %s asset [%s], because it will override an asset of a different class (%s)."), *AssetClass->GetName(), *AssetName, *ExistingAsset->GetClass()->GetName());
+			UE_LOG(LogInterchangeEngine, Error, TEXT("Cannot import the %s asset [%s] because it will override an asset of a different class (%s)."), *AssetClass->GetName(), *AssetName, *ExistingAsset->GetClass()->GetName());
 			return nullptr;
 		}
 		return ExistingAsset;
diff --git a/Engine/Source/Runtime/Interchange/Engine/Private/InterchangeManager.cpp b/Engine/Source/Runtime/Interchange/Engine/Private/InterchangeManager.cpp
--- a/Engine/Source/Runtime/Interchange/Engine/Private/InterchangeManager.cpp
+++ b/Engine/Source/Runtime/Interchange/Engine/Private/InterchangeManager.cpp
@@ -248,7 +248,7 @@ bool UE::Interchange::FImportAsyncHelper::IsClassImportAllowed(UClass* Class)
 	{
 		if (!AssetClassPermissionList->PassesFilter(Class->GetPathName()))
 		{
-			UE_LOG(LogInterchangeEngine, Display, TEXT("The creation of asset of class '%s' is not allowed in this project."), *Class->GetName());
+			UE_LOG(LogInterchangeEngine, Display, TEXT("Creating assets of class '%s' is not allowed in this project."), *Class->GetName());
 			DeniedClasses.Add(Class);
 			return false;
 		}
@@ -835,7 +835,7 @@ UInterchangePipelineBase* UE::Interchange::GeneratePipelineInstance(const FSoftO
 		else
 		{
 			//Log an error because we cannot load the python class, maybe the python script was not loaded
-			UE_LOG(LogInterchangeEngine, Error, TEXT("Cannot generate a pipeline instance because the blueprint %s do not have a valid generated class."), *PipelineInstance.GetWithoutSubPath().ToString());
+			UE_LOG(LogInterchangeEngine, Error, TEXT("Cannot generate a pipeline instance because the Blueprint %s does not have a valid generated class."), *PipelineInstance.GetWithoutSubPath().ToString());
 		}
 	}
 	else if (const UInterchangePythonPipelineAsset* PythonPipeline = Cast<UInterchangePythonPipelineAsset>(ReferenceInstance))
@@ -847,7 +847,7 @@ UInterchangePipelineBase* UE::Interchange::GeneratePipelineInstance(const FSoftO
 		else
 		{
 			//Log an error because we cannot load the python class, maybe the python script was not loaded
-			UE_LOG(LogInterchangeEngine, Error, TEXT("Cannot generate a pipeline instance because the Python pipeline asset %s do not have a valid generated pipeline instance."), *PipelineInstance.GetWithoutSubPath().ToString());
+			UE_LOG(LogInterchangeEngine, Error, TEXT("Cannot generate a pipeline instance because the Python pipeline asset %s does not have a valid generated pipeline instance."), *PipelineInstance.GetWithoutSubPath().ToString());
 		}
 	}
 	else if (const UInterchangePipelineBase* DefaultPipeline = Cast<UInterchangePipelineBase>(ReferenceInstance))
@@ -1568,7 +1568,7 @@ UInterchangeManager::ImportInternal(const FString& ContentPath, const UInterchan
 
 	if (!SourceData)
 	{
-		UE_LOG(LogInterchangeEngine, Error, TEXT("Cannot import file, the source data is invalid"));
+		UE_LOG(LogInterchangeEngine, Error, TEXT("Cannot import file. The source data is invalid."));
 		return EarlyExit();
 	}
 	
@@ -1576,7 +1576,7 @@ UInterchangeManager::ImportInternal(const FString& ContentPath, const UInterchan
 		UInterchangeTranslatorBase* Translator = GetTranslatorForSourceData(SourceData);
 		if(!Translator)
 		{
-			UE_LOG(LogInterchangeEngine, Error, TEXT("Cannot import file, the source data is not supported. See if you can enable for interchange the extension [%s]"), *FPaths::GetExtension(SourceData->GetFilename()));
+			UE_LOG(LogInterchangeEngine, Error, TEXT("Cannot import file. The source data is not supported. Try enabling the [%s] extension for Interchange."), *FPaths::GetExtension(SourceData->GetFilename()));
 			return EarlyExit();
 		}
 		Translator->ReleaseSource();
@@ -1592,7 +1592,7 @@ UInterchangeManager::ImportInternal(const FString& ContentPath, const UInterchan
 	
 	if (InterchangeImportSettings.PipelineStacks.Num() == 0)
 	{
-		UE_LOG(LogInterchangeEngine, Error, TEXT("Cannot import file, there is no pipeline stack define for %s import type"), bImportScene ? TEXT("scene") : TEXT("content"));
+		UE_LOG(LogInterchangeEngine, Error, TEXT("Cannot import file. There is no pipeline stack defined for the %s import type."), bImportScene ? TEXT("scene") : TEXT("content"));
 		return EarlyExit();
 	}
 	
@@ -1792,7 +1792,7 @@ UInterchangeManager::ImportInternal(const FString& ContentPath, const UInterchan
 				else if(!SourcePipeline)
 				{
 					//A pipeline was not loaded
-					UE_LOG(LogInterchangeEngine, Warning, TEXT("Interchange Reimport: Missing import pipeline from the reimpoting asset. The reimport might fail."));
+					UE_LOG(LogInterchangeEngine, Warning, TEXT("Interchange Reimport: Missing import pipeline from the reimporting asset. The reimport might fail."));
 				}
 			}
 		}
@@ -1916,13 +1916,13 @@ UInterchangeManager::ImportInternal(const FString& ContentPath, const UInterchan
 					for (FInterchangeStackInfo& StackInfo : PipelineStacks)
 					{
 						OutPipelines = StackInfo.Pipelines;
-						UE_LOG(LogInterchangeEngine, Warning, TEXT("Interchange import: Invalid Default stack, using stack [%s] to import"), *StackInfo.StackName.ToString());
+						UE_LOG(LogInterchangeEngine, Warning, TEXT("Interchange import: Invalid Default stack. using stack [%s] to import."), *StackInfo.StackName.ToString());
 						break;
 					}
 				}
 				else
 				{
-					UE_LOG(LogInterchangeEngine, Warning, TEXT("Interchange Import: Cannot find any valid stack, cancelling import."));
+					UE_LOG(LogInterchangeEngine, Warning, TEXT("Interchange Import: Cannot find any valid stack, canceling import."));
 					bImportCanceled = true;
 				}
 			}
@@ -1958,7 +1958,7 @@ UInterchangeManager::ImportInternal(const FString& ContentPath, const UInterchan
 			UInterchangePipelineBase* GeneratedPipeline = UE::Interchange::GeneratePipelineInstance(ImportAssetParameters.OverridePipelines[GraphPipelineIndex], PipelineInstancesPackage);
 			if (!GeneratedPipeline)
 			{
-				UE_LOG(LogInterchangeEngine, Error, TEXT("Interchange import: Override pipeline array contains a NULL pipeline. Script or code need to be fix to avoid this. "));
+				UE_LOG(LogInterchangeEngine, Error, TEXT("Interchange Import: Overridden pipeline array contains a NULL pipeline. Fix your script or code to avoid this issue."));
 				continue;
 			}
 			else
@@ -2291,7 +2291,7 @@ bool UInterchangeManager::WarnIfInterchangeIsActive()
 		return false;
 	}
 	//Tell the user they have to cancel the import before closing the editor
-	FNotificationInfo Info(NSLOCTEXT("InterchangeManager", "WarnCannotProceed", "An import process is currently underway! Please cancel it to proceed!"));
+	FNotificationInfo Info(NSLOCTEXT("InterchangeManager", "WarnCannotProceed", "An import process is currently underway. Please cancel it to proceed."));
 	Info.ExpireDuration = 5.0f;
 	TSharedPtr<SNotificationItem> WarnNotification = FSlateNotificationManager::Get().AddNotification(Info);
 	if (WarnNotification.IsValid())
diff --git a/Engine/Source/Runtime/Interchange/Engine/Private/InterchangeMeshUtilities.cpp b/Engine/Source/Runtime/Interchange/Engine/Private/InterchangeMeshUtilities.cpp
--- a/Engine/Source/Runtime/Interchange/Engine/Private/InterchangeMeshUtilities.cpp
+++ b/Engine/Source/Runtime/Interchange/Engine/Private/InterchangeMeshUtilities.cpp
@@ -29,13 +29,13 @@ TFuture<bool> UInterchangeMeshUtilities::ImportCustomLod(UObject* MeshObject, co
 	TSharedPtr<TPromise<bool>> Promise = MakeShared<TPromise<bool>>();
 	if (!MeshObject)
 	{
-		UE_LOG(LogInterchangeEngine, Warning, TEXT("FInterchangeMeshUtilities::ImportCustomLod parameter MeshObject cannot be null."));
+		UE_LOG(LogInterchangeEngine, Warning, TEXT("FInterchangeMeshUtilities::ImportCustomLod: The MeshObject parameter cannot be null."));
 		Promise->SetValue(false);
 		return Promise->GetFuture();
 	}
 	if (!IsInGameThread())
 	{
-		UE_LOG(LogInterchangeEngine, Warning, TEXT("FInterchangeMeshUtilities::ImportCustomLod Cannot ask user a file path outside of the game thread."));
+		UE_LOG(LogInterchangeEngine, Warning, TEXT("FInterchangeMeshUtilities::ImportCustomLod: Cannot ask the user for a file path outside of the game thread."));
 		Promise->SetValue(false);
 		return Promise->GetFuture();
 	}
diff --git a/Engine/Source/Runtime/Interchange/Engine/Private/Tasks/InterchangeTaskImportObject.cpp b/Engine/Source/Runtime/Interchange/Engine/Private/Tasks/InterchangeTaskImportObject.cpp
--- a/Engine/Source/Runtime/Interchange/Engine/Private/Tasks/InterchangeTaskImportObject.cpp
+++ b/Engine/Source/Runtime/Interchange/Engine/Private/Tasks/InterchangeTaskImportObject.cpp
@@ -199,7 +199,7 @@ namespace UE::Interchange::Private
 					Message->SourceAssetName = AsyncHelper->SourceDatas[SourceIndex]->GetFilename();
 					Message->DestinationAssetName = AssetName;
 					Message->AssetType = FactoryNode->GetObjectClass();
-					Message->Text = NSLOCTEXT("InternalImportObjectStartup", "BadPackage", "It was not possible to create the asset as its package was not created correctly.");
+					Message->Text = NSLOCTEXT("InternalImportObjectStartup", "BadPackage", "It was not possible to create the asset because its package was not created correctly.");
 					return ErrorImportAssetResult;
 				}
 
@@ -208,7 +208,7 @@ namespace UE::Interchange::Private
 					UInterchangeResultError_Generic* Message = Factory->AddMessage<UInterchangeResultError_Generic>();
 					Message->DestinationAssetName = AssetName;
 					Message->AssetType = FactoryNode->GetObjectClass();
-					Message->Text = NSLOCTEXT("InternalImportObjectStartup", "SourceDataOrTranslatorInvalid", "It was not possible to create the asset as its translator was not created correctly.");
+					Message->Text = NSLOCTEXT("InternalImportObjectStartup", "SourceDataOrTranslatorInvalid", "It was not possible to create the asset because its translator was not created correctly.");
 					return ErrorImportAssetResult;
 				}
 			}
@@ -297,7 +297,7 @@ void UE::Interchange::FTaskImportObject_GameThread::DoTask(ENamedThreads::Type C
 		Message->SourceAssetName = AsyncHelper->SourceDatas[SourceIndex]->GetFilename();
 		Message->DestinationAssetName = ObjectToReimport ? ObjectToReimport->GetName() : FactoryNode->GetAssetName();
 		Message->AssetType = FactoryNode->GetObjectClass();
-		Message->Text = FText::Format(NSLOCTEXT("FTaskImportObject_GameThread", "GarbageCollectIsRunning_Error", "Cannot import asset '{0}'; the garbage collect is running during the import process.")
+		Message->Text = FText::Format(NSLOCTEXT("FTaskImportObject_GameThread", "GarbageCollectIsRunning_Error", "Cannot import asset '{0}'. The garbage collector is running during the import process.")
 			, FText::FromString(AssetName));
 		return;
 	}
@@ -309,7 +309,7 @@ void UE::Interchange::FTaskImportObject_GameThread::DoTask(ENamedThreads::Type C
 		Message->SourceAssetName = AsyncHelper->SourceDatas[SourceIndex]->GetFilename();
 		Message->DestinationAssetName = ObjectToReimport ? ObjectToReimport->GetName() : FactoryNode->GetAssetName();
 		Message->AssetType = FactoryNode->GetObjectClass();
-		Message->Text = FText::Format(NSLOCTEXT("FTaskImportObject_GameThread", "InvalidBaseNodeContainer", "Cannot import asset '{0}'; the node container is invalid for this source index.")
+		Message->Text = FText::Format(NSLOCTEXT("FTaskImportObject_GameThread", "InvalidBaseNodeContainer", "Cannot import asset '{0}'. The node container is invalid for this source index.")
 			, FText::FromString(AssetName));
 		return;
 	}
@@ -356,7 +356,7 @@ void UE::Interchange::FTaskImportObject_GameThread::DoTask(ENamedThreads::Type C
 			Message->SourceAssetName = AsyncHelper->SourceDatas[SourceIndex]->GetFilename();
 			Message->DestinationAssetName = AssetName;
 			Message->AssetType = FactoryNode->GetObjectClass();
-			Message->Text = NSLOCTEXT("FTaskImportObject_GameThread", "MapExistsWithSameName", "You cannot create an asset with this name, as there is already a map file with the same name in this folder.");
+			Message->Text = NSLOCTEXT("FTaskImportObject_GameThread", "MapExistsWithSameName", "You cannot create an asset with this name because there is already a map file with the same name in this folder.");
 			return;
 		}
 
@@ -383,7 +383,7 @@ void UE::Interchange::FTaskImportObject_GameThread::DoTask(ENamedThreads::Type C
 			Message->SourceAssetName = AsyncHelper->SourceDatas[SourceIndex]->GetFilename();
 			Message->DestinationAssetName = AssetName;
 			Message->AssetType = FactoryNode->GetObjectClass();
-			Message->Text = FText::Format(NSLOCTEXT("FTaskImportObject_GameThread", "CouldntCreatePackage", "It was not possible to create a package named '{0}'; the asset will not be imported.")
+			Message->Text = FText::Format(NSLOCTEXT("FTaskImportObject_GameThread", "CouldntCreatePackage", "It was not possible to create a package named '{0}'. The asset will not be imported.")
 				, FText::FromString(PackageName));
 			return;
 		}
@@ -461,7 +461,7 @@ void UE::Interchange::FTaskImportObject_GameThread::DoTask(ENamedThreads::Type C
 			const FText ExistingClassNameText = FText::FromString(ExistingAsset->GetClass()->GetName());
 			const FText AssetNameText = FText::FromString(AssetName);
 			const FText FolderNameText = FText::FromString(FPaths::GetPath(Pkg->GetPathName()));
-			Message->Text = FText::Format(NSLOCTEXT("FTaskImportObject_GameThread", "AssetVersusFactoryClassWrong_Error", "You cannot create a '{0}' asset named '{1}' in '{2}', as there is already a '{3}' asset with the same name in this folder.")
+			Message->Text = FText::Format(NSLOCTEXT("FTaskImportObject_GameThread", "AssetVersusFactoryClassWrong_Error", "You cannot create an asset of class '{0}' named '{1}' in '{2}' because there is already an asset of class '{3}' with the same name in that folder.")
 				, TargetClassNameText
 				, AssetNameText
 				, FolderNameText
@@ -479,7 +479,7 @@ void UE::Interchange::FTaskImportObject_GameThread::DoTask(ENamedThreads::Type C
 				Message->SourceAssetName = AsyncHelper->SourceDatas[SourceIndex]->GetFilename();
 				Message->DestinationAssetName = AssetName;
 				Message->AssetType = FactoryNode->GetObjectClass();
-				Message->Text = FText::Format(NSLOCTEXT("FTaskImportObject_GameThread", "InvalidExistingAsset", "An invalid asset exist at the same asset location [{0}]. This asset will be skip.")
+				Message->Text = FText::Format(NSLOCTEXT("FTaskImportObject_GameThread", "InvalidExistingAsset", "An invalid asset exists at the same asset location [{0}]. This asset will be skipped.")
 					, FText::FromString(AssetName));
 				return;
 			}
@@ -499,7 +499,7 @@ void UE::Interchange::FTaskImportObject_GameThread::DoTask(ENamedThreads::Type C
 			Message->SourceAssetName = AsyncHelper->SourceDatas[SourceIndex]->GetFilename();
 			Message->DestinationAssetName = AssetName;
 			Message->AssetType = FactoryNode->GetObjectClass();
-			Message->Text = FText::Format(NSLOCTEXT("FTaskImportObject_GameThread", "AssetCollisionBetweenImportTask", "Multiple import task are importing the same asset at the same location [{0}]. This asset will be skip. See more information in the log")
+			Message->Text = FText::Format(NSLOCTEXT("FTaskImportObject_GameThread", "AssetCollisionBetweenImportTask", "Multiple import tasks are importing the same asset at the same location [{0}]. This asset will be skipped. See the log for more information.")
 				, FText::FromString(AssetName));
 			return;
 		}
diff --git a/Engine/Source/Runtime/Interchange/Engine/Public/InterchangeAssetImportData.h b/Engine/Source/Runtime/Interchange/Engine/Public/InterchangeAssetImportData.h
--- a/Engine/Source/Runtime/Interchange/Engine/Public/InterchangeAssetImportData.h
+++ b/Engine/Source/Runtime/Interchange/Engine/Public/InterchangeAssetImportData.h
@@ -28,7 +28,7 @@ public:
 
 
 	/**
-	 * Return the first filename stored in this data. The resulting filename will be absolute (ie, not relative to the asset).
+	 * Return the first filename stored in this data. The resulting filename will be absolute (that is, not relative to the asset).
 	 */
 	UFUNCTION(BlueprintCallable, Category = "Interchange | AssetImportData")
 	FString ScriptGetFirstFilename() const
@@ -80,7 +80,7 @@ public:
 	}
 
 	/**
-	 * This function add tags to the asset registry.
+	 * This function adds tags to the asset registry.
 	 */
 	virtual void AppendAssetRegistryTags(FAssetRegistryTagsContext Context) override
 	{
@@ -107,11 +107,11 @@ public:
 #endif
 #endif
 
-	/** On a level import, set to the UInterchangeSceneImportAsset created during the import */
+	/** On a level import, set to the UInterchangeSceneImportAsset created during the import. */
 	UPROPERTY(EditAnywhere, Category = "Interchange | AssetImportData")
 	FSoftObjectPath SceneImportAsset;
 
-	/** Returns a pointer to the UInterchangeAssetImportData referred by the input object if applicable */
+	/** Returns a pointer to the UInterchangeAssetImportData referred to by the input object, if applicable. */
 	static UInterchangeAssetImportData* GetFromObject(UObject* Object)
 	{
 		if (Object)
@@ -130,7 +130,7 @@ public:
 		return nullptr;
 	}
 
-	/** The Node UID pass to the factory that exist in the graph that was use to create this asset */
+	/** The Node UID passed to the factory that existed in the graph that was used to create this asset. */
 	UPROPERTY(VisibleAnywhere, Category = "Interchange | AssetImportData")
 	FString NodeUniqueID;
 
@@ -192,7 +192,7 @@ private:
 };
 
 /**
- * Base class to create an asset import data converter
+ * Base class to create an asset import data converter.
  */
 UCLASS(Abstract, MinimalAPI)
 class UInterchangeAssetImportDataConverterBase : public UObject
@@ -201,10 +201,10 @@ class UInterchangeAssetImportDataConverterBase : public UObject
 public:
 	/**
 	 * Convert the asset import data from the one that is in the Object to
-	 * one supporting the target extension (i.e. legacy Fbx to Interchange or vice versa)
-	 * The function should return true only if it has convert the asset import data, false otherwise.
+	 * one that supports the target extension (for example, legacy FBX to Interchange or vice-versa)
+	 * The function should return true only if it has converted the asset import data, or false otherwise.
 	 * 
-	 * The system will call all object deriving from this class until one convert the data.
+	 * The system will call all objects that derive from this class until one converts the data.
 	 */
 	virtual bool ConvertImportData(UObject* Object, const FString& TargetExtension) const
 	{
diff --git a/Engine/Source/Runtime/Interchange/Engine/Public/InterchangeFilePickerBase.h b/Engine/Source/Runtime/Interchange/Engine/Public/InterchangeFilePickerBase.h
--- a/Engine/Source/Runtime/Interchange/Engine/Public/InterchangeFilePickerBase.h
+++ b/Engine/Source/Runtime/Interchange/Engine/Public/InterchangeFilePickerBase.h
@@ -14,15 +14,15 @@ struct FInterchangeFilePickerParameters
 {
 	GENERATED_BODY()
 
-	//If true the user will be able to select multiple files.
+	//If true, the user will be able to select multiple files.
 	UPROPERTY(EditAnywhere, Category = "Interchange | File Picker")
 	bool bAllowMultipleFiles = false;
 
-	//If not empty it will override the default title
+	//If not empty, this will override the default title.
 	UPROPERTY(EditAnywhere, Category = "Interchange | File Picker")
 	FText Title = FText();
 
-	//Set the default open path that the dialog will show to the user
+	//Set the default open path that the dialog will show to the user.
 	UPROPERTY(EditAnywhere, Category = "Interchange | File Picker")
 	FString DefaultPath = TEXT("");
 };
@@ -35,12 +35,12 @@ class UInterchangeFilePickerBase : public UObject
 public:
 
 	/**
-	 * Non virtual helper to allow blueprint to implement event base function to implement the function FilePickerForTranslatorAssetType,
+	 * Non-virtual helper that allows Blueprint to implement an event-based function to implement FilePickerForTranslatorAssetType().
 	 */
 	UFUNCTION(BlueprintNativeEvent, BlueprintCallable, Category = "Interchange | Utilities")
 	INTERCHANGEENGINE_API bool ScriptedFilePickerForTranslatorAssetType(const EInterchangeTranslatorAssetType TranslatorAssetType, FInterchangeFilePickerParameters& Parameters, TArray<FString>& OutFilenames);
 
-	/** The default implementation (call if the blueprint do not have any implementation) will call the virtual FilePickerForTranslatorAssetType */
+	/** The default implementation, which is called if the Blueprint does not have any implementation, calls the virtual FilePickerForTranslatorAssetType(). */
 	bool ScriptedFilePickerForTranslatorAssetType_Implementation(const EInterchangeTranslatorAssetType TranslatorAssetType, FInterchangeFilePickerParameters& Parameters, TArray<FString>& OutFilenames)
 	{
 		//By default we call the virtual function FilePickerForTranslatorAssetType
@@ -48,12 +48,12 @@ public:
 	}
 
 	/**
-	 * Non virtual helper to allow blueprint to implement event base function to implement the function FilePickerForTranslatorType,
+	 * Non-virtual helper that allows Blueprint to implement an event-based function to implement FilePickerForTranslatorType().
 	 */
 	UFUNCTION(BlueprintNativeEvent, BlueprintCallable, Category = "Interchange | Utilities")
 	INTERCHANGEENGINE_API bool ScriptedFilePickerForTranslatorType(const EInterchangeTranslatorType TranslatorType, FInterchangeFilePickerParameters& Parameters, TArray<FString>& OutFilenames);
 
-	/** The default implementation (call if the blueprint do not have any implementation) will call the virtual FilePickerForTranslatorType */
+	/** The default implementation, which is called if the Blueprint does not have any implementation, calls the virtual FilePickerForTranslatorType(). */
 	bool ScriptedFilePickerForTranslatorType_Implementation(const EInterchangeTranslatorType TranslatorAssetType, FInterchangeFilePickerParameters& Parameters, TArray<FString>& OutFilenames)
 	{
 		//By default we call the virtual function FilePickerForTranslatorType
@@ -67,7 +67,7 @@ protected:
 	 * This function must set OutFilename with a valid file that can be import to create the specified TranslatorAssetType.
 	 * We expect the user to see a file picker dialog when the editor is available.
 	 * 
-	 * @Return - true if it set OutFilename with a valid file path, false otherwise.
+	 * @Return - true if OutFilename was set with a valid file path, or false otherwise.
 	 */
 	virtual bool FilePickerForTranslatorAssetType(const EInterchangeTranslatorAssetType TranslatorAssetType, FInterchangeFilePickerParameters& Parameters, TArray<FString>& OutFilenames) { return false; }
 
@@ -75,7 +75,7 @@ protected:
 	 * This function must set OutFilename with a valid file that can be import to create the specified TranslatorType.
 	 * We expect the user to see a file picker dialog when the editor is available.
 	 *
-	 * @Return - true if it set OutFilename with a valid file path, false otherwise.
+	 * @Return - true if OutFilename was set with a valid file path, or false otherwise.
 	 */
 	virtual bool FilePickerForTranslatorType(const EInterchangeTranslatorType TranslatorType, FInterchangeFilePickerParameters& Parameters, TArray<FString>& OutFilenames) { return false; }
 
