Clean up typos and phrasing in Interchange modules - contribution from Robb Surridge

diff --git a/Engine/Source/Runtime/Interchange/Engine/Public/InterchangeManager.h b/Engine/Source/Runtime/Interchange/Engine/Public/InterchangeManager.h
--- a/Engine/Source/Runtime/Interchange/Engine/Public/InterchangeManager.h
+++ b/Engine/Source/Runtime/Interchange/Engine/Public/InterchangeManager.h
@@ -36,7 +36,7 @@ class UInterchangePipelineBase;
 class UInterchangePythonPipelineBase;
 class ULevel;
 
-/** Some utilities delegates for the automation of interchange */
+/** Some utility delegates for automating Interchange. */
 DECLARE_DYNAMIC_DELEGATE_OneParam(FOnObjectImportDoneDynamic, UObject*, Object);
 DECLARE_DELEGATE_OneParam(FOnObjectImportDoneNative, UObject*);
 
@@ -86,25 +86,25 @@ namespace UE
 
 		struct FImportAsyncHelperData
 		{
-			//True if the import process is unattended. We cannot show UI  if the import is automated
+			// True if the import process is unattended. We cannot show UI if the import is automated.
 			bool bIsAutomated = false;
 
-			// True if redirectors will be followed when determining what location to import an asset
+			// True if redirectors will be followed when determining what location to import an asset.
 			bool bFollowRedirectors = false;
 
-			//We can import assets or full scene
+			// We can import assets or full scene.
 			EImportType ImportType = EImportType::ImportType_None;
 
-			//True if we are reimporting assets or scene
+			// True if we are reimporting assets or a scene.
 			UObject* ReimportObject = nullptr;
 
-			//Level to import into in case of a scene import
+			// Level to import into, if we are doing a scene import.
 			ULevel* ImportLevel = nullptr;
 
-			/** Optional custom name to import as */
+			/** Optional custom name for the import. */
 			FString DestinationName;
 
-			/** Overwrite existing assets */
+			/** Whether or not to overwrite existing assets. */
 			bool bReplaceExisting = true;
 
 			/**
@@ -143,15 +143,15 @@ namespace UE
 			INTERCHANGEENGINE_API void SetDone();
 			INTERCHANGEENGINE_API void WaitUntilDone();
 
-			// Assets are only made available once they have been completely imported (passed through the entire import pipeline)
+			// Assets are only made available once they have been completely imported (passed through the entire import pipeline).
 			// While the status isn't EStatus::Done, the list can grow between subsequent calls.
 			// FAssetImportResult holds a reference to the assets so that they aren't garbage collected.
 			INTERCHANGEENGINE_API const TArray< UObject* >& GetImportedObjects() const;
 
-			// Helper to get the first asset of a certain class. Use when expecting a single asset of that class to be imported since the order isn't deterministic.
+			// Helper to get the first asset of a certain class. Use when expecting a single asset of that class to be imported, because the order isn't deterministic.
 			INTERCHANGEENGINE_API UObject* GetFirstAssetOfClass(UClass* InClass) const;
 
-			// Return the results of this asset import operation
+			// Return the results of this asset import operation.
 			UInterchangeResultsContainer* GetResults() const { return Results; }
 
 			// Adds an asset to the list of imported assets.
@@ -160,7 +160,7 @@ namespace UE
 			// Callback when the status switches to done.
 			INTERCHANGEENGINE_API void OnDone(TFunction< void(FImportResult&) > Callback);
 
-			// Internal delegates (use the FImportAssetParameters when calling the interchange import functions to set those)
+			// Internal delegates. To set these, use the FImportAssetParameters when calling the Interchange import functions.
 			FOnObjectImportDoneDynamic OnObjectDone;
 			FOnObjectImportDoneNative OnObjectDoneNative;
 
@@ -209,21 +209,21 @@ namespace UE
 				return TEXT("UE::Interchange::FImportAsyncHelper");
 			}
 
-			/** Unique id for this asynchelper. */
+			/** Unique ID for this async helper. */
 			int32 UniqueId;
 
-			//The content base path where to import
+			//The base path to import the content.
 			FString ContentBasePath;
 
-			//The following Arrays are per source data
+			//The following arrays are per source data
 			TArray<TStrongObjectPtr<UInterchangeBaseNodeContainer>> BaseNodeContainers;
 			TArray<TObjectPtr<UInterchangeSourceData>> SourceDatas;
 			TArray<TObjectPtr<UInterchangeTranslatorBase>> Translators;
 
-			//Pipelines array is not per source data
+			//The pipelines array is not per source data.
 			TArray<TObjectPtr<UInterchangePipelineBase>> Pipelines;
-			//The original pipelines asset to save in the asset reimport data. The original pipeline can restore python class member value.
-			//Python class instanced assets cannot be saved, so we have to serialize in json the data to restore it when we do a re-import.
+			//The original pipelines asset to save in the asset reimport data. The original pipeline can restore Python class member value.
+			//Python class instanced assets cannot be saved, so we have to serialize in JSON the data to restore it when we do a reimport.
 			TArray<UObject*> OriginalPipelines;
 
 			TArray<FGraphEventRef> TranslatorTasks;
@@ -242,7 +242,7 @@ namespace UE
 			// Package where the Pipeline Instances are stored during an import.
 			FString PipelineInstancesPackageName;
 
-			//Return true if we can import this class, false otherwise
+			//Return true if we can import this class, or false otherwise.
 			bool IsClassImportAllowed(UClass* Class);
 
 			UPackage* GetCreatedPackage(const FString& PackageName) const;
@@ -253,9 +253,9 @@ namespace UE
 
 			struct FImportedObjectInfo
 			{
-				UObject* ImportedObject = nullptr; // The object that was imported
-				UInterchangeFactoryBase* Factory = nullptr; //The factory that created the imported object
-				UInterchangeFactoryBaseNode* FactoryNode = nullptr; //The node that describes the object
+				UObject* ImportedObject = nullptr; // The object that was imported.
+				UInterchangeFactoryBase* Factory = nullptr; //The factory that created the imported object.
+				UInterchangeFactoryBaseNode* FactoryNode = nullptr; //The node that describes the object.
 				bool bIsReimport;
 			};
 
@@ -270,7 +270,7 @@ namespace UE
 			void IterateImportedSceneObjectsPerSourceIndex(TFunction< void(int32 SourceIndex, const TArray<FImportedObjectInfo>& ImportedObjects) > Callback) const;
 
 			/*
-			 * Return true if the Object is imported by this async import, false otherwise
+			 * Return true if the Object is imported by this async import, or false otherwise.
 			 */
 			bool IsImportingObject(UObject* Object) const;
 
@@ -279,14 +279,14 @@ namespace UE
 			FAssetImportResultRef AssetImportResult;
 			FSceneImportResultRef SceneImportResult;
 			
-			//If we cancel the tasks, we set this boolean to true
+			//If we cancel the tasks, we set this Boolean to true.
 			std::atomic<bool> bCancel;
 
 			void SendAnalyticImportEndData();
 			void ReleaseTranslatorsSource();
 
 			/**
-			 * Wait synchronously after graph parsing task is done and return the GraphEventArray up to the completion TaskGraphEvent
+			 * Wait synchronously after the graph parsing task is done, and return the GraphEventArray up to the completion TaskGraphEvent.
 			 */
 			FGraphEventArray GetCompletionTaskGraphEvent();
 
@@ -296,16 +296,16 @@ namespace UE
 
 		private:
 			FCriticalSection ClassPermissionLock;
-			// Set of classes which creation has been denied
+			// Set of classes whose creation has been denied.
 			TSet<UClass*> DeniedClasses;
-			// Set of classes which creation is allowed
+			// Set of classes whose creation is allowed.
 			TSet<UClass*> AllowedClasses;
 
-			//Create package map, Key is package name. We cannot create package asynchronously so we have to create a game thread task to do this
+			//Created package map. The key is the package name. We cannot create packages asynchronously, so we have to create a game thread task to do this.
 			mutable FCriticalSection CreatedPackagesLock;
 			TMap<FString, UPackage*> CreatedPackages;
 
-			// Created factories map, Key is factory node UID
+			// Created factories map. The key is the factory node UID.
 			mutable FCriticalSection CreatedFactoriesLock;
 			TMap<FString, TObjectPtr<UInterchangeFactoryBase>> CreatedFactories;
 
@@ -320,7 +320,7 @@ namespace UE
 
 		void SanitizeObjectName(FString& ObjectName);
 
-		/* This function take an asset representing a pipeline and generate a UInterchangePipelineBase asset. */
+		/* This function takes an asset that represents a pipeline and generates a UInterchangePipelineBase asset. */
 		INTERCHANGEENGINE_API UInterchangePipelineBase* GeneratePipelineInstance(const FSoftObjectPath& PipelineInstance, UPackage* PipelineInstancePackage = nullptr);
 
 		INTERCHANGEENGINE_API UInterchangePipelineBase* GeneratePipelineInstanceInSourceAssetPackage(const FSoftObjectPath& PipelineInstance);
@@ -329,7 +329,7 @@ namespace UE
 } //ns UE
 
 /**
- * This class is use to pass override pipelines in the ImportAssetTask Options member
+ * This class is used to pass override pipelines in the ImportAssetTask Options member.
  */
 UCLASS(Transient, BlueprintType, MinimalAPI)
 class UInterchangePipelineStackOverride : public UObject
@@ -355,56 +355,56 @@ struct FImportAssetParameters
 {
 	GENERATED_USTRUCT_BODY()
 
-	// If the import is a reimport for a specific asset set the asset to reimport here
+	// If the import is a reimport for a specific asset, set the asset to reimport here.
 	UPROPERTY(EditAnywhere, BlueprintReadWrite, Category = "Interchange|ImportAsset")
 	TObjectPtr<UObject> ReimportAsset = nullptr;
 
-	// If we are doing a reimport, set the source index here. Some asset have more then one source file when they import partial part of there content.
+	// If we are doing a reimport, set the source index here. Some assets have more then one source file that each contains part of the asset content.
 	UPROPERTY(EditAnywhere, BlueprintReadWrite, Category = "Interchange|ImportAsset")
 	int32 ReimportSourceIndex = INDEX_NONE;
 
-	// Tell interchange that import is automated and it shouldn't present a model window
+	// Tell Interchange that import is automated and it shouldn't present a modal window.
 	UPROPERTY(EditAnywhere, BlueprintReadWrite, Category = "Interchange|ImportAsset")
 	bool bIsAutomated = false;
 
-	// Tell interchange to follow redirectors when determining the location an asset will be imported
+	// Tell Interchange to follow redirectors when determining the location an asset will be imported.
 	UPROPERTY(EditAnywhere, BlueprintReadWrite, Category = "Interchange|ImportAsset")
 	bool bFollowRedirectors = false;
 
-	// Adding some override will tell interchange to use the specific custom set pipelines instead of letting the user or the system chose
+	// Adding overrides tells Interchange to use the specific custom set of pipelines instead of letting the user or the system choose.
 	UPROPERTY(EditAnywhere, BlueprintReadWrite, Category = "Interchange|ImportAsset", meta = (AllowedClasses = "/Script/InterchangeCore.InterchangePipelineBase, /Script/InterchangeEngine.InterchangeBlueprintPipelineBase, /Script/InterchangeEngine.InterchangePythonPipelineAsset"))
 	TArray<FSoftObjectPath> OverridePipelines;
 
-	//Level to import into in case of a scene import
+	//Level to import into when doing a scene import.
 	UPROPERTY(EditAnywhere, BlueprintReadWrite, Category = "Interchange|ImportAsset")
 	TObjectPtr<ULevel> ImportLevel = nullptr;
 
-	/** Optional custom name to import as */
+	/** Optional custom name for the import. */
 	UPROPERTY(EditAnywhere, BlueprintReadWrite, Category = "Interchange|ImportAsset")
 	FString DestinationName;
 
-	/** Overwrite existing assets */
+	/** Determies whether to overwrite existing assets. */
 	UPROPERTY(EditAnywhere, BlueprintReadWrite, Category = "Interchange|ImportAsset")
 	bool bReplaceExisting = true;
 
-	/* Delegates used track the imported objects */
+	/* Delegates used to track the imported objects. */
 
-	// This is called each time an asset is imported or reimported from the import call
+	// This is called each time an asset is imported or reimported from the import call.
 	UPROPERTY(EditAnywhere, BlueprintReadWrite, Category = "Interchange|ImportAsset", meta=(PinHiddenByDefault))
 	FOnObjectImportDoneDynamic OnAssetDone;
 	FOnObjectImportDoneNative OnAssetDoneNative;
 
-	// This is called when all the assets where imported from the source data
+	// This is called when all the assets were imported from the source data.
 	UPROPERTY(EditAnywhere, BlueprintReadWrite, Category = "Interchange|ImportAsset", meta=(PinHiddenByDefault))
 	FOnImportDoneDynamic OnAssetsImportDone;
 	FOnImportDoneNative OnAssetsImportDoneNative;
 
-	// This is called each time an object in the scene is imported or reimported from the import call
+	// This is called each time an object in the scene is imported or reimported from the import call.
 	UPROPERTY(EditAnywhere, BlueprintReadWrite, Category = "Interchange|ImportAsset", meta=(PinHiddenByDefault))
 	FOnObjectImportDoneDynamic OnSceneObjectDone;
 	FOnObjectImportDoneNative OnSceneObjectDoneNative;
 
-	// This is called when all the scene object where imported from the source data
+	// This is called when all the scene objects were imported from the source data.
 	UPROPERTY(EditAnywhere, BlueprintReadWrite, Category = "Interchange|ImportAsset", meta=(PinHiddenByDefault))
 	FOnImportDoneDynamic OnSceneImportDone;
 	FOnImportDoneNative OnSceneImportDoneNative;
@@ -419,9 +419,9 @@ public:
 	UInterchangeManager(const FObjectInitializer& ObjectInitializer);
 
 	/**
-	 * Return the interchange manager singleton pointer.
+	 * Return the pointer to the Interchange Manager singleton.
 	 *
-	 * @note - We need to return a pointer to have a valid Blueprint callable function
+	 * @note - We need to return a pointer to have a valid Blueprint-callable function.
 	 */
 	UFUNCTION(BlueprintCallable, Category = "Interchange | Import Manager")
 	static UInterchangeManager* GetInterchangeManagerScripted()
@@ -429,22 +429,22 @@ public:
 		return &GetInterchangeManager();
 	}
 
-	/** Return the interchange manager singleton.*/
+	/** Return the Interchange Manager singleton. */
 	static INTERCHANGEENGINE_API UInterchangeManager& GetInterchangeManager();
 
-	/** Return the CVar which make interchange enable or not.*/
+	/** Return the CVar that enables or disables Interchange. */
 	static INTERCHANGEENGINE_API bool IsInterchangeImportEnabled();
-	/** Set the CVar which make interchange enable or not.*/
+	/** Set the CVar that enables or disables Interchange. */
 	static INTERCHANGEENGINE_API void SetInterchangeImportEnabled(bool bEnabled);
 
-	/** delegate type fired when new assets have been imported. Note: InCreatedObject can be NULL if import failed. Params: UFactory* InFactory, UObject* InCreatedObject */
+	/** Delegate type that is fired when new assets have been imported. Note: InCreatedObject can be NULL if the import failed. Params: UFactory* InFactory, UObject* InCreatedObject. */
 	DECLARE_MULTICAST_DELEGATE_OneParam(FInterchangeOnAssetPostImport, UObject*);
-	/** delegate type fired when new assets have been reimported. Note: InCreatedObject can be NULL if import failed. Params: UObject* InCreatedObject */
+	/** Delegate type that is fired when new assets have been reimported. Note: InCreatedObject can be NULL if the import failed. Params: UObject* InCreatedObject. */
 	DECLARE_MULTICAST_DELEGATE_OneParam(FInterchangeOnAssetPostReimport, UObject*);
-	/** delegate type fired when import results in an error */
+	/** Delegate type that is fired when the import results in an error. */
 	DECLARE_MULTICAST_DELEGATE_OneParam(FInterchangeOnBatchImportComplete, TStrongObjectPtr<UInterchangeResultsContainer>);
 
-	// Delegates used to register and unregister
+	// Delegates used to register and unregister.
 
 	FInterchangeOnAssetPostImport OnAssetPostImport;
 	FInterchangeOnAssetPostReimport OnAssetPostReimport;
@@ -453,47 +453,47 @@ public:
 	FSimpleMulticastDelegate OnPreDestroyInterchangeManager;
 
 	/**
-	 * Any translator must register to the manager
-	 * @Param Translator - The UClass of the translator you want to register
-	 * @return true if the translator class can be register false otherwise.
+	 * All translators must be registered with the manager.
+	 * @Param Translator - The UClass of the translator you want to register.
+	 * @return true if the translator class can be registered, or false otherwise.
 	 *
-	 * @Note if you register multiple time the same class it will return true for every call
-	 * @Note The order in which the translators are registered will be the same as the order used to select a translator to import a file
+	 * @Note If you register the same class multiple times, this returns true for every call.
+	 * @Note The order in which the translators are registered will be the same as the order used to select a translator to import a file.
 	 */
 	INTERCHANGEENGINE_API bool RegisterTranslator(const UClass* TranslatorClass);
 
 	/**
-	 * Any factory must register to the manager
-	 * @Param Factory - The UClass of the factory you want to register
-	 * @return true if the factory class can be register false otherwise.
+	 * All factories must be registered with the manager.
+	 * @Param Factory - The UClass of the factory you want to register.
+	 * @return true if the factory class can be registered, or false otherwise.
 	 *
-	 * @Note if you register multiple time the same class it will return true for every call
+	 * @Note If you register the same class multiple time, this returns true for every call.
 	 */
 	INTERCHANGEENGINE_API bool RegisterFactory(const UClass* Factory);
 
 	/**
-	 * Any writer must register to the manager
-	 * @Param Writer - The UClass of the writer you want to register
-	 * @return true if the writer class can be register false otherwise.
+	 * All writers must be registered with the manager,
+	 * @Param Writer - The UClass of the writer you want to register,
+	 * @return true if the writer class can be registered, or false otherwise.
 	 *
-	 * @Note if you register multiple time the same class it will return true for every call
+	 * @Note If you register the same class multiple time, this returns true for every call.
 	 */
 	INTERCHANGEENGINE_API bool RegisterWriter(const UClass* Writer);
 
 	/**
-	 * Any converter must register to the manager
-	 * @Param Converter - The UClass of the converter you want to register
-	 * @return true if the converter class can be register false otherwise.
+	 * All converters must be registered with the manager.
+	 * @Param Converter - The UClass of the converter you want to register.
+	 * @return true if the converter class can be registered, or false otherwise.
 	 *
-	 * @Note if you register multiple time the same class it will return true for every call
+	 * @Note If you register the same class multiple time, this returns true for every call.
 	 */
 	INTERCHANGEENGINE_API bool RegisterImportDataConverter(const UClass* Converter);
 
 	/**
-	 * Call all the registered converter, if one converter want 
-	 * @Param Object - The Object to convert the import data
-	 * @Param Extension - The file extension we want to import
-	 * @return true if one of the converter has convert the data. False otherwise.
+	 * Call all the registered converters to see if any converter can convert the data.
+	 * @Param Object - The Object to convert the import data.
+	 * @Param Extension - The file extension we want to import.
+	 * @return true if one of the converter has converted the data, or false otherwise.
 	 */
 	INTERCHANGEENGINE_API bool ConvertImportData(UObject* Object, const FString& Extension) const;
 
@@ -521,44 +521,43 @@ public:
 	INTERCHANGEENGINE_API TArray<FString> GetSupportedFormatsForObject(const UObject* Object) const;
 
 	/**
-	 * Look if there is a registered translator for this source data.
-	 * This allow us to by pass the original asset tools system to import supported asset.
-	 * @Param SourceData - The source data input we want to translate to Uod
-	 * @Param bSceneImportOnly - Enforces a check on the type of translator
-	 * @return True if there is a registered translator that can handle handle this source data, false otherwise, when bSceneImportOnly is false.
+	 * Check whether there is a registered translator for this source data.
+	 * This allows us to bypass the original asset tools system to import supported assets.
+	 * @Param SourceData - The source data input we want to translate to Uod.
+	 * @return True if there is a registered translator that can handle handle this source data, or false otherwise, when bSceneImportOnly is false.
 	 * Otherwise, returns true only if the translator supports level import
 	 * @Note: Temporary until FBX level import is production ready
 	 */
 	INTERCHANGEENGINE_API bool CanTranslateSourceData(const UInterchangeSourceData* SourceData, bool bSceneImportOnly = false) const;
 
 	/**
-	 * Returns true if Interchange can create that type of assets and is able to translate its source file.
+	 * Returns true if Interchange can create this type of asset and is able to translate its source files.
 	 * @Param Object - The object we want to reimport.
 	 * @Param OutFilenames - An array that is filled with the object's source filenames if the operation is successful.
 	 */
 	INTERCHANGEENGINE_API bool CanReimport(const UObject* Object, TArray<FString>& OutFilenames) const;
 
 	/**
-	 * Call this to start an import asset process, the caller must specify a source data.
-	 * This import process can import many different asset, but all in the game content.
+	 * Call this to start an asset import process. The caller must specify the source data.
+	 * This process can import many different assets into the game content.
 	 *
-	 * @Param ContentPath - The content path where to import the assets
-	 * @Param SourceData - The source data input we want to translate
-	 * @param ImportAssetParameters - All import asset parameter we need to pass to the import asset function
-	 * @return true if the import succeed, false otherwise.
+	 * @Param ContentPath - The path where the imported assets will be created.
+	 * @Param SourceData - The source data input to translate.
+	 * @param ImportAssetParameters - All parameters that need to be passed to the import asset function.
+	 * @return true if the import succeeds, or false otherwise.
 	 */
 	UFUNCTION(BlueprintCallable, Category = "Interchange | Import Manager")
 	INTERCHANGEENGINE_API bool ImportAsset(const FString& ContentPath, const UInterchangeSourceData* SourceData, const FImportAssetParameters& ImportAssetParameters);
 	INTERCHANGEENGINE_API UE::Interchange::FAssetImportResultRef ImportAssetAsync(const FString& ContentPath, const UInterchangeSourceData* SourceData, const FImportAssetParameters& ImportAssetParameters);
 
 	/**
-	 * Call this to start an import scene process, the caller must specify a source data.
-	 * This import process can import many different asset and there transform (USceneComponent) and store the result in a blueprint and add the blueprint to the level.
+	 * Call this to start a scene import process. The caller must specify the source data.
+	 * This process can import many different assets and their transforms (USceneComponent), store the result in a Blueprint, and add the Blueprint to the level.
 	 *
-	 * @Param ContentPath - The content path where to import the assets
-	 * @Param SourceData - The source data input we want to translate, this object will be duplicate to allow multithread safe operations
-	 * @param ImportAssetParameters - All import asset parameter we need to pass to the import asset function
-	 * @return true if the import succeed, false otherwise.
+	 * @Param ContentPath - The path where the imported assets will be created.
+	 * @Param SourceData - The source data input to translate. This object will be duplicated to allow thread-safe operations.
+	 * @param ImportAssetParameters - All parameters that need to be passed to the import asset function.
+	 * @return true if the import succeeds, or false otherwise.
 	 */
 	UFUNCTION(BlueprintCallable, Category = "Interchange | Import Manager")
 	INTERCHANGEENGINE_API bool ImportScene(const FString& ContentPath, const UInterchangeSourceData* SourceData, const FImportAssetParameters& ImportAssetParameters);
@@ -567,44 +566,43 @@ public:
 	ImportSceneAsync(const FString& ContentPath, const UInterchangeSourceData* SourceData, const FImportAssetParameters& ImportAssetParameters);
 
 	/**
-	 * Call this to start an export asset process, the caller must specify a source data.
+	 * Call this to start an asset export process. The caller must specify a source data.
 	 * 
-	 * @Param SourceData - The source data output 
-	 * @Param bIsAutomated - If true the exporter will not show any UI or dialog
-	 * @return true if the import succeed, false otherwise.
+	 * @Param Asset - The asset to export.
+	 * @Param bIsAutomated - If true, the exporter will not show any UI or dialogs.
+	 * @return true if the export succeeds, or false otherwise.
 	 */
 	UFUNCTION(BlueprintCallable, Category = "Interchange | Export Manager")
 	INTERCHANGEENGINE_API bool ExportAsset(const UObject* Asset, bool bIsAutomated = false);
 
 	/**
-	 * Call this to start an export scene process, the caller must specify a source data
-	 * This import process can import many different asset and there transform (USceneComponent) and store the result in a blueprint and add the blueprint to the level.
-	 * @Param SourceData - The source data input we want to translate
-	 * @Param bIsAutomated - If true the import asset will not show any UI or dialog
-	 * @return true if the import succeed, false otherwise.
+	 * Call this to start a scene export process. The caller must specify a source data.
+	 * @Param World - The scene to export.
+	 * @Param bIsAutomated - If true, the import process will not show any UI or dialogs.
+	 * @return true if the export succeeds, or false otherwise.
 	 */
 	UFUNCTION(BlueprintCallable, Category = "Interchange | Export Manager")
 	INTERCHANGEENGINE_API bool ExportScene(const UObject* World, bool bIsAutomated = false);
 
 	/*
-	* Script helper to create a source data object pointing on a file on disk
-	* @Param InFilename: Specify a file on disk
+	* Script helper to create a source data object that points to a file on disk.
+	* @Param InFilename: Specify a file on disk.
 	* @return: A new UInterchangeSourceData.
 	*/
 	UFUNCTION(BlueprintCallable, Category = "Interchange | Import Manager")
 	static INTERCHANGEENGINE_API UInterchangeSourceData* CreateSourceData(const FString& InFileName);
 
 	/**
-	* Script helper to get a registered factory for a specified class
-	* @Param FactoryClass: The class we search a registerd factory
-	* @return: if found, we return the factory class that is registered. Return NULL if nothing found.
+	* Script helper to get a registered factory for a specified class.
+	* @Param FactoryClass: The class whose registered factory you want to find.
+	* @return: The registered factory class if found, or NULL if no registered factory was found.
 	*/
 	UFUNCTION(BlueprintCallable, Category = "Interchange | Import Manager")
 	INTERCHANGEENGINE_API const UClass* GetRegisteredFactoryClass(const UClass* ClassToMake) const;
 
 	/**
-	 * Return an FImportAsynHelper pointer. The pointer is deleted when ReleaseAsyncHelper is call.
-	 * @param Data - The data we want to pass to the different import tasks
+	 * Return a pointer to an FImportAsyncHelper. The pointer is deleted when ReleaseAsyncHelper is called.
+	 * @param Data - The data to pass to the different import tasks.
 	 */
 	INTERCHANGEENGINE_API TSharedRef<UE::Interchange::FImportAsyncHelper, ESPMode::ThreadSafe> CreateAsyncHelper(const UE::Interchange::FImportAsyncHelperData& Data, const FImportAssetParameters& ImportAssetParameters);
 
@@ -613,40 +611,40 @@ public:
 
 	/*
 	 * Return the first translator that can translate the source data.
-	 * @Param SourceData - The source data for which we search a translator.
-	 * @return return a matching translator or nullptr if there is no match.
+	 * @Param SourceData - The source data that you want a translator for.
+	 * @return return a matching translator, or nullptr if no translators exist for the source data.
 	 */
 	INTERCHANGEENGINE_API UInterchangeTranslatorBase* GetTranslatorForSourceData(const UInterchangeSourceData* SourceData) const;
 
 	/**
-	 * Return true if the Interchange is active (importing or exporting), return false otherwise.
+	 * Return true if Interchange is actively importing or exporting, or false otherwise.
 	 */
 	INTERCHANGEENGINE_API bool IsInterchangeActive();
 
 	/**
-	 * Return false if the Interchange is not active (importing or exporting).
-	 * If the interchange is active, it will display a notification to let the user know they can cancel the asynchronous import/export
-	 * to be able to complete the operation they want to do. (The exit editor operation is calling this)
+	 * Return false if Interchange is not actively importing or exporting.
+	 * If Interchange is active, it will display a notification to let the user know they can cancel the asynchronous import/export
+	 * to be able to complete the operation they requested. (This is called by the exit editor operation.)
 	 */
 	INTERCHANGEENGINE_API bool WarnIfInterchangeIsActive();
 
 	/**
-	 * Look if there is a translator registered that can translate the source data with the specified PayloadInterface
-	 * @Param SourceData - The source data input we want to translate to Uod
-	 * @return true if the source data can be translated using the specified PayloadInterface, false otherwise.
+	 * Check whether there is a translator registered that can translate the source data with the specified PayloadInterface.
+	 * @Param SourceData - The source data input you want to translate to Uod.
+	 * @return true if the source data can be translated using the specified PayloadInterface, or false otherwise.
 	 */
 	INTERCHANGEENGINE_API bool CanTranslateSourceDataWithPayloadInterface(const UInterchangeSourceData* SourceData, const UClass* PayloadInterfaceClass) const;
 
 	/*
 	 * Return the first translator that can translate the source data with the specified PayloadInterface.
-	 * @Param SourceData - The source data for which we search a translator.
+	 * @Param SourceData - The source data that you want a translator for.
 	 * @Param PayloadInterfaceClass - The PayloadInterface that the translator must implement.
-	 * @return return a matching translator implementing the specified PayloadInterface or nullptr if there is no match.
+	 * @return return a matching translator that implements the specified PayloadInterface, or nullptr if none exists.
 	 */
 	INTERCHANGEENGINE_API UInterchangeTranslatorBase* GetTranslatorSupportingPayloadInterfaceForSourceData(const UInterchangeSourceData* SourceData, const UClass* PayloadInterfaceClass) const;
 
 	/**
-	 * Return true if the object is being imported, return false otherwise. If the user import multiple file in the same folder its possible to
+	 * Return true if the object is being imported, or false otherwise. If the user imports multiple file in the same folder, it's possible to
 	 * have the same asset name in two different files.
 	 */
 	INTERCHANGEENGINE_API bool IsObjectBeingImported(UObject* Object) const;
@@ -659,43 +657,43 @@ public:
 	INTERCHANGEENGINE_API bool EnqueuePostImportTask(TSharedPtr<FInterchangePostImportTask> PostImportTask);
 protected:
 
-	/** Return true if we can show some UI */
+	/** Return true if Interchange can show UI. */
 	static INTERCHANGEENGINE_API bool IsAttended();
 
 	/*
-	 * Find all Pipeline candidate (c++, blueprint and python).
-	 * @Param SourceData - The source data for which we search a translator.
-	 * @return return a matching translator or nullptr if there is no match.
+	 * Find all pipeline candidates (C++, Blueprint, and Python).
+	 * @Param SourceData - The source data you need a translator for.
+	 * @return return a matching translator, or nullptr if there is no match.
 	 */
 	INTERCHANGEENGINE_API void FindPipelineCandidate(TArray<UClass*>& PipelineCandidates);
 
 	/**
-	 * This function cancel all task and finish them has fast as possible.
-	 * We use this if the user cancel the work or if the editor is exiting.
-	 * @note - This is a asynchronous call, tasks will be completed (cancel) soon.
+	 * This function cancels all tasks and finishes them as fast as possible.
+	 * We use this if the user cancels the work or if the editor exits.
+	 * @note - This is a asynchronous call. Tasks will be completed (canceled) soon.
 	 */
 	INTERCHANGEENGINE_API void CancelAllTasks();
 
 	/**
-	 * Wait synchronously that all tasks are done
+	 * Wait synchronously until all tasks are done.
 	 */
 	INTERCHANGEENGINE_API void WaitUntilAllTasksDone(bool bCancel);
 
 
 	/**
-	 * If we set the mode to active we will setup the timer and add the thread that will block the GC.
-	 * If the we set the mode to inactive we will remove the timer and finish the thread that block the GC.
+	 * If we set the mode to active, we will setup the timer and add the thread that will block the GC.
+	 * If the we set the mode to inactive, we will remove the timer and finish the thread that blocks the GC.
 	 */
 	INTERCHANGEENGINE_API void SetActiveMode(bool IsActive);
 
 	/**
-	 * Start task until we reach the taskgraph worker number.
-	 * @param bCancelAllTasks - If true we will start all task but with the cancel state set, so task will complete fast and call the completion task
+	 * Start tasks until we reach the taskgraph worker number.
+	 * @param bCancelAllTasks - If true, we will start all tasks but with the cancel state set, so tasks will complete fast and call the completion task.
 	 */
 	INTERCHANGEENGINE_API void StartQueuedTasks(bool bCancelAllTasks = false);
 
 	/**
-	 * Called by the public Import functions
+	 * Called by the public Import functions.
 	 */
 	INTERCHANGEENGINE_API TTuple<UE::Interchange::FAssetImportResultRef, UE::Interchange::FSceneImportResultRef>
 	ImportInternal(const FString& ContentPath, const UInterchangeSourceData* SourceData, const FImportAssetParameters& ImportAssetParameters, const UE::Interchange::EImportType ImportType);
@@ -713,25 +711,25 @@ private:
 	TMap<UClass*, bool> NonParallelTranslatorLocks;
 	TMap<UClass*, TArray<FQueuedTaskData>> NonParallelTranslatorQueueTasks;
 	
-	//Queue all incoming tasks if there is more started task then we have cores
+	//Queue all incoming tasks if there are more started task then we have cores.
 	TQueue<FQueuedTaskData> QueuedTasks;
 	int32 QueueTaskCount = 0;
 
-	//Queue all incoming task that are not direct import or re-import of assets.
-	// The post import tasks are execute only when all the QueuedTasks are completed and empty
-	//Those task can start an import task (QueuedTasks), one example is the skeletal mesh alternate skinning which require to import several file for each profiles
+	//Queue all incoming tasks that are not directly import or re-import of assets.
+	// The post import tasks are executed only when all the QueuedTasks are completed and empty.
+	//Those tasks can start an import task (QueuedTasks), one example is the skeletal mesh alternate skinning which requires to import several files for each profile.
 	TQueue<TSharedPtr<FInterchangePostImportTask>> QueuedPostImportTasks;
 
 	//Ticker which is on only if we have some QueuedPostImportTasks
 	FTSTicker::FDelegateHandle	QueuedPostImportTasksTickerHandle;
 
-	//By using pointer, there is no issue if the array get resize
+	//By using pointer, there is no issue if the array get resized.
 	TArray<TSharedPtr<UE::Interchange::FImportAsyncHelper, ESPMode::ThreadSafe> > ImportTasks;
 
 	TSharedPtr<FAsyncTaskNotification> Notification = nullptr;
 	FTSTicker::FDelegateHandle NotificationTickHandle;
 
-	// Caching the registered translator classes to avoid double registration fast
+	// Caching the registered translator classes to avoid double registration.
 	UPROPERTY()
 	TSet<TObjectPtr<const UClass>> RegisteredTranslatorsClass;
 
@@ -744,30 +742,30 @@ private:
 //	UPROPERTY()
 //	TMap<TObjectPtr<const UClass>, TObjectPtr<UInterchangePipelineBase> > RegisteredPipelines;
 
-	//The manager will create only one factory per type
+	//The manager will create only one factory per type.
 	UPROPERTY()
 	TMap<TObjectPtr<const UClass>, TObjectPtr<const UClass> > RegisteredFactoryClasses;
 
-	//The manager will create only one writer per type
+	//The manager will create only one writer per type.
 	UPROPERTY()
 	TMap<TObjectPtr<const UClass>, TObjectPtr<UInterchangeWriterBase> > RegisteredWriters;
 
-	//The manager will create only one converter per type
+	//The manager will create only one converter per type.
 	UPROPERTY()
 	TMap<TObjectPtr<const UClass>, TObjectPtr<UInterchangeAssetImportDataConverterBase> > RegisteredConverters;
 
-	//If interchange is currently importing we have a timer to watch the cancel and we block GC 
+	//If interchange is currently importing, we have a timer to watch the cancel and we block GC.
 	FThreadSafeBool bIsActive = false;
 
-	//If the user wants to use the same import pipeline stack for all the queue task
+	//If the user wants to use the same import pipeline stack for all queued tasks.
 	//This map is reset when the ImportTasks array is empty.
 	TMap<UClass*, TArray<UInterchangePipelineBase*>> ImportAllWithSamePipelines;
 
 	//Indicates that the import process was canceled by the user.
-	//This boolean is reset to false when the ImportTasks array is empty.
+	//This Boolean is reset to false when the ImportTasks array is empty.
 	bool bImportCanceled = false;
 
-	//We want to avoid starting an import task during a GC
+	//We want to avoid starting an import task during a GC.
 	FDelegateHandle GCEndDelegate;
 	bool bGCEndDelegateCancellAllTask = false;
 
diff --git a/Engine/Source/Runtime/Interchange/Engine/Public/InterchangePipelineConfigurationBase.h b/Engine/Source/Runtime/Interchange/Engine/Public/InterchangePipelineConfigurationBase.h
--- a/Engine/Source/Runtime/Interchange/Engine/Public/InterchangePipelineConfigurationBase.h
+++ b/Engine/Source/Runtime/Interchange/Engine/Public/InterchangePipelineConfigurationBase.h
@@ -39,7 +39,7 @@ class UInterchangePipelineConfigurationBase : public UObject
 public:
 
 	/**
-	 * Non virtual helper to allow blueprint to implement event base function to implement the ShowPipelineConfigurationDialog,
+	 * Non-virtual helper that allows Blueprint to implement an event-based function to implement ShowPipelineConfigurationDialog().
 	 */
 	UFUNCTION(BlueprintNativeEvent, BlueprintCallable, Category = "Interchange | Translator")
 	INTERCHANGEENGINE_API EInterchangePipelineConfigurationDialogResult ScriptedShowPipelineConfigurationDialog(TArray<FInterchangeStackInfo>& PipelineStacks
@@ -48,7 +48,7 @@ public:
 		, UInterchangeTranslatorBase* Translator
 		, UInterchangeBaseNodeContainer* BaseNodeContainer);
 
-	/** The default implementation (call if the blueprint do not have any implementation) will call the virtual ExecuteImportPipeline */
+	/** The default implementation, which is called if the Blueprint does not have any implementation, calls the virtual ShowPipelineConfigurationDialog(). */
 	EInterchangePipelineConfigurationDialogResult ScriptedShowPipelineConfigurationDialog_Implementation(TArray<FInterchangeStackInfo>& PipelineStacks
 		, TArray<UInterchangePipelineBase*>& OutPipelines
 		, UInterchangeSourceData* SourceData
@@ -60,7 +60,7 @@ public:
 	}
 
 	/**
-	 * Non virtual helper to allow blueprint to implement event base function to implement the ShowScenePipelineConfigurationDialog,
+	 * Non-virtual helper that allows Blueprint to implement an event-based function to implement ShowScenePipelineConfigurationDialog().
 	 */
 	UFUNCTION(BlueprintNativeEvent, BlueprintCallable, Category = "Interchange | Translator")
 	INTERCHANGEENGINE_API EInterchangePipelineConfigurationDialogResult ScriptedShowScenePipelineConfigurationDialog(TArray<FInterchangeStackInfo>& PipelineStacks
@@ -69,7 +69,7 @@ public:
 		, UInterchangeTranslatorBase* Translator
 		, UInterchangeBaseNodeContainer* BaseNodeContainer);
 
-	/** The default implementation (call if the blueprint do not have any implementation) will call the virtual ExecuteImportPipeline */
+	/** The default implementation, which is called if the Blueprint does not have any implementation, calls the virtual ShowScenePipelineConfigurationDialog(). */
 	EInterchangePipelineConfigurationDialogResult ScriptedShowScenePipelineConfigurationDialog_Implementation(TArray<FInterchangeStackInfo>& PipelineStacks
 		, TArray<UInterchangePipelineBase*>& OutPipelines
 		, UInterchangeSourceData* SourceData
@@ -81,7 +81,7 @@ public:
 	}
 
 	/**
-	 * Non virtual helper to allow blueprint to implement event base function to implement the ShowReimportPipelineConfigurationDialog,
+	 * Non-virtual helper that allows Blueprint to implement an event-based function to implement ShowReimportPipelineConfigurationDialog().
 	 */
 	UFUNCTION(BlueprintNativeEvent, BlueprintCallable, Category = "Interchange | Translator")
 	INTERCHANGEENGINE_API EInterchangePipelineConfigurationDialogResult ScriptedShowReimportPipelineConfigurationDialog(TArray<FInterchangeStackInfo>& PipelineStacks
@@ -91,7 +91,7 @@ public:
 		, UInterchangeBaseNodeContainer* BaseNodeContainer
 		, UObject* ReimportAsset);
 
-	/** The default implementation (call if the blueprint do not have any implementation) will call the virtual ExecuteImportPipeline */
+	/** The default implementation, which is called if the Blueprint does not have any implementation, calls the virtual ShowReimportPipelineConfigurationDialog(). */
 	EInterchangePipelineConfigurationDialogResult ScriptedShowReimportPipelineConfigurationDialog_Implementation(TArray<FInterchangeStackInfo>& PipelineStacks
 		, TArray<UInterchangePipelineBase*>& OutPipelines
 		, UInterchangeSourceData* SourceData
@@ -106,7 +106,7 @@ public:
 protected:
 
 	/**
-	 * This function show a dialog use to configure pipeline stacks and return a stack name that tell the caller the user choice.
+	 * This function shows a dialog used to configure pipeline stacks and returns a stack name that tells the caller the user's choice.
 	 */
 	virtual EInterchangePipelineConfigurationDialogResult ShowPipelineConfigurationDialog(TArray<FInterchangeStackInfo>& PipelineStacks
 		, TArray<UInterchangePipelineBase*>& OutPipelines
@@ -119,7 +119,7 @@ protected:
 	}
 
 	/**
-	 * This function show a dialog use to configure pipeline stacks and return a stack name that tell the caller the user choice.
+	 * This function shows a dialog used to configure pipeline stacks and returns a stack name that tells the caller the user's choice.
 	 */
 	virtual EInterchangePipelineConfigurationDialogResult ShowScenePipelineConfigurationDialog(TArray<FInterchangeStackInfo>& PipelineStacks
 		, TArray<UInterchangePipelineBase*>& OutPipelines
@@ -132,7 +132,7 @@ protected:
 	}
 
 	/**
-	 * This function show a dialog use to configure pipeline stacks and return a stack name that tell the caller the user choice.
+	 * This function shows a dialog used to configure pipeline stacks and returns a stack name that tells the caller the user's choice.
 	 */
 	virtual EInterchangePipelineConfigurationDialogResult ShowReimportPipelineConfigurationDialog(TArray<FInterchangeStackInfo>& PipelineStacks
 		, TArray<UInterchangePipelineBase*>& OutPipelines
diff --git a/Engine/Source/Runtime/Interchange/Engine/Public/InterchangeProjectSettings.h b/Engine/Source/Runtime/Interchange/Engine/Public/InterchangeProjectSettings.h
--- a/Engine/Source/Runtime/Interchange/Engine/Public/InterchangeProjectSettings.h
+++ b/Engine/Source/Runtime/Interchange/Engine/Public/InterchangeProjectSettings.h
@@ -30,10 +30,11 @@ struct FInterchangePipelineStack
 {
 	GENERATED_BODY()
 	
+	/** The list of pipelines in this stack. The pipelines are executed in fixed order, from top to bottom. */
 	UPROPERTY(EditAnywhere, Category = "Pipelines", meta = (AllowedClasses = "/Script/InterchangeCore.InterchangePipelineBase, /Script/InterchangeEngine.InterchangeBlueprintPipelineBase, /Script/InterchangeEngine.InterchangePythonPipelineAsset"))
 	TArray<FSoftObjectPath> Pipelines;
 
-	/** This tells Interchange which pipeline to add based on the type of the source */
+	/** Specifies a different list of pipelines for this stack to use when importing data from specific translators. */
 	UPROPERTY(EditAnywhere, Category = "TranslatorPipelines")
 	TArray<FInterchangeTranslatorPipelines> PerTranslatorPipelines;
 };
@@ -43,15 +44,15 @@ struct FInterchangeImportSettings
 {
 	GENERATED_BODY()
 
-	/** All the available pipeline stacks you want to use to import with interchange. The chosen pipeline stack execute all the pipelines from top to bottom order. You can order them by using the grip on the left of any pipelines.*/
+	/** Configures the pipeline stacks that are available when importing assets with Interchange. */
 	UPROPERTY(EditAnywhere, Category = "Pipeline")
 	TMap<FName, FInterchangePipelineStack> PipelineStacks;
 
-	/** This tell interchange which pipeline stack to select when importing.*/
+	/** Specifies which pipeline stack Interchange should use by default. */
 	UPROPERTY(EditAnywhere, Category = "Pipeline")
 	FName DefaultPipelineStack = NAME_None;
 
-	/** This tell interchange which class to instance to show the import dialog.*/
+	/** Specifies the class that should be used to define the configuration dialog that Interchange shows on import. */
 	UPROPERTY(EditAnywhere, Category = "Pipeline")
 	TSoftClassPtr <UInterchangePipelineConfigurationBase> ImportDialogClass;
 
@@ -89,7 +90,7 @@ struct FInterchangeContentImportSettings : public FInterchangeImportSettings
 {
 	GENERATED_BODY()
 
-	/** This tell interchange which pipeline stack to select when importing.*/
+	/** Specifies a different pipeline stack for Interchange to use by default when importing specific types of assets. */
 	UPROPERTY(EditAnywhere, Category = "Pipeline", Meta=(DisplayAfter="DefaultPipelineStack"))
 	TMap<EInterchangeTranslatorAssetType, FName> DefaultPipelineStackOverride;
 
@@ -105,7 +106,7 @@ class UInterchangeProjectSettings : public UDeveloperSettings
 
 public:
 	/**
-	 * Settings used when importing into the content browser.
+	 * Settings used when importing into the Content Browser.
 	 */
 	UPROPERTY(EditAnywhere, config, Category = "ImportContent")
 	FInterchangeContentImportSettings ContentImportSettings;
@@ -116,21 +117,20 @@ public:
 	UPROPERTY(EditAnywhere, config, Category = "ImportIntoLevel")
 	FInterchangeImportSettings SceneImportSettings;
 
-	/** This tells interchange which file picker class to construct when we need to choose a file for a source.*/
+	/** This tells Interchange which file picker class to construct when we need to choose a file for a source. */
 	UPROPERTY(EditAnywhere, config, Category = "EditorInterface")
 	TSoftClassPtr <UInterchangeFilePickerBase> FilePickerClass;
 
 	/**
-	 * If checked, interchange translators and legacy importer will default static mesh geometry to smooth edge when the smoothing information is missing.
-	 * This option exist to allows old project to import the same way as before if their workflows need static mesh edges to be hard when the smoothing
-	 * info is missing.
+	 * If enabled, both Interchange translators and the legacy import process smooth the edges of static meshes that don't contain smoothing information.
+	 * If you have an older project that relies on leaving hard edges by default, you can disable this setting to preserve consistency with older assets.
 	 */
 	UPROPERTY(EditAnywhere, config, Category = "Generic|ImportSettings")
 	bool bStaticMeshUseSmoothEdgesIfSmoothingInformationIsMissing = true;
 
 	/**
-	 * This tells interchange which is the pipeline class to use when editor tools want to import or reimport tools with bake settings.
-	 * UnrealEd code depend on this class to be set and this property is only editable in the ini file directly.
+	 * Specifies which pipeline class Interchange should use when editor tools import or reimport an asset with base settings.
+	 * Unreal Editor depends on this class to be set. You can only edit this property in the .ini file.
 	 */
 	UPROPERTY(EditAnywhere, config, Category = "Editor Generic Pipeline Class")
 	TSoftClassPtr <UInterchangePipelineBase> GenericPipelineClass;
diff --git a/Engine/Source/Runtime/Interchange/Engine/Public/InterchangePythonPipelineBase.h b/Engine/Source/Runtime/Interchange/Engine/Public/InterchangePythonPipelineBase.h
--- a/Engine/Source/Runtime/Interchange/Engine/Public/InterchangePythonPipelineBase.h
+++ b/Engine/Source/Runtime/Interchange/Engine/Public/InterchangePythonPipelineBase.h
@@ -10,7 +10,7 @@
 #include "InterchangePythonPipelineBase.generated.h"
 
 /*
- * This class represent a python pipeline. It is use by the TSoftClassPtr of the PythonPipeline asset.
+ * This class represents a Python pipeline. It is used by the TSoftClassPtr of the PythonPipeline asset.
  *
  */
 UCLASS(BlueprintType, Abstract, MinimalAPI)
@@ -35,12 +35,12 @@ struct FPropertyData
 };
 
 /*
- * This class is a helper class for python pipeline. It allow picker to filter correctly the content browser instance we want to support
+ * This class is a helper class for Python pipelines. It allows the class picker to filter the Content Browser instance we want to support
  * in the pipeline stack.
  *
- * Use this class factory to be able to create an instance of a python pipeline in the content browser.
- * A python pipeline do not have any locked properties and all default value must be set in the python script.
- * This restriction exist because python class are transient, so any assets create from a python class cannot be save.
+ * Use this class factory to be able to create an instance of a Python pipeline in the Content Browser.
+ * A Python pipeline does not have any locked properties, and all default values must be set in the Python script.
+ * This restriction exists because Python classes are transient, so any assets created from a Python class cannot be saved.
  * 
  */
 UCLASS(BlueprintType, MinimalAPI)
@@ -58,11 +58,11 @@ public:
 
 	INTERCHANGEENGINE_API void SetupFromPipeline(const UInterchangePythonPipelineBase* PythonPipeline, const bool bRegeneratePipeline = true);
 
-	/** The python class we want to use has a pipeline */
+	/** The Python class we want to use as a pipeline. */
 	UPROPERTY(EditAnywhere, Category = "Interchange|Python")
 	TSoftClassPtr<UInterchangePythonPipelineBase> PythonClass;
 
-	/** The transient pipeline we generate when we load the python pipeline */
+	/** The transient pipeline we generate when we load the Python pipeline. */
 	UPROPERTY(VisibleAnywhere, Transient, Category = "Interchange|Python")
 	TObjectPtr<UInterchangePythonPipelineBase> GeneratedPipeline;
 
