Fixed issue when reimporting twice a file after an object has an edited attribute

#rnx

diff --git a/Engine/Source/Runtime/Interchange/Core/Private/Nodes/InterchangeBaseNode.cpp b/Engine/Source/Runtime/Interchange/Core/Private/Nodes/InterchangeBaseNode.cpp
--- a/Engine/Source/Runtime/Interchange/Core/Private/Nodes/InterchangeBaseNode.cpp
+++ b/Engine/Source/Runtime/Interchange/Core/Private/Nodes/InterchangeBaseNode.cpp
@@ -502,6 +502,16 @@ void UInterchangeBaseNode::CopyStorageAttributes(const UInterchangeBaseNode* Sou
 	UE::Interchange::FAttributeStorage::CopyStorageAttributes(*(SourceNode->Attributes), *(DestinationNode->Attributes), AttributeKeys);
 }
 
+void UInterchangeBaseNode::CopyStorageAttributes(const UInterchangeBaseNode* SourceNode, UE::Interchange::FAttributeStorage& DestinationStorage, TArray<UE::Interchange::FAttributeKey>& AttributeKeys)
+{
+	UE::Interchange::FAttributeStorage::CopyStorageAttributes(*(SourceNode->Attributes), DestinationStorage, AttributeKeys);
+}
+
+void UInterchangeBaseNode::CopyStorageAttributes(const UE::Interchange::FAttributeStorage& SourceStorage, UInterchangeBaseNode* DestinationNode, TArray<UE::Interchange::FAttributeKey>& AttributeKeys)
+{
+	UE::Interchange::FAttributeStorage::CopyStorageAttributes(SourceStorage, *(DestinationNode->Attributes), AttributeKeys);
+}
+
 void UInterchangeBaseNode::CopyStorage(const UInterchangeBaseNode* SourceNode, UInterchangeBaseNode* DestinationNode)
 {
 	*(DestinationNode->Attributes) = *(SourceNode->Attributes);
diff --git a/Engine/Source/Runtime/Interchange/Core/Public/Nodes/InterchangeBaseNode.h b/Engine/Source/Runtime/Interchange/Core/Public/Nodes/InterchangeBaseNode.h
--- a/Engine/Source/Runtime/Interchange/Core/Public/Nodes/InterchangeBaseNode.h
+++ b/Engine/Source/Runtime/Interchange/Core/Public/Nodes/InterchangeBaseNode.h
@@ -525,6 +525,8 @@ public:
 	static INTERCHANGECORE_API void CompareNodeStorage(const UInterchangeBaseNode* NodeA, const UInterchangeBaseNode* NodeB, TArray<UE::Interchange::FAttributeKey>& RemovedAttributes, TArray<UE::Interchange::FAttributeKey>& AddedAttributes, TArray<UE::Interchange::FAttributeKey>& ModifiedAttributes);
 	
 	static INTERCHANGECORE_API void CopyStorageAttributes(const UInterchangeBaseNode* SourceNode, UInterchangeBaseNode* DestinationNode, TArray<UE::Interchange::FAttributeKey>& AttributeKeys);
+	static INTERCHANGECORE_API void CopyStorageAttributes(const UInterchangeBaseNode* SourceNode, UE::Interchange::FAttributeStorage& DestinationStorage, TArray<UE::Interchange::FAttributeKey>& AttributeKeys);
+	static INTERCHANGECORE_API void CopyStorageAttributes(const UE::Interchange::FAttributeStorage& SourceStorage, UInterchangeBaseNode* DestinationNode, TArray<UE::Interchange::FAttributeKey>& AttributeKeys);
 
 	static INTERCHANGECORE_API void CopyStorage(const UInterchangeBaseNode* SourceNode, UInterchangeBaseNode* DestinationNode);
 	
diff --git a/Engine/Source/Runtime/Interchange/Engine/Private/InterchangeImportCommon.cpp b/Engine/Source/Runtime/Interchange/Engine/Private/InterchangeImportCommon.cpp
--- a/Engine/Source/Runtime/Interchange/Engine/Private/InterchangeImportCommon.cpp
+++ b/Engine/Source/Runtime/Interchange/Engine/Private/InterchangeImportCommon.cpp
@@ -12,6 +12,8 @@
 #include "InterchangeTranslatorBase.h"
 #include "Nodes/InterchangeBaseNodeContainer.h"
 #include "Nodes/InterchangeFactoryBaseNode.h"
+#include "Types/AttributeStorage.h"
+
 #include "UObject/Object.h"
 #include "EditorFramework/AssetImportData.h"
 #include "HAL/FileManager.h"
@@ -346,10 +348,19 @@ namespace UE::Interchange
 				TArray<FAttributeKey> ModifiedAttributes;
 				UInterchangeBaseNode::CompareNodeStorage(PreviousAssetNode, CurrentAssetNode, RemovedAttributes, AddedAttributes, ModifiedAttributes);
 
-				//set all ModifedAttributes from the CurrentAssetNode to the pipeline node. This will put back all user changes
+				//Cache all modified attributes from the pipeline node
+				UE::Interchange::FAttributeStorage CachedAttributes;
+				UInterchangeBaseNode::CopyStorageAttributes(PipelineAssetNode, CachedAttributes, ModifiedAttributes);
+
+				//Set all ModifedAttributes from the CurrentAssetNode to PipelineAssetNode
+				//This way, call to ApplyAllCustomAttributeToObject will preserve modified attributes from CurrentAssetNode
 				UInterchangeBaseNode::CopyStorageAttributes(CurrentAssetNode, PipelineAssetNode, ModifiedAttributes);
-				//Now apply the pipeline node attribute to the asset
+
+				//Apply the pipeline node's attributes to the asset
 				PipelineAssetNode->ApplyAllCustomAttributeToObject(Asset);
+
+				//Restore all modified attributes back to the pipeline node
+				UInterchangeBaseNode::CopyStorageAttributes(CachedAttributes, PipelineAssetNode, ModifiedAttributes);
 				break;
 			}
 		}
