[Interchange] Skeleton creation added for SkeletalMesh re-import without a Skeleton present.

#rnx

diff --git a/Engine/Plugins/Interchange/Runtime/Source/Import/Private/Mesh/InterchangeSkeletonFactory.cpp b/Engine/Plugins/Interchange/Runtime/Source/Import/Private/Mesh/InterchangeSkeletonFactory.cpp
--- a/Engine/Plugins/Interchange/Runtime/Source/Import/Private/Mesh/InterchangeSkeletonFactory.cpp
+++ b/Engine/Plugins/Interchange/Runtime/Source/Import/Private/Mesh/InterchangeSkeletonFactory.cpp
@@ -10,6 +10,10 @@
 #include "Nodes/InterchangeBaseNodeContainer.h"
 #include "Rendering/SkeletalMeshLODImporterData.h"
 
+#if WITH_EDITORONLY_DATA
+#include "Engine/SkeletalMesh.h"
+#endif
+
 #include UE_INLINE_GENERATED_CPP_BY_NAME(InterchangeSkeletonFactory)
 
 UClass* UInterchangeSkeletonFactory::GetFactoryClass() const
@@ -99,3 +103,21 @@ UInterchangeFactoryBase::FImportAssetResult UInterchangeSkeletonFactory::BeginIm
 
 	return ImportAssetResult;
 }
+
+UObject* UInterchangeSkeletonFactory::GetObjectToReimport(UObject* ReimportObject, const UInterchangeFactoryBaseNode& FactoryNode, const FString& PackageName, const FString& AssetName, const FString& SubPathString)
+{
+#if WITH_EDITORONLY_DATA
+	if (ReimportObject &&
+		ReimportObject->GetClass()->IsChildOf(USkeletalMesh::StaticClass()) &&
+		FactoryNode.GetObjectClass()->IsChildOf(USkeleton::StaticClass()))
+	{
+		//The Skeleton property of the UInterchangeGenericCommonSkeletalMeshesAndAnimationsProperties will be populated for re-import,
+		//which in turn will disable the UInterchangeSkeletonFactoryNodes.
+		//Meaning, if at re-import the Skeleton already exists, the SkeletonFactoryNode will be disabled.
+		//Hence, if we get to this branch execution that means that we are Reimporting a SkeletalMesh without a Skeleton present.
+		return nullptr;
+	}
+#endif
+
+	return ReimportObject;
+}
\ No newline at end of file
diff --git a/Engine/Plugins/Interchange/Runtime/Source/Import/Public/Mesh/InterchangeSkeletonFactory.h b/Engine/Plugins/Interchange/Runtime/Source/Import/Public/Mesh/InterchangeSkeletonFactory.h
--- a/Engine/Plugins/Interchange/Runtime/Source/Import/Public/Mesh/InterchangeSkeletonFactory.h
+++ b/Engine/Plugins/Interchange/Runtime/Source/Import/Public/Mesh/InterchangeSkeletonFactory.h
@@ -22,6 +22,8 @@ public:
 	virtual EInterchangeFactoryAssetType GetFactoryAssetType() override { return EInterchangeFactoryAssetType::Meshes; }
 	virtual FImportAssetResult BeginImportAsset_GameThread(const FImportAssetObjectParams& Arguments) override;
 
+	virtual UObject* GetObjectToReimport(UObject* ReimportObject, const UInterchangeFactoryBaseNode& FactoryNode, const FString& PackageName, const FString& AssetName, const FString& SubPathString) override;
+
 	// Interchange factory base interface end
 	//////////////////////////////////////////////////////////////////////////
 };
diff --git a/Engine/Source/Runtime/Interchange/Core/Public/InterchangeFactoryBase.h b/Engine/Source/Runtime/Interchange/Core/Public/InterchangeFactoryBase.h
--- a/Engine/Source/Runtime/Interchange/Core/Public/InterchangeFactoryBase.h
+++ b/Engine/Source/Runtime/Interchange/Core/Public/InterchangeFactoryBase.h
@@ -327,4 +327,15 @@ public:
 
 	UPROPERTY()
 	TObjectPtr<UInterchangeResultsContainer> Results;
+
+	/**
+	* Acquires the Object to be re-imported.
+	* Function is called from FFactoryCommon::GetObjectToReiomport.
+	* Which in turn is called from the InterchangeTask related functions to acquire the Object to be Reimported.
+	* Allows Factories to override, in order to provide the desired ReimportObject based on Factory requirements.
+	*/
+	virtual UObject* GetObjectToReimport(UObject* ReimportObject, const UInterchangeFactoryBaseNode& FactoryNode, const FString& PackageName, const FString& AssetName, const FString& SubPathString)
+	{
+		return ReimportObject;
+	}
 };
diff --git a/Engine/Source/Runtime/Interchange/Engine/Private/InterchangeImportCommon.cpp b/Engine/Source/Runtime/Interchange/Engine/Private/InterchangeImportCommon.cpp
--- a/Engine/Source/Runtime/Interchange/Engine/Private/InterchangeImportCommon.cpp
+++ b/Engine/Source/Runtime/Interchange/Engine/Private/InterchangeImportCommon.cpp
@@ -11,6 +11,7 @@
 #include "InterchangeSourceData.h"
 #include "InterchangeTranslatorBase.h"
 #include "Nodes/InterchangeBaseNodeContainer.h"
+#include "InterchangeFactoryBase.h"
 #include "Nodes/InterchangeFactoryBaseNode.h"
 #include "Types/AttributeStorage.h"
 
@@ -366,8 +367,13 @@ namespace UE::Interchange
 		}
 	}
 
-	UObject* FFactoryCommon::GetObjectToReimport(UObject* ReimportObject, const UInterchangeFactoryBaseNode& FactoryNode, const FString& PackageName, const FString& AssetName, const FString& SubPathString)
+	UObject* FFactoryCommon::GetObjectToReimport(UInterchangeFactoryBase* Factory, UObject* ReimportObject, const UInterchangeFactoryBaseNode& FactoryNode, const FString& PackageName, const FString& AssetName, const FString& SubPathString)
 	{
+		UObject* ReimportObjectCandidate = Factory->GetObjectToReimport(ReimportObject, FactoryNode, PackageName, AssetName, SubPathString);
+		if (ReimportObjectCandidate != ReimportObject)
+		{
+			return ReimportObjectCandidate;
+		}
 #if WITH_EDITORONLY_DATA
 		if (ReimportObject && !ReimportObject->GetClass()->IsChildOf(FactoryNode.GetObjectClass()))
 		{
diff --git a/Engine/Source/Runtime/Interchange/Engine/Private/Tasks/InterchangeTaskCreateSceneObjects.cpp b/Engine/Source/Runtime/Interchange/Engine/Private/Tasks/InterchangeTaskCreateSceneObjects.cpp
--- a/Engine/Source/Runtime/Interchange/Engine/Private/Tasks/InterchangeTaskCreateSceneObjects.cpp
+++ b/Engine/Source/Runtime/Interchange/Engine/Private/Tasks/InterchangeTaskCreateSceneObjects.cpp
@@ -81,7 +81,7 @@ void UE::Interchange::FTaskCreateSceneObjects::DoTask(ENamedThreads::Type Curren
 		CreateSceneObjectsParams.ObjectName = SceneNodeName;
 		CreateSceneObjectsParams.FactoryNode = FactoryNode;
 		CreateSceneObjectsParams.Level = ImportLevel;
-		CreateSceneObjectsParams.ReimportObject = FFactoryCommon::GetObjectToReimport(ReimportObject, *FactoryNode, WorldPath, WorldName, NodePrefix + SceneNodeName);
+		CreateSceneObjectsParams.ReimportObject = FFactoryCommon::GetObjectToReimport(Factory, ReimportObject, *FactoryNode, WorldPath, WorldName, NodePrefix + SceneNodeName);
 		CreateSceneObjectsParams.ReimportFactoryNode = FFactoryCommon::GetFactoryNode(ReimportObject, WorldPath, WorldName, NodePrefix + SceneNodeName);
 
 		if (AsyncHelper->BaseNodeContainers.IsValidIndex(SourceIndex))
diff --git a/Engine/Source/Runtime/Interchange/Engine/Private/Tasks/InterchangeTaskImportObject.cpp b/Engine/Source/Runtime/Interchange/Engine/Private/Tasks/InterchangeTaskImportObject.cpp
--- a/Engine/Source/Runtime/Interchange/Engine/Private/Tasks/InterchangeTaskImportObject.cpp
+++ b/Engine/Source/Runtime/Interchange/Engine/Private/Tasks/InterchangeTaskImportObject.cpp
@@ -160,7 +160,7 @@ namespace UE::Interchange::Private
 		FString AssetName;
 		Private::InternalGetPackageName(*AsyncHelper, SourceIndex, PackageBasePath, FactoryNode, PackageName, AssetName);
 		bool bSkipAsset = false;
-		UObject* ObjectToReimport = UE::Interchange::FFactoryCommon::GetObjectToReimport(AsyncHelper->TaskData.ReimportObject, *FactoryNode, PackageName, AssetName);
+		UObject* ObjectToReimport = UE::Interchange::FFactoryCommon::GetObjectToReimport(Factory, AsyncHelper->TaskData.ReimportObject, *FactoryNode, PackageName, AssetName);
 		if (ObjectToReimport)
 		{
 			UInterchangeBaseNodeContainer* NodeContainer = nullptr;
@@ -289,7 +289,7 @@ void UE::Interchange::FTaskImportObject_GameThread::DoTask(ENamedThreads::Type C
 	FString AssetName;
 	Private::InternalGetPackageName(*AsyncHelper, SourceIndex, PackageBasePath, FactoryNode, PackageName, AssetName);
 
-	UObject* ObjectToReimport = FFactoryCommon::GetObjectToReimport(AsyncHelper->TaskData.ReimportObject, *FactoryNode, PackageName, AssetName);
+	UObject* ObjectToReimport = FFactoryCommon::GetObjectToReimport(Factory, AsyncHelper->TaskData.ReimportObject, *FactoryNode, PackageName, AssetName);
 	if (!ensure(!IsGarbageCollecting()))
 	{
 		//Skip this asset
@@ -647,7 +647,7 @@ void UE::Interchange::FTaskImportObjectFinalize_GameThread::DoTask(ENamedThreads
 
 				const FString PackagePathName = ImportAssetResult.ImportedObject->GetPathName();
 				const FString AssetName = ImportAssetResult.ImportedObject->GetName();
-				const UObject* ObjectToReimport = FFactoryCommon::GetObjectToReimport(AsyncHelper->TaskData.ReimportObject, *FactoryNode, PackagePathName, AssetName);
+				const UObject* ObjectToReimport = FFactoryCommon::GetObjectToReimport(Factory, AsyncHelper->TaskData.ReimportObject, *FactoryNode, PackagePathName, AssetName);
 				AssetInfo.bIsReimport = bool(ObjectToReimport != nullptr);
 			}
 
diff --git a/Engine/Source/Runtime/Interchange/Engine/Public/InterchangeImportCommon.h b/Engine/Source/Runtime/Interchange/Engine/Public/InterchangeImportCommon.h
--- a/Engine/Source/Runtime/Interchange/Engine/Public/InterchangeImportCommon.h
+++ b/Engine/Source/Runtime/Interchange/Engine/Public/InterchangeImportCommon.h
@@ -11,6 +11,7 @@ class UAssetImportData;
 class UInterchangeAssetImportData;
 class UInterchangeBaseNode;
 class UInterchangeBaseNodeContainer;
+class UInterchangeFactoryBase;
 class UInterchangeFactoryBaseNode;
 class UInterchangePipelineBase;
 class UInterchangeSourceData;
@@ -119,7 +120,7 @@ namespace UE::Interchange
 		 * @param AssetName: Asset name of the actual object to reimport
 		 * @param SubPathString: Optional subobject name
 		 */
-		static 	INTERCHANGEENGINE_API UObject* GetObjectToReimport(UObject* ReimportObject, const UInterchangeFactoryBaseNode& FactoryNode, const FString& PackageName, const FString& AssetName, const FString& SubPathString = FString());
+		static 	INTERCHANGEENGINE_API UObject* GetObjectToReimport(UInterchangeFactoryBase* Factory, UObject* ReimportObject, const UInterchangeFactoryBaseNode& FactoryNode, const FString& PackageName, const FString& AssetName, const FString& SubPathString = FString());
 
 		/**
 		 * If the ReimportObject is a UInterchangeSceneImportAsset, returns the factory node
