Interchange: Verify the interchange manager pointer in the GC callback

#rnx

diff --git a/Engine/Source/Runtime/Interchange/Engine/Private/InterchangeManager.cpp b/Engine/Source/Runtime/Interchange/Engine/Private/InterchangeManager.cpp
--- a/Engine/Source/Runtime/Interchange/Engine/Private/InterchangeManager.cpp
+++ b/Engine/Source/Runtime/Interchange/Engine/Private/InterchangeManager.cpp
@@ -935,7 +935,7 @@ UInterchangeManager& UInterchangeManager::GetInterchangeManager()
 
 		InterchangeManager->GCEndDelegate = FCoreUObjectDelegates::GetPostGarbageCollect().AddLambda([]()
 			{
-				if (IsInterchangeImportEnabled())
+				if (IsInterchangeImportEnabled() && InterchangeManager.IsValid())
 				{
 					InterchangeManager->StartQueuedTasks(InterchangeManager->bGCEndDelegateCancellAllTask);
 				}
