Interchange: Fix the import dialog asset picker by using the transient package directly and not a sub transient package.

#rnx

diff --git a/Engine/Plugins/Interchange/Editor/Source/Pipelines/Private/SInterchangePipelineConfigurationDialog.cpp b/Engine/Plugins/Interchange/Editor/Source/Pipelines/Private/SInterchangePipelineConfigurationDialog.cpp
--- a/Engine/Plugins/Interchange/Editor/Source/Pipelines/Private/SInterchangePipelineConfigurationDialog.cpp
+++ b/Engine/Plugins/Interchange/Editor/Source/Pipelines/Private/SInterchangePipelineConfigurationDialog.cpp
@@ -42,7 +42,6 @@ static FAutoConsoleVariableRef CCvarInterchangeEnableFBXImport(
 	ECVF_Default);
 
 const FName ReimportStackName = TEXT("ReimportPipeline");
-const FString ReimportPipelinePrefix = TEXT("reimport_");
 
 void SInterchangePipelineItem::Construct(
 	const FArguments& InArgs,
@@ -161,24 +160,12 @@ SInterchangePipelineConfigurationDialog::~SInterchangePipelineConfigurationDialo
  // As this is the name displayed in the Dialog, conflicts won't matter.
 FString SInterchangePipelineConfigurationDialog::GetPipelineDisplayName(const UInterchangePipelineBase* Pipeline)
 {
-	static int32 RightChopIndex = ReimportPipelinePrefix.Len();
 
 	FString PipelineDisplayName = Pipeline->ScriptedGetPipelineDisplayName();
 	if(PipelineDisplayName.IsEmpty())
 	{
 		PipelineDisplayName = Pipeline->GetName();
 	}
-	if (PipelineDisplayName.StartsWith(ReimportPipelinePrefix))
-	{
-		PipelineDisplayName = PipelineDisplayName.RightChop(RightChopIndex);
-	}
-
-	FString StackName;
-	FString DisplayName;
-	if (PipelineDisplayName.Split("_", &StackName, &DisplayName))
-	{
-		return DisplayName;
-	}
 
 	return PipelineDisplayName;
 }
@@ -316,7 +303,7 @@ TSharedRef<SBox> SInterchangePipelineConfigurationDialog::SpawnPipelineConfigura
 			for (const TObjectPtr<UInterchangePipelineBase>& DefaultPipeline : Stack.Pipelines)
 			{
 				check(DefaultPipeline);
-				if (UInterchangePipelineBase* GeneratedPipeline = UE::Interchange::GeneratePipelineInstanceInSourceAssetPackage(DefaultPipeline))
+				if (UInterchangePipelineBase* GeneratedPipeline = UE::Interchange::GeneratePipelineInstance(DefaultPipeline))
 				{
 					GeneratedPipeline->TransferAdjustSettings(DefaultPipeline);
 					if (Stack.StackName == ReimportStackName)
@@ -809,7 +796,7 @@ FReply SInterchangePipelineConfigurationDialog::OnResetToDefault()
 							TObjectPtr<UInterchangePipelineBase> PipelineElement = PipelineListViewItems[PipelineIndex]->Pipeline;
 							if (PipelineElement.Get() == Pipeline)
 							{
-								if (UInterchangePipelineBase* GeneratedPipeline = UE::Interchange::GeneratePipelineInstanceInSourceAssetPackage(DefaultPipeline))
+								if (UInterchangePipelineBase* GeneratedPipeline = UE::Interchange::GeneratePipelineInstance(DefaultPipeline))
 								{
 									GeneratedPipeline->TransferAdjustSettings(DefaultPipeline);
 									GeneratedPipeline->SetBasicLayoutMode(bBasicLayout);
@@ -905,12 +892,6 @@ void SInterchangePipelineConfigurationDialog::ClosePipelineConfiguration(const E
 		//Fill the OutPipelines array
 		for (TSharedPtr<FInterchangePipelineItemType> PipelineElement : PipelineListViewItems)
 		{
-			if (!bReimport)
-			{
-				// Create a name that would not cause conflict when this asset maybe reimported.
-				FString NewPipelineName = ReimportPipelinePrefix + PipelineElement->DisplayName;
-				PipelineElement->Pipeline->Rename(*NewPipelineName);
-			}
 			OutPipelines->Add(PipelineElement->Pipeline);
 		}
 	}
@@ -984,7 +965,7 @@ void SInterchangePipelineConfigurationDialog::RefreshStack(bool bStackSelectionC
 		for (const TObjectPtr<UInterchangePipelineBase>& DefaultPipeline : Stack.Pipelines)
 		{
 			check(DefaultPipeline);
-			if (UInterchangePipelineBase* GeneratedPipeline = UE::Interchange::GeneratePipelineInstanceInSourceAssetPackage(DefaultPipeline))
+			if (UInterchangePipelineBase* GeneratedPipeline = UE::Interchange::GeneratePipelineInstance(DefaultPipeline))
 			{
 				GeneratedPipeline->TransferAdjustSettings(DefaultPipeline);
 				if (Stack.StackName != ReimportStackName || !bStackSelectionChange)
diff --git a/Engine/Source/Runtime/Interchange/Core/Private/InterchangePipelineBase.cpp b/Engine/Source/Runtime/Interchange/Core/Private/InterchangePipelineBase.cpp
--- a/Engine/Source/Runtime/Interchange/Core/Private/InterchangePipelineBase.cpp
+++ b/Engine/Source/Runtime/Interchange/Core/Private/InterchangePipelineBase.cpp
@@ -33,6 +33,8 @@ void UInterchangePipelineBase::LoadSettings(const FName PipelineStackName)
 void UInterchangePipelineBase::SaveSettings(const FName PipelineStackName)
 {
 	SaveSettingsInternal(PipelineStackName, GEditorPerProjectIni);
+	//Flush the config only once, we do not want to flush when recursively save the sub object pipeline.
+	GConfig->Flush(0);
 }
 
 void UInterchangePipelineBase::AdjustSettingsForContext(EInterchangePipelineContext ReimportType, TObjectPtr<UObject> ReimportAsset)
@@ -288,7 +290,6 @@ void UInterchangePipelineBase::SaveSettingsInternal(const FName PipelineStackNam
 			}
 		}
 	}
-	GConfig->Flush(0);
 }
 
 
diff --git a/Engine/Source/Runtime/Interchange/Engine/Private/InterchangeManager.cpp b/Engine/Source/Runtime/Interchange/Engine/Private/InterchangeManager.cpp
--- a/Engine/Source/Runtime/Interchange/Engine/Private/InterchangeManager.cpp
+++ b/Engine/Source/Runtime/Interchange/Engine/Private/InterchangeManager.cpp
@@ -612,24 +612,6 @@ void UE::Interchange::FImportAsyncHelper::CleanUp()
 		}
 	}
 	CreatedFactories.Empty();
-
-
-	for (TObjectIterator<UInterchangePipelineBase> PipelinesItr;PipelinesItr;++PipelinesItr)
-	{
-		const FString PackageName = PipelinesItr->GetPackage()->GetName();
-		if (PipelineInstancesPackageName == PackageName)
-		{
-			PipelinesItr->RemoveFromRoot();
-			PipelinesItr->MarkAsGarbage();
-		}
-	}
-
-	if (UPackage* PipelineInstancesPackage = FindPackage(nullptr, *PipelineInstancesPackageName))
-	{
-		PipelineInstancesPackage->ClearFlags(RF_Public| RF_Standalone| RF_Transactional);
-		PipelineInstancesPackage->SetFlags(RF_Transient);
-		PipelineInstancesPackage->MarkAsGarbage();
-	}
 }
 
 UE::Interchange::FImportResult::FImportResult()
@@ -807,18 +789,12 @@ void UE::Interchange::SanitizeObjectName(FString& ObjectName)
 	}
 }
 
-UInterchangePipelineBase* UE::Interchange::GeneratePipelineInstanceInSourceAssetPackage(const FSoftObjectPath& SourcePipeline)
+UInterchangePipelineBase* UE::Interchange::GeneratePipelineInstance(const FSoftObjectPath& PipelineInstance)
 {
-	FString PackageName = FPackageUtils::ExtractPackageName(SourcePipeline.ToString());
-	UPackage* TargetPackage = FindPackage(nullptr, *PackageName);
-	return GeneratePipelineInstance(SourcePipeline, TargetPackage);
-}
-
-UInterchangePipelineBase* UE::Interchange::GeneratePipelineInstance(const FSoftObjectPath& PipelineInstance, UPackage* PipelineInstancePackage /*= nullptr*/)
-{
-	if (!PipelineInstancePackage)
+	UPackage* PipelineInstancePackage = GetTransientPackage();
+	if (!ensure(PipelineInstancePackage))
 	{
-		PipelineInstancePackage = GetTransientPackage();
+		return nullptr;
 	}
 	UObject* ReferenceInstance = PipelineInstance.TryLoad();
 	if (!ReferenceInstance)
@@ -1104,7 +1080,7 @@ bool UInterchangeManager::ConvertImportData(const UObject* SourceImportData, FIm
 				const FInterchangePipelineStack& PipelineStack = InterchangeImportSettings.PipelineStacks.FindChecked(InterchangeImportSettings.DefaultPipelineStack);
 				for (const FSoftObjectPath& PipelinePath : PipelineStack.Pipelines)
 				{
-					if (UInterchangePipelineBase* GeneratedPipeline = UE::Interchange::GeneratePipelineInstance(PipelinePath, GetTransientPackage()))
+					if (UInterchangePipelineBase* GeneratedPipeline = UE::Interchange::GeneratePipelineInstance(PipelinePath))
 					{
 						GeneratedPipeline->AdjustSettingsForContext(EInterchangePipelineContext::AssetImport, nullptr);
 						if (GeneratedPipeline->IsA(ConvertedPipelineClass))
@@ -1702,17 +1678,6 @@ UInterchangeManager::ImportInternal(const FString& ContentPath, const UInterchan
 		Pipeline->DestinationName = TaskData.DestinationName;
 	};
 
-	// Use counter to guarantee uniqueness of packages on each call to ImportInternal
-	static uint64 ImportCounter = 0;
-	static FString TransientPackageBasePath = GetTransientPackage()->GetPathName();
-
-	const FString PackageName = FString::Printf(TEXT("InterchangePipelinePackage-%llu"), ++ImportCounter);
-	AsyncHelper->PipelineInstancesPackageName = TransientPackageBasePath / PackageName;
-
-	UPackage* PipelineInstancesPackage = CreatePackage(*AsyncHelper->PipelineInstancesPackageName);
-	PipelineInstancesPackage->ClearFlags(RF_Public | RF_Standalone);
-	PipelineInstancesPackage->SetPackageFlags(PKG_NewlyCreated);
-
 	const bool bSkipImportDialog = AsyncTranslator ? ImportAllWithSamePipelines.Contains(AsyncTranslator->GetClass()) : false;
 	if (bSkipImportDialog)
 	{
@@ -1790,7 +1755,7 @@ UInterchangeManager::ImportInternal(const FString& ContentPath, const UInterchan
 				if (SourcePipeline && SourcePipeline->SupportReimport()) //Its possible a pipeline doesnt exist anymore so it wont load into memory when we loading the outer asset
 				{
 					//Duplicate the pipeline saved in the asset import data
-					UInterchangePipelineBase* GeneratedPipeline = Cast<UInterchangePipelineBase>(StaticDuplicateObject(SourcePipeline, PipelineInstancesPackage));
+					UInterchangePipelineBase* GeneratedPipeline = Cast<UInterchangePipelineBase>(StaticDuplicateObject(SourcePipeline, GetTransientPackage()));
 					// Make sure that the instance does not carry over standalone and public flags as they are not actual assets to be persisted
 					GeneratedPipeline->ClearFlags(EObjectFlags::RF_Standalone | EObjectFlags::RF_Public);
 					AdjustPipelineSettingForContext(GeneratedPipeline);
@@ -1829,14 +1794,8 @@ UInterchangeManager::ImportInternal(const FString& ContentPath, const UInterchan
 
 				for (int32 PipelineIndex = 0; PipelineIndex < Pipelines->Num(); ++PipelineIndex)
 				{
-					if (UInterchangePipelineBase* GeneratedPipeline = UE::Interchange::GeneratePipelineInstance((*Pipelines)[PipelineIndex], PipelineInstancesPackage))
+					if (UInterchangePipelineBase* GeneratedPipeline = UE::Interchange::GeneratePipelineInstance((*Pipelines)[PipelineIndex]))
 					{
-						if (bShowPipelineStacksConfigurationDialog)
-						{
-							FString CurrentName = GeneratedPipeline->GetName();
-							FString NewName = FString::Printf(TEXT("%s_%s"), *StackName.ToString(), *CurrentName);
-							ensure(GeneratedPipeline->Rename(*NewName, nullptr, REN_DoNotDirty | REN_NonTransactional));
-						}
 						AdjustPipelineSettingForContext(GeneratedPipeline);
 						StackInfo.Pipelines.Add(GeneratedPipeline);
 					}
@@ -1962,7 +1921,7 @@ UInterchangeManager::ImportInternal(const FString& ContentPath, const UInterchan
 	{
 		for (int32 GraphPipelineIndex = 0; GraphPipelineIndex < ImportAssetParameters.OverridePipelines.Num(); ++GraphPipelineIndex)
 		{
-			UInterchangePipelineBase* GeneratedPipeline = UE::Interchange::GeneratePipelineInstance(ImportAssetParameters.OverridePipelines[GraphPipelineIndex], PipelineInstancesPackage);
+			UInterchangePipelineBase* GeneratedPipeline = UE::Interchange::GeneratePipelineInstance(ImportAssetParameters.OverridePipelines[GraphPipelineIndex]);
 			if (!GeneratedPipeline)
 			{
 				UE_LOG(LogInterchangeEngine, Error, TEXT("Interchange Import: Overridden pipeline array contains a NULL pipeline. Fix your script or code to avoid this issue."));
diff --git a/Engine/Source/Runtime/Interchange/Engine/Public/InterchangeManager.h b/Engine/Source/Runtime/Interchange/Engine/Public/InterchangeManager.h
--- a/Engine/Source/Runtime/Interchange/Engine/Public/InterchangeManager.h
+++ b/Engine/Source/Runtime/Interchange/Engine/Public/InterchangeManager.h
@@ -239,9 +239,6 @@ namespace UE
 			FGraphEventRef PreCompletionTask;
 			FGraphEventRef CompletionTask;
 
-			// Package where the Pipeline Instances are stored during an import.
-			FString PipelineInstancesPackageName;
-
 			//Return true if we can import this class, or false otherwise.
 			bool IsClassImportAllowed(UClass* Class);
 
@@ -321,10 +318,7 @@ namespace UE
 		void SanitizeObjectName(FString& ObjectName);
 
 		/* This function takes an asset that represents a pipeline and generates a UInterchangePipelineBase asset. */
-		INTERCHANGEENGINE_API UInterchangePipelineBase* GeneratePipelineInstance(const FSoftObjectPath& PipelineInstance, UPackage* PipelineInstancePackage = nullptr);
-
-		INTERCHANGEENGINE_API UInterchangePipelineBase* GeneratePipelineInstanceInSourceAssetPackage(const FSoftObjectPath& PipelineInstance);
-
+		INTERCHANGEENGINE_API UInterchangePipelineBase* GeneratePipelineInstance(const FSoftObjectPath& PipelineInstance);
 	} //ns interchange
 } //ns UE
 
