Interchange: Fix the import animation only. We now set the custom reference object for all factory node so the dependencies works. The import animation only need a valid skeleton to import animation, so we test if the skeleton is empty instead of testing if there is a skeletal mesh.

#rnx

diff --git a/Engine/Plugins/Interchange/Runtime/Source/Import/Private/Animation/InterchangeAnimSequenceFactory.cpp b/Engine/Plugins/Interchange/Runtime/Source/Import/Private/Animation/InterchangeAnimSequenceFactory.cpp
--- a/Engine/Plugins/Interchange/Runtime/Source/Import/Private/Animation/InterchangeAnimSequenceFactory.cpp
+++ b/Engine/Plugins/Interchange/Runtime/Source/Import/Private/Animation/InterchangeAnimSequenceFactory.cpp
@@ -943,22 +943,6 @@ UInterchangeFactoryBase::FImportAssetResult UInterchangeAnimSequenceFactory::Beg
 		return ImportAssetResult;
 	}
 
-	FString SkeletalMeshFactoryNodeUid;
-	if (SkeletonFactoryNode->GetCustomSkeletalMeshFactoryNodeUid(SkeletalMeshFactoryNodeUid))
-	{
-		const UInterchangeSkeletalMeshFactoryNode* SkeletalMeshFactoryNode = Cast<const UInterchangeSkeletalMeshFactoryNode>(Arguments.NodeContainer->GetFactoryNode(SkeletalMeshFactoryNodeUid));
-
-		FSoftObjectPath SkeletalMesh;
-		SkeletalMeshFactoryNode->GetCustomReferenceObject(SkeletalMesh);
-
-		if (NewAnimSequence == nullptr && !SkeletalMesh.IsValid())
-		{
-			//For AnimSequence (single) Re-Imports the SkeletalMesh is not created in the pipeline.
-			UE_LOG(LogInterchangeImport, Error, TEXT("SkeletalMesh does not exist, the skeleton and skeletal mesh is obligatory to import this animsequence [%s]!"), *Arguments.AssetName);
-			return ImportAssetResult;
-		}
-	}
-
 	FSoftObjectPath SkeletonFactoryNodeReferenceObject;
 	SkeletonFactoryNode->GetCustomReferenceObject(SkeletonFactoryNodeReferenceObject);
 
@@ -987,11 +971,17 @@ UInterchangeFactoryBase::FImportAssetResult UInterchangeAnimSequenceFactory::Beg
 
 		if (!ensure(Skeleton))
 		{
-			UE_LOG(LogInterchangeImport, Error, TEXT("Invalid Skeleton when importing animation sequence asset %s."), *Arguments.AssetName);
+			UE_LOG(LogInterchangeImport, Error, TEXT("Invalid skeleton when importing animation sequence asset %s."), *Arguments.AssetName);
 			return ImportAssetResult;
 		}
 	}
 
+	if (Skeleton->GetReferenceSkeleton().GetRawBoneNum() == 0)
+	{
+		UE_LOG(LogInterchangeImport, Error, TEXT("Invalid empty skeleton when importing animation sequence asset %s."), *Arguments.AssetName);
+		return ImportAssetResult;
+	}
+
 	//Verify if the bone track animation is valid (sequence length versus framerate ...)
 	if(!IsBoneTrackAnimationValid(AnimSequenceFactoryNode, Arguments))
 	{
diff --git a/Engine/Source/Runtime/Interchange/Engine/Private/Tasks/InterchangeTaskImportObject.cpp b/Engine/Source/Runtime/Interchange/Engine/Private/Tasks/InterchangeTaskImportObject.cpp
--- a/Engine/Source/Runtime/Interchange/Engine/Private/Tasks/InterchangeTaskImportObject.cpp
+++ b/Engine/Source/Runtime/Interchange/Engine/Private/Tasks/InterchangeTaskImportObject.cpp
@@ -317,6 +317,50 @@ void UE::Interchange::FTaskImportObject_GameThread::DoTask(ENamedThreads::Type C
 
 	bool bSkipObjectNoReplace = false;
 	UObject* ExistingAsset = nullptr;
+
+	auto FollowRedirectorCode = [AsyncHelper, &Pkg, &PackageName, &AssetName]()
+		{
+			if (AsyncHelper->TaskData.bFollowRedirectors)
+			{
+				if (UObjectRedirector* Redirector = FindObject<UObjectRedirector>(Pkg, *AssetName))
+				{
+					if (Redirector->DestinationObject)
+					{
+						Pkg = Redirector->DestinationObject->GetPackage();
+						if (FPackageName::GetLongPackageAssetName(PackageName) == AssetName)
+						{
+							AssetName = FPackageName::GetLongPackageAssetName(Pkg->GetName());
+						}
+					}
+				}
+			}
+		};
+
+	//If the factory node is disable see if there is an existing UObject for it
+	if (!FactoryNode->IsEnabled())
+	{
+		Private::InternalGetPackageName(*AsyncHelper, SourceIndex, PackageBasePath, FactoryNode, PackageName, AssetName);
+		if (!FPackageUtils::IsMapPackageAsset(PackageName))
+		{
+			Pkg = FindPackage(nullptr, *PackageName);
+			if (Pkg)
+			{
+				FollowRedirectorCode();
+				ExistingAsset = StaticFindObject(nullptr, Pkg, *AssetName);
+				if (ExistingAsset)
+				{
+					//Do not call factory for a disabled node but ensure a valid custom reference object for other nodes that depend on this one
+					FSoftObjectPath Reference;
+					if (!FactoryNode->GetCustomReferenceObject(Reference))
+					{
+						FactoryNode->SetCustomReferenceObject(ExistingAsset);
+					}
+				}
+			}
+		}
+		//The node is disabled return now
+		return;
+	}
 	//If we do a reimport no need to create a package
 	if (ObjectToReimport)
 	{
@@ -388,19 +432,9 @@ void UE::Interchange::FTaskImportObject_GameThread::DoTask(ENamedThreads::Type C
 			return;
 		}
 
-		if (!bPackageWasCreated && AsyncHelper->TaskData.bFollowRedirectors)
+		if (!bPackageWasCreated)
 		{
-			if (UObjectRedirector* Redirector = FindObject<UObjectRedirector>(Pkg, *AssetName))
-			{
-				if (Redirector->DestinationObject)
-				{
-					Pkg = Redirector->DestinationObject->GetPackage();
-					if (FPackageName::GetLongPackageAssetName(PackageName) == AssetName)
-					{
-						AssetName = FPackageName::GetLongPackageAssetName(Pkg->GetName());
-					}
-				}
-			}
+			FollowRedirectorCode();
 		}
 		ExistingAsset = StaticFindObject(nullptr, Pkg, *AssetName);
 		if (!bSkipObjectNoReplace && ExistingAsset && !AsyncHelper->TaskData.bReplaceExisting)
diff --git a/Engine/Source/Runtime/Interchange/Engine/Private/Tasks/InterchangeTaskParsing.cpp b/Engine/Source/Runtime/Interchange/Engine/Private/Tasks/InterchangeTaskParsing.cpp
--- a/Engine/Source/Runtime/Interchange/Engine/Private/Tasks/InterchangeTaskParsing.cpp
+++ b/Engine/Source/Runtime/Interchange/Engine/Private/Tasks/InterchangeTaskParsing.cpp
@@ -140,12 +140,6 @@ void UE::Interchange::FTaskParsing::DoTask(ENamedThreads::Type CurrentThread, co
 			const bool bCanImportSceneNode = AsyncHelper->TaskData.ImportType == EImportType::ImportType_Scene;
 			BaseNodeContainer->IterateNodesOfType<UInterchangeFactoryBaseNode>([&](const FString& NodeUID, UInterchangeFactoryBaseNode* FactoryNode)
 			{
-				if (!FactoryNode->IsEnabled())
-				{
-					//Do not call factory for a disabled node
-					return;
-				}
-
 				UClass* ObjectClass = FactoryNode->GetObjectClass();
 				if (ObjectClass != nullptr)
 				{
