Interchange: fix crash cause by null pointer use in interchange manager conversion.

#rnx

diff --git a/Engine/Source/Runtime/Interchange/Engine/Private/InterchangeManager.cpp b/Engine/Source/Runtime/Interchange/Engine/Private/InterchangeManager.cpp
--- a/Engine/Source/Runtime/Interchange/Engine/Private/InterchangeManager.cpp
+++ b/Engine/Source/Runtime/Interchange/Engine/Private/InterchangeManager.cpp
@@ -1042,6 +1042,11 @@ bool UInterchangeManager::RegisterImportDataConverter(const UClass* Converter)
 
 bool UInterchangeManager::ConvertImportData(UObject* Object, const FString& Extension) const
 {
+	if (!Object)
+	{
+		return false;
+	}
+
 	for (TPair<TObjectPtr<const UClass>, TObjectPtr<UInterchangeAssetImportDataConverterBase>> RegisteredConverter : RegisteredConverters)
 	{
 		if (RegisteredConverter.Value->ConvertImportData(Object, Extension))
@@ -1054,6 +1059,11 @@ bool UInterchangeManager::ConvertImportData(UObject* Object, const FString& Exte
 
 bool UInterchangeManager::ConvertImportData(const UObject* SourceImportData, const UClass* DestinationClass, UObject** DestinationImportData) const
 {
+	if (!SourceImportData)
+	{
+		return false;
+	}
+
 	for (TPair<TObjectPtr<const UClass>, TObjectPtr<UInterchangeAssetImportDataConverterBase>> RegisteredConverter : RegisteredConverters)
 	{
 		if (RegisteredConverter.Value->CanConvertClass(SourceImportData->GetClass(), DestinationClass))
