Interchange: Increase the log Interchange source maximum length from 64 to 128 characters.

#rnx

diff --git a/Engine/Source/Runtime/Interchange/Core/Private/InterchangeResult.cpp b/Engine/Source/Runtime/Interchange/Core/Private/InterchangeResult.cpp
--- a/Engine/Source/Runtime/Interchange/Core/Private/InterchangeResult.cpp
+++ b/Engine/Source/Runtime/Interchange/Core/Private/InterchangeResult.cpp
@@ -105,7 +105,7 @@ FText UInterchangeResult::GetMessageLogText() const
 	FString SourceText = SourceAssetName;
 
 	// Make sure file path is not taking the whole line
-	const int32 MaximumSourceNameLength = 64;
+	const int32 MaximumSourceNameLength = 128;
 	if (SourceText.Len() > MaximumSourceNameLength)
 	{
 		SourceText = TEXT("...") + SourceText.RightChop(MaximumSourceNameLength-3);
