Interchange: Always show the full path of the import file when showing a message that contain the source filename.

#rnx

diff --git a/Engine/Source/Runtime/Interchange/Core/Private/InterchangeResult.cpp b/Engine/Source/Runtime/Interchange/Core/Private/InterchangeResult.cpp
--- a/Engine/Source/Runtime/Interchange/Core/Private/InterchangeResult.cpp
+++ b/Engine/Source/Runtime/Interchange/Core/Private/InterchangeResult.cpp
@@ -104,11 +104,12 @@ FText UInterchangeResult::GetMessageLogText() const
 {
 	FString SourceText = SourceAssetName;
 
-	// Make sure file path is not taking the whole line
-	const int32 MaximumSourceNameLength = 128;
+	// Make sure file path is not too long
+	const int32 MaximumSourceNameLength = 512;
 	if (SourceText.Len() > MaximumSourceNameLength)
 	{
-		SourceText = TEXT("...") + SourceText.RightChop(MaximumSourceNameLength-3);
+		const int32 CountingFromLeftRightChopIndex = SourceText.Len() - (MaximumSourceNameLength-3);
+		SourceText = TEXT("...") + SourceText.RightChop(CountingFromLeftRightChopIndex);
 	}
 
 	const FText AssetTypeName = FText::FromString(AssetType ? AssetType->GetName() : TEXT("Unknown"));
