Interchange: Push the post edit change after the import is done, but call a factory build step where we had the posteditchange call. Factory will build the asset if needed at this moment and the posteditchange done at the end should simply cache the result and call any broadcast.

#rnx

#virtualized

diff --git a/Engine/Plugins/Interchange/Runtime/Source/Import/Private/Animation/InterchangeAnimSequenceFactory.cpp b/Engine/Plugins/Interchange/Runtime/Source/Import/Private/Animation/InterchangeAnimSequenceFactory.cpp
--- a/Engine/Plugins/Interchange/Runtime/Source/Import/Private/Animation/InterchangeAnimSequenceFactory.cpp
+++ b/Engine/Plugins/Interchange/Runtime/Source/Import/Private/Animation/InterchangeAnimSequenceFactory.cpp
@@ -1367,6 +1367,29 @@ void UInterchangeAnimSequenceFactory::SetupObject_GameThread(const FSetupObjectP
 #endif
 }
 
+void UInterchangeAnimSequenceFactory::BuildObject_GameThread(const FSetupObjectParams& Arguments, bool& OutPostEditchangeCalled)
+{
+	check(IsInGameThread());
+	OutPostEditchangeCalled = false;
+#if WITH_EDITOR
+	if (Arguments.ImportedObject)
+	{
+		check(AnimSequence == CastChecked<UAnimSequence>(Arguments.ImportedObject));
+		if (AnimSequence)
+		{
+			const bool bNeedPostProcess = !AnimSequence->IsCompressedDataValid();
+
+			// @Todo fix me: This is temporary fix to make sure they always have compressed data
+			if (AnimSequence->IsDataModelValid() && bNeedPostProcess)
+			{
+				AnimSequence->ClearAllCachedCookedPlatformData();
+				AnimSequence->BeginCacheDerivedDataForCurrentPlatform();
+			}
+		}
+	}
+#endif
+}
+
 bool UInterchangeAnimSequenceFactory::GetSourceFilenames(const UObject* Object, TArray<FString>& OutSourceFilenames) const
 {
 #if WITH_EDITORONLY_DATA
diff --git a/Engine/Plugins/Interchange/Runtime/Source/Import/Private/Mesh/InterchangeSkeletalMeshFactory.cpp b/Engine/Plugins/Interchange/Runtime/Source/Import/Private/Mesh/InterchangeSkeletalMeshFactory.cpp
--- a/Engine/Plugins/Interchange/Runtime/Source/Import/Private/Mesh/InterchangeSkeletalMeshFactory.cpp
+++ b/Engine/Plugins/Interchange/Runtime/Source/Import/Private/Mesh/InterchangeSkeletalMeshFactory.cpp
@@ -2063,6 +2063,22 @@ void UInterchangeSkeletalMeshFactory::SetupObject_GameThread(const FSetupObjectP
 #endif //WITH_EDITORONLY_DATA
 }
 
+void UInterchangeSkeletalMeshFactory::BuildObject_GameThread(const FSetupObjectParams& Arguments, bool& OutPostEditchangeCalled)
+{
+	check(IsInGameThread());
+	OutPostEditchangeCalled = false;
+#if WITH_EDITOR
+	if (Arguments.ImportedObject)
+	{
+		if (USkeletalMesh* SkeletalMesh = CastChecked<USkeletalMesh>(Arguments.ImportedObject))
+		{
+			//Start an async build of the staticmesh
+			SkeletalMesh->Build();
+		}
+	}
+#endif
+}
+
 void UInterchangeSkeletalMeshFactory::FinalizeObject_GameThread(const FSetupObjectParams& Arguments)
 {
 	TRACE_CPUPROFILER_EVENT_SCOPE(UInterchangeSkeletalMeshFactory::FinalizeObject_GameThread)
diff --git a/Engine/Plugins/Interchange/Runtime/Source/Import/Private/Mesh/InterchangeStaticMeshFactory.cpp b/Engine/Plugins/Interchange/Runtime/Source/Import/Private/Mesh/InterchangeStaticMeshFactory.cpp
--- a/Engine/Plugins/Interchange/Runtime/Source/Import/Private/Mesh/InterchangeStaticMeshFactory.cpp
+++ b/Engine/Plugins/Interchange/Runtime/Source/Import/Private/Mesh/InterchangeStaticMeshFactory.cpp
@@ -814,6 +814,26 @@ void UInterchangeStaticMeshFactory::SetupObject_GameThread(const FSetupObjectPar
 #endif
 }
 
+void UInterchangeStaticMeshFactory::BuildObject_GameThread(const FSetupObjectParams& Arguments, bool& OutPostEditchangeCalled)
+{
+	check(IsInGameThread());
+	OutPostEditchangeCalled = false;
+#if WITH_EDITOR
+	if (Arguments.ImportedObject)
+	{
+		if (UStaticMesh* StaticMesh = CastChecked<UStaticMesh>(Arguments.ImportedObject))
+		{
+			//Start an async build of the staticmesh
+			UStaticMesh::FBuildParameters BuildParameters;
+			BuildParameters.bInSilent = true;
+			BuildParameters.bInRebuildUVChannelData = true;
+			BuildParameters.bInEnforceLightmapRestrictions = true;
+			StaticMesh->Build(BuildParameters);
+		}
+	}
+#endif
+}
+
 
 TArray<UInterchangeStaticMeshFactory::FMeshPayload> UInterchangeStaticMeshFactory::GetMeshPayloads(const FImportAssetObjectParams& Arguments, const TArray<FString>& MeshUids) const
 {
diff --git a/Engine/Plugins/Interchange/Runtime/Source/Import/Public/Animation/InterchangeAnimSequenceFactory.h b/Engine/Plugins/Interchange/Runtime/Source/Import/Public/Animation/InterchangeAnimSequenceFactory.h
--- a/Engine/Plugins/Interchange/Runtime/Source/Import/Public/Animation/InterchangeAnimSequenceFactory.h
+++ b/Engine/Plugins/Interchange/Runtime/Source/Import/Public/Animation/InterchangeAnimSequenceFactory.h
@@ -33,6 +33,7 @@ public:
 	virtual FImportAssetResult EndImportAsset_GameThread(const FImportAssetObjectParams& Arguments) override;
 
 	virtual void SetupObject_GameThread(const FSetupObjectParams& Arguments) override;
+	virtual void BuildObject_GameThread(const FSetupObjectParams& Arguments, bool& OutPostEditchangeCalled) override;
 
 	virtual bool GetSourceFilenames(const UObject* Object, TArray<FString>& OutSourceFilenames) const override;
 	virtual bool SetSourceFilename(const UObject* Object, const FString& SourceFilename, int32 SourceIndex) const override;
diff --git a/Engine/Plugins/Interchange/Runtime/Source/Import/Public/Mesh/InterchangeSkeletalMeshFactory.h b/Engine/Plugins/Interchange/Runtime/Source/Import/Public/Mesh/InterchangeSkeletalMeshFactory.h
--- a/Engine/Plugins/Interchange/Runtime/Source/Import/Public/Mesh/InterchangeSkeletalMeshFactory.h
+++ b/Engine/Plugins/Interchange/Runtime/Source/Import/Public/Mesh/InterchangeSkeletalMeshFactory.h
@@ -58,6 +58,7 @@ public:
 	virtual FImportAssetResult EndImportAsset_GameThread(const FImportAssetObjectParams& Arguments) override;
 	virtual void Cancel() override;
 	virtual void SetupObject_GameThread(const FSetupObjectParams& Arguments) override;
+	virtual void BuildObject_GameThread(const FSetupObjectParams& Arguments, bool& OutPostEditchangeCalled) override;
 	virtual void FinalizeObject_GameThread(const FSetupObjectParams& Arguments) override;
 	virtual bool GetSourceFilenames(const UObject* Object, TArray<FString>& OutSourceFilenames) const override;
 	virtual bool SetSourceFilename(const UObject* Object, const FString& SourceFilename, int32 SourceIndex) const override;
diff --git a/Engine/Plugins/Interchange/Runtime/Source/Import/Public/Mesh/InterchangeStaticMeshFactory.h b/Engine/Plugins/Interchange/Runtime/Source/Import/Public/Mesh/InterchangeStaticMeshFactory.h
--- a/Engine/Plugins/Interchange/Runtime/Source/Import/Public/Mesh/InterchangeStaticMeshFactory.h
+++ b/Engine/Plugins/Interchange/Runtime/Source/Import/Public/Mesh/InterchangeStaticMeshFactory.h
@@ -34,6 +34,7 @@ public:
 	virtual FImportAssetResult ImportAsset_Async(const FImportAssetObjectParams& Arguments) override;
 	virtual FImportAssetResult EndImportAsset_GameThread(const FImportAssetObjectParams& Arguments) override;
 	virtual void SetupObject_GameThread(const FSetupObjectParams& Arguments) override;
+	virtual void BuildObject_GameThread(const FSetupObjectParams& Arguments, bool& OutPostEditchangeCalled) override;
 	virtual bool GetSourceFilenames(const UObject* Object, TArray<FString>& OutSourceFilenames) const override;
 	virtual bool SetSourceFilename(const UObject* Object, const FString& SourceFilename, int32 SourceIndex) const override;
 	virtual void BackupSourceData(const UObject* Object) const override;
diff --git a/Engine/Source/Editor/UnrealEd/Private/Toolkits/AssetEditorToolkit.cpp b/Engine/Source/Editor/UnrealEd/Private/Toolkits/AssetEditorToolkit.cpp
--- a/Engine/Source/Editor/UnrealEd/Private/Toolkits/AssetEditorToolkit.cpp
+++ b/Engine/Source/Editor/UnrealEd/Private/Toolkits/AssetEditorToolkit.cpp
@@ -231,7 +231,7 @@ void FAssetEditorToolkit::InitAssetEditor( const EToolkitMode::Type Mode, const
 						{
 							if (AsyncAsset->IsCompiling())
 							{
-								return EVisibility::Hidden;
+								return EVisibility::Collapsed;
 							}
 						}
 					}
diff --git a/Engine/Source/Runtime/Interchange/Core/Public/InterchangeFactoryBase.h b/Engine/Source/Runtime/Interchange/Core/Public/InterchangeFactoryBase.h
--- a/Engine/Source/Runtime/Interchange/Core/Public/InterchangeFactoryBase.h
+++ b/Engine/Source/Runtime/Interchange/Core/Public/InterchangeFactoryBase.h
@@ -49,10 +49,11 @@ public:
  * 2. ImportAsset_Async - Import source data (retrieve payloads) and set up properties asynchronously on any thread.
  * 3. EndImportAsset_GameThread - Anything you need to do on the game thread to finalize the import source data and set up properties. For example, conflict resolution that needs UI.
  * 
- * The last two steps are helpful to change the imported or reimported UObject before and after the PostEditChange (render data build) is called on the asset.
+ * The last three steps can modify the created UObject
  * 
- * 4. SetupObject_GameThread - Do any UObject setup required before the build (before PostEditChange), the UObject dependencies should exist and have all the source data and properties imported.
- * 5. FinalizeObject_GameThread - Do any final UObject setup after the build (after PostEditChange). Note that the build of an asset can be asynchronous and this function will be call after the async build is done.
+ * 4. SetupObject_GameThread - Do any UObject setup required before the build, the UObject dependencies should exist and have all the source data and properties imported.
+ * 5. BuildObject_GameThread - Build the asset if it can be built.
+ * 6. FinalizeObject_GameThread - Do any final UObject setup after the build. Note that the build of an asset can be asynchronous and this function will be call after the async build is done.
  * 
  * Scene factory implementation:
  * 
@@ -266,6 +267,18 @@ public:
 		check(IsInGameThread());
 	}
 	
+	virtual void BuildObject_GameThread(const FSetupObjectParams& Arguments, bool &OutPostEditchangeCalled)
+	{
+		check(IsInGameThread());
+#if WITH_EDITOR
+		if (Arguments.ImportedObject)
+		{
+			Arguments.ImportedObject->PostEditChange();
+			OutPostEditchangeCalled = true;
+		}
+#endif //WITH_EDITOR
+	}
+
 	/*
 	 * Do any final UObject setup after the build (after PostEditChange).
 	 * @note - This function is called at the end of the pre-completion task (after PostEditChange is called for the asset).
diff --git a/Engine/Source/Runtime/Interchange/Engine/Private/Tasks/InterchangeTaskCompletion.cpp b/Engine/Source/Runtime/Interchange/Engine/Private/Tasks/InterchangeTaskCompletion.cpp
--- a/Engine/Source/Runtime/Interchange/Engine/Private/Tasks/InterchangeTaskCompletion.cpp
+++ b/Engine/Source/Runtime/Interchange/Engine/Private/Tasks/InterchangeTaskCompletion.cpp
@@ -71,7 +71,6 @@ void UE::Interchange::FTaskPreCompletion::DoTask(ENamedThreads::Type CurrentThre
 			if (bCallPostImportGameThreadCallback && ObjectInfo.Factory)
 			{
 				Arguments.ImportedObject = ImportedObject;
-				// Should we assert if there is no factory node?
 				Arguments.FactoryNode = ObjectInfo.FactoryNode;
 				Arguments.NodeUniqueID = ObjectInfo.FactoryNode ? ObjectInfo.FactoryNode->GetUniqueID() : FString();
 				Arguments.bIsReimport = ObjectInfo.bIsReimport;
@@ -128,7 +127,7 @@ void UE::Interchange::FTaskPreCompletion::DoTask(ENamedThreads::Type CurrentThre
 		}
 
 #if WITH_EDITOR
-		//Second iteration to call PostEditChange
+		//Second iteration to call BuildObject
 		for (const FImportAsyncHelper::FImportedObjectInfo& ObjectInfo : ImportedObjects)
 		{
 			UObject* ImportedObject = ObjectInfo.ImportedObject;
@@ -136,7 +135,13 @@ void UE::Interchange::FTaskPreCompletion::DoTask(ENamedThreads::Type CurrentThre
 			{
 				continue;
 			}
-			ImportedObject->PostEditChange();
+			//The base class of the factory will call posteditchange, but other factory can instead simply build
+			//the asset asynchronously and the post edit change will be call later
+			Arguments.ImportedObject = ImportedObject;
+			Arguments.FactoryNode = ObjectInfo.FactoryNode;
+			Arguments.NodeUniqueID = ObjectInfo.FactoryNode ? ObjectInfo.FactoryNode->GetUniqueID() : FString();
+			Arguments.bIsReimport = ObjectInfo.bIsReimport;
+			ObjectInfo.Factory->BuildObject_GameThread(Arguments, ObjectInfo.bPostEditChangeCalled);
 		}
 #endif //WITH_EDITOR
 
@@ -199,13 +204,20 @@ void UE::Interchange::FTaskCompletion::DoTask(ENamedThreads::Type CurrentThread,
 			{
 				for (const FImportAsyncHelper::FImportedObjectInfo& AssetInfo : AssetInfos)
 				{
-					UObject* Asset = AssetInfo.ImportedObject;
-					if (AsyncHelper->TaskData.ReimportObject && AsyncHelper->TaskData.ReimportObject == Asset)
+					if (UObject* Asset = AssetInfo.ImportedObject)
 					{
-						UInterchangeManager::GetInterchangeManager().OnAssetPostReimport.Broadcast(Asset);
+						//Call post edit change if it was not call previously
+						if (!AssetInfo.bPostEditChangeCalled)
+						{
+							Asset->PostEditChange();
+						}
+						if (AsyncHelper->TaskData.ReimportObject && AsyncHelper->TaskData.ReimportObject == Asset)
+						{
+							UInterchangeManager::GetInterchangeManager().OnAssetPostReimport.Broadcast(Asset);
+						}
+						//We broadcast this event for both import and reimport.
+						UInterchangeManager::GetInterchangeManager().OnAssetPostImport.Broadcast(Asset);
 					}
-					//We broadcast this event for both import and reimport.
-					UInterchangeManager::GetInterchangeManager().OnAssetPostImport.Broadcast(Asset);
 				}
 
 				UE_LOG(LogInterchangeEngine, Display, TEXT("Interchange import completed [%s]"), *AsyncHelper->SourceDatas[SourceIndex]->ToDisplayString());
diff --git a/Engine/Source/Runtime/Interchange/Engine/Public/InterchangeManager.h b/Engine/Source/Runtime/Interchange/Engine/Public/InterchangeManager.h
--- a/Engine/Source/Runtime/Interchange/Engine/Public/InterchangeManager.h
+++ b/Engine/Source/Runtime/Interchange/Engine/Public/InterchangeManager.h
@@ -254,6 +254,7 @@ namespace UE
 				UInterchangeFactoryBase* Factory = nullptr; //The factory that created the imported object.
 				UInterchangeFactoryBaseNode* FactoryNode = nullptr; //The node that describes the object.
 				bool bIsReimport;
+				mutable bool bPostEditChangeCalled = false; //This field is set by the PreCompletionTask and need to be mutable
 			};
 
 			FImportedObjectInfo& AddDefaultImportedAssetGetRef(int32 SourceIndex);
