Fix CIS build
Interchange: Update the comment in the factory base class to reflect the latest API.

#rnx

diff --git a/Engine/Source/Runtime/Interchange/Core/Public/InterchangeFactoryBase.h b/Engine/Source/Runtime/Interchange/Core/Public/InterchangeFactoryBase.h
--- a/Engine/Source/Runtime/Interchange/Core/Public/InterchangeFactoryBase.h
+++ b/Engine/Source/Runtime/Interchange/Core/Public/InterchangeFactoryBase.h
@@ -259,14 +259,22 @@ public:
 	};
 
 	/*
-	 * Do any UObject setup required before the build (before PostEditChange) and after all dependency UObjects have been imported.
-	 * @note - This function is called when starting the pre-completion task (before PostEditChange is called for the asset).
+	 * Do any UObject setup required before the build and after all dependency UObjects have been imported.
+	 * @note - This function is called when starting the pre-completion task (before asset build is called for the asset).
 	 */
 	virtual void SetupObject_GameThread(const FSetupObjectParams& Arguments)
 	{
 		check(IsInGameThread());
 	}
 	
+	/*
+	 * Build the asset if it can be built.
+	 * 
+	 * @Param Arguments - The setup object data.
+	 * @Param OutPostEditchangeCalled - Set it to true if your implementation has called PostEditChange. Set it to false, interchange will call PostEditChange at the end of the import.
+	 * @note - This function is called when starting the pre-completion task.
+	 * @note - The default implementation for asset that do not have any build step is to call PostEditChange and set OutPostEditchangeCalled to true.
+	 */
 	virtual void BuildObject_GameThread(const FSetupObjectParams& Arguments, bool &OutPostEditchangeCalled)
 	{
 		check(IsInGameThread());
@@ -280,8 +288,8 @@ public:
 	}
 
 	/*
-	 * Do any final UObject setup after the build (after PostEditChange).
-	 * @note - This function is called at the end of the pre-completion task (after PostEditChange is called for the asset).
+	 * Do any final UObject setup after the asset is built
+	 * @note - This function is called after asset is built in FTaskWaitAssetCompilation
 	 */
 	virtual void FinalizeObject_GameThread(const FSetupObjectParams& Arguments)
 	{
diff --git a/Engine/Source/Runtime/Interchange/Engine/Private/Tasks/InterchangeTaskCompletion.cpp b/Engine/Source/Runtime/Interchange/Engine/Private/Tasks/InterchangeTaskCompletion.cpp
--- a/Engine/Source/Runtime/Interchange/Engine/Private/Tasks/InterchangeTaskCompletion.cpp
+++ b/Engine/Source/Runtime/Interchange/Engine/Private/Tasks/InterchangeTaskCompletion.cpp
@@ -207,10 +207,12 @@ void UE::Interchange::FTaskCompletion::DoTask(ENamedThreads::Type CurrentThread,
 					if (UObject* Asset = AssetInfo.ImportedObject)
 					{
 						//Call post edit change if it was not call previously
+#if WITH_EDITOR
 						if (!AssetInfo.bPostEditChangeCalled)
 						{
 							Asset->PostEditChange();
 						}
+#endif //WITH_EDITOR
 						if (AsyncHelper->TaskData.ReimportObject && AsyncHelper->TaskData.ReimportObject == Asset)
 						{
 							UInterchangeManager::GetInterchangeManager().OnAssetPostReimport.Broadcast(Asset);
