Interchange: Fix staticmesh material match crash.

Add a parameter to pipeline API AdjustSettingsForContext, this parameter is the node container so it give more context to adjust the settings. Mesh pipeline can force mesh type when reimporting a mesh in case the file contain the wrong mesh type.

#rnx

diff --git a/Engine/Plugins/Experimental/Enterprise/DatasmithInterchange/Source/DatasmithInterchange/Private/InterchangeDatasmithPipeline.cpp b/Engine/Plugins/Experimental/Enterprise/DatasmithInterchange/Source/DatasmithInterchange/Private/InterchangeDatasmithPipeline.cpp
--- a/Engine/Plugins/Experimental/Enterprise/DatasmithInterchange/Source/DatasmithInterchange/Private/InterchangeDatasmithPipeline.cpp
+++ b/Engine/Plugins/Experimental/Enterprise/DatasmithInterchange/Source/DatasmithInterchange/Private/InterchangeDatasmithPipeline.cpp
@@ -120,28 +120,28 @@ UInterchangeDatasmithPipeline::UInterchangeDatasmithPipeline()
 	AnimationPipeline->CommonSkeletalMeshesAndAnimationsProperties = CommonSkeletalMeshesAndAnimationsProperties;
 }
 
-void UInterchangeDatasmithPipeline::AdjustSettingsForContext(EInterchangePipelineContext ImportType, TObjectPtr<UObject> ReimportAsset)
+void UInterchangeDatasmithPipeline::AdjustSettingsForContext(EInterchangePipelineContext ImportType, TObjectPtr<UObject> ReimportAsset, const UInterchangeBaseNodeContainer* InBaseNodeContainer)
 {
-	Super::AdjustSettingsForContext(ImportType, ReimportAsset);
+	Super::AdjustSettingsForContext(ImportType, ReimportAsset, InBaseNodeContainer);
 
 	if (MaterialPipeline)
 	{
-		MaterialPipeline->AdjustSettingsForContext(ImportType, ReimportAsset);
+		MaterialPipeline->AdjustSettingsForContext(ImportType, ReimportAsset, InBaseNodeContainer);
 	}
 
 	if (MeshPipeline)
 	{
-		MeshPipeline->AdjustSettingsForContext(ImportType, ReimportAsset);
+		MeshPipeline->AdjustSettingsForContext(ImportType, ReimportAsset, InBaseNodeContainer);
 	}
 
 	if (LevelPipeline)
 	{
-		LevelPipeline->AdjustSettingsForContext(ImportType, ReimportAsset);
+		LevelPipeline->AdjustSettingsForContext(ImportType, ReimportAsset, InBaseNodeContainer);
 	}
 
 	if (AnimationPipeline)
 	{
-		AnimationPipeline->AdjustSettingsForContext(ImportType, ReimportAsset);
+		AnimationPipeline->AdjustSettingsForContext(ImportType, ReimportAsset, InBaseNodeContainer);
 	}
 }
 
@@ -150,7 +150,7 @@ void UInterchangeDatasmithPipeline::PostDuplicate(bool bDuplicateForPIE)
 	// Only adjust settings if there is anything cached.
 	if (CachePipelineContext != EInterchangePipelineContext::None)
 	{
-		AdjustSettingsForContext(CachePipelineContext, CacheReimportObject.Get());
+		AdjustSettingsForContext(CachePipelineContext, CacheReimportObject.Get(), nullptr);
 	}
 }
 
diff --git a/Engine/Plugins/Experimental/Enterprise/DatasmithInterchange/Source/DatasmithInterchange/Public/InterchangeDatasmithPipeline.h b/Engine/Plugins/Experimental/Enterprise/DatasmithInterchange/Source/DatasmithInterchange/Public/InterchangeDatasmithPipeline.h
--- a/Engine/Plugins/Experimental/Enterprise/DatasmithInterchange/Source/DatasmithInterchange/Public/InterchangeDatasmithPipeline.h
+++ b/Engine/Plugins/Experimental/Enterprise/DatasmithInterchange/Source/DatasmithInterchange/Public/InterchangeDatasmithPipeline.h
@@ -64,7 +64,7 @@ public:
 protected:
 	virtual void ExecutePipeline(UInterchangeBaseNodeContainer* InBaseNodeContainer, const TArray<UInterchangeSourceData*>& SourceDatas, const FString& ContentBasePath) override;
 	virtual void ExecutePostImportPipeline(const UInterchangeBaseNodeContainer* InBaseNodeContainer, const FString& NodeKey, UObject* CreatedAsset, bool bIsAReimport) override;
-	virtual void AdjustSettingsForContext(EInterchangePipelineContext ImportType, TObjectPtr<UObject> ReimportAsset) override;
+	virtual void AdjustSettingsForContext(EInterchangePipelineContext ImportType, TObjectPtr<UObject> ReimportAsset, const UInterchangeBaseNodeContainer* BaseNodeContainer) override;
 
 	virtual bool CanExecuteOnAnyThread(EInterchangePipelineTask PipelineTask) override
 	{
diff --git a/Engine/Plugins/Interchange/Editor/Source/Pipelines/Private/SInterchangePipelineConfigurationDialog.cpp b/Engine/Plugins/Interchange/Editor/Source/Pipelines/Private/SInterchangePipelineConfigurationDialog.cpp
--- a/Engine/Plugins/Interchange/Editor/Source/Pipelines/Private/SInterchangePipelineConfigurationDialog.cpp
+++ b/Engine/Plugins/Interchange/Editor/Source/Pipelines/Private/SInterchangePipelineConfigurationDialog.cpp
@@ -477,6 +477,13 @@ void SInterchangePipelineConfigurationDialog::Construct(const FArguments& InArgs
 		ensure(bReimport);
 	}
 	
+	FText ReuseSettingsTooltipText = LOCTEXT("InspectorGraphWindow_ReuseSettingsTooltipText", "When importing multiple files this checkbox allow users to use the same settings for source of the same extension.");
+	const bool bTranslatorThreadSafe = Translator.IsValid() ? Translator->IsThreadSafe() : false;
+	if (!bTranslatorThreadSafe)
+	{
+		FString Extension = SourceData.IsValid() ? FPaths::GetExtension(SourceData->GetFilename()) : TEXT("N/A");
+		ReuseSettingsTooltipText = FText::Format(LOCTEXT("InspectorGraphWindow_ReuseSettingsNotThreadSafeTooltipText", "{0} translator is not thread safe and must use the same settings for subsequent files"), FText::FromString(Extension));
+	}
 
 	check(OutPipelines);
 
@@ -646,14 +653,23 @@ void SInterchangePipelineConfigurationDialog::Construct(const FArguments& InArgs
 					.Padding(4.f, 0.f)
 					[
 						SNew(STextBlock)
+						.IsEnabled(bTranslatorThreadSafe)
 						.Text(LOCTEXT("InspectorGraphWindow_ReuseSettings", "Use the same settings for subsequent files"))
+						.ToolTipText(ReuseSettingsTooltipText)
 					]
 					+ SHorizontalBox::Slot()
 					.Padding(4.f, 0.f)
 					[
 						SAssignNew(UseSameSettingsForAllCheckBox, SCheckBox)
 						.IsChecked(true)
-						.IsEnabled(this, &SInterchangePipelineConfigurationDialog::IsImportButtonEnabled)
+						.IsEnabled_Lambda([this, bTranslatorThreadSafe]()
+							{
+								if (!bTranslatorThreadSafe)
+								{
+									return false;
+								}
+								return IsImportButtonEnabled();
+							})
 					]
 				]
 				+ SHorizontalBox::Slot()
diff --git a/Engine/Plugins/Interchange/Runtime/Source/Import/Private/Mesh/InterchangeSkeletalMeshFactory.cpp b/Engine/Plugins/Interchange/Runtime/Source/Import/Private/Mesh/InterchangeSkeletalMeshFactory.cpp
--- a/Engine/Plugins/Interchange/Runtime/Source/Import/Private/Mesh/InterchangeSkeletalMeshFactory.cpp
+++ b/Engine/Plugins/Interchange/Runtime/Source/Import/Private/Mesh/InterchangeSkeletalMeshFactory.cpp
@@ -1684,6 +1684,15 @@ UInterchangeFactoryBase::FImportAssetResult UInterchangeSkeletalMeshFactory::End
 		{
 			if (FSkeletalMeshLODInfo* LodInfo = SkeletalMesh->GetLODInfo(ImportAssetObjectLODData.LodIndex))
 			{
+				if (SkeletalMesh->GetMaterials().IsEmpty() && !ImportAssetObjectLODData.ImportedMaterials.IsEmpty() )
+				{
+					for (SkeletalMeshImportData::FMaterial ImportMaterial : ImportAssetObjectLODData.ImportedMaterials)
+					{
+						UMaterialInterface* MaterialInterface = ImportMaterial.Material.Get() ? ImportMaterial.Material.Get() : UMaterial::GetDefaultMaterial(MD_Surface);
+						SkeletalMesh->GetMaterials().Add(FSkeletalMaterial(MaterialInterface, FName(ImportMaterial.MaterialImportName), FName(ImportMaterial.MaterialImportName)));
+					}
+				}
+
 				if (SkeletalMesh->GetMaterials().Num() > 0)
 				{
 					FLODUtilities::FSkeletalMeshMatchImportedMaterialsParameters Parameters;
diff --git a/Engine/Plugins/Interchange/Runtime/Source/Import/Private/Mesh/InterchangeStaticMeshFactory.cpp b/Engine/Plugins/Interchange/Runtime/Source/Import/Private/Mesh/InterchangeStaticMeshFactory.cpp
--- a/Engine/Plugins/Interchange/Runtime/Source/Import/Private/Mesh/InterchangeStaticMeshFactory.cpp
+++ b/Engine/Plugins/Interchange/Runtime/Source/Import/Private/Mesh/InterchangeStaticMeshFactory.cpp
@@ -566,7 +566,10 @@ UInterchangeFactoryBase::FImportAssetResult UInterchangeStaticMeshFactory::EndIm
 					//This is an unmatched section, its either added or we did not recover the name
 					int32 MaterialSlotIndex = StaticMesh->GetMaterialIndexFromImportedMaterialSlotName(SlotNames[PolygonGroupID]);
 					//Missing material slot should have been added before
-					ensure(MaterialSlotIndex != INDEX_NONE);
+					if (MaterialSlotIndex == INDEX_NONE)
+					{
+						MaterialSlotIndex = 0;
+					}
 					NewSectionInfoMapData.Add(FMeshSectionInfo(MaterialSlotIndex));
 				}
 			}
diff --git a/Engine/Plugins/Interchange/Runtime/Source/Pipelines/Private/InterchangeGenericAnimationPipeline.cpp b/Engine/Plugins/Interchange/Runtime/Source/Pipelines/Private/InterchangeGenericAnimationPipeline.cpp
--- a/Engine/Plugins/Interchange/Runtime/Source/Pipelines/Private/InterchangeGenericAnimationPipeline.cpp
+++ b/Engine/Plugins/Interchange/Runtime/Source/Pipelines/Private/InterchangeGenericAnimationPipeline.cpp
@@ -36,9 +36,9 @@ bool UInterchangeGenericAnimationPipeline::CanEditChange(const FProperty* InProp
 }
 #endif
 
-void UInterchangeGenericAnimationPipeline::AdjustSettingsForContext(EInterchangePipelineContext ImportType, TObjectPtr<UObject> ReimportAsset)
+void UInterchangeGenericAnimationPipeline::AdjustSettingsForContext(EInterchangePipelineContext ImportType, TObjectPtr<UObject> ReimportAsset, const UInterchangeBaseNodeContainer* InBaseNodeContainer)
 {
-	Super::AdjustSettingsForContext(ImportType, ReimportAsset);
+	Super::AdjustSettingsForContext(ImportType, ReimportAsset, InBaseNodeContainer);
 
 #if WITH_EDITOR
 	check(CommonSkeletalMeshesAndAnimationsProperties.IsValid());
diff --git a/Engine/Plugins/Interchange/Runtime/Source/Pipelines/Private/InterchangeGenericAssetsPipeline.cpp b/Engine/Plugins/Interchange/Runtime/Source/Pipelines/Private/InterchangeGenericAssetsPipeline.cpp
--- a/Engine/Plugins/Interchange/Runtime/Source/Pipelines/Private/InterchangeGenericAssetsPipeline.cpp
+++ b/Engine/Plugins/Interchange/Runtime/Source/Pipelines/Private/InterchangeGenericAssetsPipeline.cpp
@@ -125,23 +125,23 @@ bool UInterchangeGenericAssetsPipeline::IsSettingsAreValid(TOptional<FText>& Out
 	return Super::IsSettingsAreValid(OutInvalidReason);
 }
 
-void UInterchangeGenericAssetsPipeline::AdjustSettingsForContext(EInterchangePipelineContext ImportType, TObjectPtr<UObject> ReimportAsset)
+void UInterchangeGenericAssetsPipeline::AdjustSettingsForContext(EInterchangePipelineContext ImportType, TObjectPtr<UObject> ReimportAsset, const UInterchangeBaseNodeContainer* InBaseNodeContainer)
 {
-	Super::AdjustSettingsForContext(ImportType, ReimportAsset);
+	Super::AdjustSettingsForContext(ImportType, ReimportAsset, InBaseNodeContainer);
 
 	if (MaterialPipeline)
 	{
-		MaterialPipeline->AdjustSettingsForContext(ImportType, ReimportAsset);
+		MaterialPipeline->AdjustSettingsForContext(ImportType, ReimportAsset, InBaseNodeContainer);
 	}
 
 	if (MeshPipeline)
 	{
-		MeshPipeline->AdjustSettingsForContext(ImportType, ReimportAsset);
+		MeshPipeline->AdjustSettingsForContext(ImportType, ReimportAsset, InBaseNodeContainer);
 	}
 
 	if (AnimationPipeline)
 	{
-		AnimationPipeline->AdjustSettingsForContext(ImportType, ReimportAsset);
+		AnimationPipeline->AdjustSettingsForContext(ImportType, ReimportAsset, InBaseNodeContainer);
 	}
 }
 
@@ -433,7 +433,7 @@ void UInterchangeGenericAssetsPipeline::CreateMaterialConflict(UStaticMesh* Stat
 		}
 	}
 	//Compare and cache the results
-	bool bHasConflict = false;
+	bool bHasConflict = MaterialNodePerMaterialSlotName.IsEmpty() && AssetImportMaterialNames.Num() > 1; //If there is no material dependencies for the mesh factory node and we have more then one material in the static mesh, we have a conflict
 	TArray<int32> MatchMaterials;
 	int32 MatchMaterialCount = 0;
 	int32 ImportMaterialIndex = 0;
diff --git a/Engine/Plugins/Interchange/Runtime/Source/Pipelines/Private/InterchangeGenericMaterialPipeline.cpp b/Engine/Plugins/Interchange/Runtime/Source/Pipelines/Private/InterchangeGenericMaterialPipeline.cpp
--- a/Engine/Plugins/Interchange/Runtime/Source/Pipelines/Private/InterchangeGenericMaterialPipeline.cpp
+++ b/Engine/Plugins/Interchange/Runtime/Source/Pipelines/Private/InterchangeGenericMaterialPipeline.cpp
@@ -602,13 +602,13 @@ bool UInterchangeGenericMaterialPipeline::IsSettingsAreValid(TOptional<FText>& O
 	return Super::IsSettingsAreValid(OutInvalidReason);
 }
 
-void UInterchangeGenericMaterialPipeline::AdjustSettingsForContext(EInterchangePipelineContext ImportType, TObjectPtr<UObject> ReimportAsset)
+void UInterchangeGenericMaterialPipeline::AdjustSettingsForContext(EInterchangePipelineContext ImportType, TObjectPtr<UObject> ReimportAsset, const UInterchangeBaseNodeContainer* InBaseNodeContainer)
 {
-	Super::AdjustSettingsForContext(ImportType, ReimportAsset);
+	Super::AdjustSettingsForContext(ImportType, ReimportAsset, InBaseNodeContainer);
 
 	if (TexturePipeline)
 	{
-		TexturePipeline->AdjustSettingsForContext(ImportType, ReimportAsset);
+		TexturePipeline->AdjustSettingsForContext(ImportType, ReimportAsset, InBaseNodeContainer);
 	}
 #if WITH_EDITOR
 	TArray<FString> HideCategories;
diff --git a/Engine/Plugins/Interchange/Runtime/Source/Pipelines/Private/InterchangeGenericMeshPipeline.cpp b/Engine/Plugins/Interchange/Runtime/Source/Pipelines/Private/InterchangeGenericMeshPipeline.cpp
--- a/Engine/Plugins/Interchange/Runtime/Source/Pipelines/Private/InterchangeGenericMeshPipeline.cpp
+++ b/Engine/Plugins/Interchange/Runtime/Source/Pipelines/Private/InterchangeGenericMeshPipeline.cpp
@@ -21,9 +21,9 @@
 
 #include UE_INLINE_GENERATED_CPP_BY_NAME(InterchangeGenericMeshPipeline)
 
-void UInterchangeGenericMeshPipeline::AdjustSettingsForContext(EInterchangePipelineContext ImportType, TObjectPtr<UObject> ReimportAsset)
+void UInterchangeGenericMeshPipeline::AdjustSettingsForContext(EInterchangePipelineContext ImportType, TObjectPtr<UObject> ReimportAsset, const UInterchangeBaseNodeContainer* InBaseNodeContainer)
 {
-	Super::AdjustSettingsForContext(ImportType, ReimportAsset);
+	Super::AdjustSettingsForContext(ImportType, ReimportAsset, InBaseNodeContainer);
 
 #if WITH_EDITOR
 
@@ -75,6 +75,23 @@ void UInterchangeGenericMeshPipeline::AdjustSettingsForContext(EInterchangePipel
 	TArray<FString> HideSubCategories;
 	if (ImportType == EInterchangePipelineContext::AssetReimport)
 	{
+		bool bContainSkeletalMesh = false;
+		bool bContainStaticMesh = false;
+
+		if (InBaseNodeContainer)
+		{
+			InBaseNodeContainer->BreakableIterateNodesOfType<UInterchangeMeshNode>([&bContainSkeletalMesh, &bContainStaticMesh](const FString& NodeUid, UInterchangeMeshNode* MeshNode)
+				{
+					if (!MeshNode->IsMorphTarget())
+					{
+						MeshNode->IsSkinnedMesh() ? bContainSkeletalMesh = true : bContainStaticMesh = true;
+					}
+					return bContainSkeletalMesh && bContainStaticMesh;
+				});
+		}
+
+		CommonMeshesProperties->bAutoDetectMeshType = false;
+
 		HideSubCategories.Add(TEXT("Build"));
 		if (USkeletalMesh* SkeletalMesh = Cast<USkeletalMesh>(ReimportAsset))
 		{
@@ -87,7 +104,7 @@ void UInterchangeGenericMeshPipeline::AdjustSettingsForContext(EInterchangePipel
 			}
 			bImportStaticMeshes = false;
 			HideCategories.Add(StaticMeshesCategory);
-			if(SkeletalMeshImportContentType == EInterchangeSkeletalMeshContentType::Geometry)
+			if(!bContainSkeletalMesh || SkeletalMeshImportContentType == EInterchangeSkeletalMeshContentType::Geometry)
 			{
 				CommonMeshesProperties->ForceAllMeshAsType = EInterchangeForceMeshType::IFMT_SkeletalMesh;
 			}
@@ -96,6 +113,11 @@ void UInterchangeGenericMeshPipeline::AdjustSettingsForContext(EInterchangePipel
 		{
 			HideCategories.Add(SkeletalMeshesCategory);
 			HideCategories.Add(CommonSkeletalMeshesAndAnimationCategory);
+			
+			if (!bContainStaticMesh)
+			{
+				CommonMeshesProperties->ForceAllMeshAsType = EInterchangeForceMeshType::IFMT_StaticMesh;
+			}
 		}
 		else if (UAnimSequence* AnimSequence = Cast<UAnimSequence>(ReimportAsset))
 		{
diff --git a/Engine/Plugins/Interchange/Runtime/Source/Pipelines/Private/InterchangeGenericScenesPipeline.cpp b/Engine/Plugins/Interchange/Runtime/Source/Pipelines/Private/InterchangeGenericScenesPipeline.cpp
--- a/Engine/Plugins/Interchange/Runtime/Source/Pipelines/Private/InterchangeGenericScenesPipeline.cpp
+++ b/Engine/Plugins/Interchange/Runtime/Source/Pipelines/Private/InterchangeGenericScenesPipeline.cpp
@@ -218,9 +218,9 @@ namespace UE::Interchange::Private
 	}
 }
 
-void UInterchangeGenericLevelPipeline::AdjustSettingsForContext(EInterchangePipelineContext ImportType, TObjectPtr<UObject> ReimportAsset)
+void UInterchangeGenericLevelPipeline::AdjustSettingsForContext(EInterchangePipelineContext ImportType, TObjectPtr<UObject> ReimportAsset, const UInterchangeBaseNodeContainer* InBaseNodeContainer)
 {
-	Super::AdjustSettingsForContext(ImportType, ReimportAsset);
+	Super::AdjustSettingsForContext(ImportType, ReimportAsset, InBaseNodeContainer);
 
 	bIsReimportContext = ReimportAsset != nullptr;
 }
diff --git a/Engine/Plugins/Interchange/Runtime/Source/Pipelines/Private/InterchangeGenericTexturePipeline.cpp b/Engine/Plugins/Interchange/Runtime/Source/Pipelines/Private/InterchangeGenericTexturePipeline.cpp
--- a/Engine/Plugins/Interchange/Runtime/Source/Pipelines/Private/InterchangeGenericTexturePipeline.cpp
+++ b/Engine/Plugins/Interchange/Runtime/Source/Pipelines/Private/InterchangeGenericTexturePipeline.cpp
@@ -110,9 +110,9 @@ namespace UE::Interchange::Private
 	}
 }
 
-void UInterchangeGenericTexturePipeline::AdjustSettingsForContext(EInterchangePipelineContext ImportType, TObjectPtr<UObject> ReimportAsset)
+void UInterchangeGenericTexturePipeline::AdjustSettingsForContext(EInterchangePipelineContext ImportType, TObjectPtr<UObject> ReimportAsset, const UInterchangeBaseNodeContainer* InBaseNodeContainer)
 {
-	Super::AdjustSettingsForContext(ImportType, ReimportAsset);
+	Super::AdjustSettingsForContext(ImportType, ReimportAsset, InBaseNodeContainer);
 #if WITH_EDITOR
 	TArray<FString> HideCategories;
 	bool bIsObjectATexture = !ReimportAsset ? false : ReimportAsset.IsA(UTexture::StaticClass());
diff --git a/Engine/Plugins/Interchange/Runtime/Source/Pipelines/Private/InterchangeMaterialXPipeline.cpp b/Engine/Plugins/Interchange/Runtime/Source/Pipelines/Private/InterchangeMaterialXPipeline.cpp
--- a/Engine/Plugins/Interchange/Runtime/Source/Pipelines/Private/InterchangeMaterialXPipeline.cpp
+++ b/Engine/Plugins/Interchange/Runtime/Source/Pipelines/Private/InterchangeMaterialXPipeline.cpp
@@ -682,9 +682,9 @@ UInterchangeMaterialXPipeline::UInterchangeMaterialXPipeline()
 	}
 }
 
-void UInterchangeMaterialXPipeline::AdjustSettingsForContext(EInterchangePipelineContext ImportType, TObjectPtr<UObject> ReimportAsset)
+void UInterchangeMaterialXPipeline::AdjustSettingsForContext(EInterchangePipelineContext ImportType, TObjectPtr<UObject> ReimportAsset, const UInterchangeBaseNodeContainer* InBaseNodeContainer)
 {
-	Super::AdjustSettingsForContext(ImportType, ReimportAsset);
+	Super::AdjustSettingsForContext(ImportType, ReimportAsset, InBaseNodeContainer);
 
 	if (!MaterialXSettings->AreRequiredPackagesLoaded())
 	{
diff --git a/Engine/Plugins/Interchange/Runtime/Source/Pipelines/Private/InterchangeglTFPipeline.cpp b/Engine/Plugins/Interchange/Runtime/Source/Pipelines/Private/InterchangeglTFPipeline.cpp
--- a/Engine/Plugins/Interchange/Runtime/Source/Pipelines/Private/InterchangeglTFPipeline.cpp
+++ b/Engine/Plugins/Interchange/Runtime/Source/Pipelines/Private/InterchangeglTFPipeline.cpp
@@ -287,9 +287,9 @@ UInterchangeGLTFPipeline::UInterchangeGLTFPipeline()
 {
 }
 
-void UInterchangeGLTFPipeline::AdjustSettingsForContext(EInterchangePipelineContext ImportType, TObjectPtr<UObject> ReimportAsset)
+void UInterchangeGLTFPipeline::AdjustSettingsForContext(EInterchangePipelineContext ImportType, TObjectPtr<UObject> ReimportAsset, const UInterchangeBaseNodeContainer* InBaseNodeContainer)
 {
-	Super::AdjustSettingsForContext(ImportType, ReimportAsset);
+	Super::AdjustSettingsForContext(ImportType, ReimportAsset, InBaseNodeContainer);
 
 	TArray<FString> MaterialInstanceIssues = GLTFPipelineSettings->ValidateMaterialInstancesAndParameters();
 	for (const FString& MaterialInstanceIssue : MaterialInstanceIssues)
diff --git a/Engine/Plugins/Interchange/Runtime/Source/Pipelines/Public/InterchangeGenericAnimationPipeline.h b/Engine/Plugins/Interchange/Runtime/Source/Pipelines/Public/InterchangeGenericAnimationPipeline.h
--- a/Engine/Plugins/Interchange/Runtime/Source/Pipelines/Public/InterchangeGenericAnimationPipeline.h
+++ b/Engine/Plugins/Interchange/Runtime/Source/Pipelines/Public/InterchangeGenericAnimationPipeline.h
@@ -118,7 +118,7 @@ public:
 	UPROPERTY()
 	FString SourceAnimationName;
 
-	virtual void AdjustSettingsForContext(EInterchangePipelineContext ImportType, TObjectPtr<UObject> ReimportAsset) override;
+	virtual void AdjustSettingsForContext(EInterchangePipelineContext ImportType, TObjectPtr<UObject> ReimportAsset, const UInterchangeBaseNodeContainer* InBaseNodeContainer) override;
 	
 protected:
 
diff --git a/Engine/Plugins/Interchange/Runtime/Source/Pipelines/Public/InterchangeGenericAssetsPipeline.h b/Engine/Plugins/Interchange/Runtime/Source/Pipelines/Public/InterchangeGenericAssetsPipeline.h
--- a/Engine/Plugins/Interchange/Runtime/Source/Pipelines/Public/InterchangeGenericAssetsPipeline.h
+++ b/Engine/Plugins/Interchange/Runtime/Source/Pipelines/Public/InterchangeGenericAssetsPipeline.h
@@ -98,7 +98,7 @@ public:
 	virtual bool IsSettingsAreValid(TOptional<FText>& OutInvalidReason) const override;
 
 
-	virtual void AdjustSettingsForContext(EInterchangePipelineContext ImportType, TObjectPtr<UObject> ReimportAsset) override;
+	virtual void AdjustSettingsForContext(EInterchangePipelineContext ImportType, TObjectPtr<UObject> ReimportAsset, const UInterchangeBaseNodeContainer* InBaseNodeContainer) override;
 #if WITH_EDITOR
 	virtual void FilterPropertiesFromTranslatedData(UInterchangeBaseNodeContainer* InBaseNodeContainer) override;
 
diff --git a/Engine/Plugins/Interchange/Runtime/Source/Pipelines/Public/InterchangeGenericMaterialPipeline.h b/Engine/Plugins/Interchange/Runtime/Source/Pipelines/Public/InterchangeGenericMaterialPipeline.h
--- a/Engine/Plugins/Interchange/Runtime/Source/Pipelines/Public/InterchangeGenericMaterialPipeline.h
+++ b/Engine/Plugins/Interchange/Runtime/Source/Pipelines/Public/InterchangeGenericMaterialPipeline.h
@@ -99,7 +99,7 @@ public:
 	/** BEGIN UInterchangePipelineBase overrides */
 	virtual void PreDialogCleanup(const FName PipelineStackName) override;
 	virtual bool IsSettingsAreValid(TOptional<FText>& OutInvalidReason) const override;
-	virtual void AdjustSettingsForContext(EInterchangePipelineContext ImportType, TObjectPtr<UObject> ReimportAsset) override;
+	virtual void AdjustSettingsForContext(EInterchangePipelineContext ImportType, TObjectPtr<UObject> ReimportAsset, const UInterchangeBaseNodeContainer* InBaseNodeContainer) override;
 #if WITH_EDITOR
 	virtual void FilterPropertiesFromTranslatedData(UInterchangeBaseNodeContainer* InBaseNodeContainer) override;
 	virtual bool IsPropertyChangeNeedRefresh(const FPropertyChangedEvent& PropertyChangedEvent) override;
diff --git a/Engine/Plugins/Interchange/Runtime/Source/Pipelines/Public/InterchangeGenericMeshPipeline.h b/Engine/Plugins/Interchange/Runtime/Source/Pipelines/Public/InterchangeGenericMeshPipeline.h
--- a/Engine/Plugins/Interchange/Runtime/Source/Pipelines/Public/InterchangeGenericMeshPipeline.h
+++ b/Engine/Plugins/Interchange/Runtime/Source/Pipelines/Public/InterchangeGenericMeshPipeline.h
@@ -219,7 +219,7 @@ public:
 	UPROPERTY(EditAnywhere, BlueprintReadWrite, Category = "Skeletal Meshes", meta = (SubCategory = "Build"))
 	int32 BoneInfluenceLimit = 0;
 
-	virtual void AdjustSettingsForContext(EInterchangePipelineContext ImportType, TObjectPtr<UObject> ReimportAsset) override;
+	virtual void AdjustSettingsForContext(EInterchangePipelineContext ImportType, TObjectPtr<UObject> ReimportAsset, const UInterchangeBaseNodeContainer* InBaseNodeContainer) override;
 
 	virtual void PreDialogCleanup(const FName PipelineStackName) override;
 
diff --git a/Engine/Plugins/Interchange/Runtime/Source/Pipelines/Public/InterchangeGenericScenesPipeline.h b/Engine/Plugins/Interchange/Runtime/Source/Pipelines/Public/InterchangeGenericScenesPipeline.h
--- a/Engine/Plugins/Interchange/Runtime/Source/Pipelines/Public/InterchangeGenericScenesPipeline.h
+++ b/Engine/Plugins/Interchange/Runtime/Source/Pipelines/Public/InterchangeGenericScenesPipeline.h
@@ -46,7 +46,7 @@ public:
 	bool bDeleteMissingAssets = false;
 
 	/** BEGIN UInterchangePipelineBase overrides */
-	virtual void AdjustSettingsForContext(EInterchangePipelineContext ImportType, TObjectPtr<UObject> ReimportAsset) override;
+	virtual void AdjustSettingsForContext(EInterchangePipelineContext ImportType, TObjectPtr<UObject> ReimportAsset, const UInterchangeBaseNodeContainer* InBaseNodeContainer) override;
 
 protected:
 
diff --git a/Engine/Plugins/Interchange/Runtime/Source/Pipelines/Public/InterchangeGenericTexturePipeline.h b/Engine/Plugins/Interchange/Runtime/Source/Pipelines/Public/InterchangeGenericTexturePipeline.h
--- a/Engine/Plugins/Interchange/Runtime/Source/Pipelines/Public/InterchangeGenericTexturePipeline.h
+++ b/Engine/Plugins/Interchange/Runtime/Source/Pipelines/Public/InterchangeGenericTexturePipeline.h
@@ -66,7 +66,7 @@ public:
 	bool bAllowNonPowerOfTwo = false;
 
 public:
-	virtual void AdjustSettingsForContext(EInterchangePipelineContext ImportType, TObjectPtr<UObject> ReimportAsset) override;
+	virtual void AdjustSettingsForContext(EInterchangePipelineContext ImportType, TObjectPtr<UObject> ReimportAsset, const UInterchangeBaseNodeContainer* InBaseNodeContainer) override;
 	virtual void ExecutePipeline(UInterchangeBaseNodeContainer* InBaseNodeContainer, const TArray<UInterchangeSourceData*>& InSourceDatas, const FString& ContentBasePath) override;
 	virtual void ExecutePostFactoryPipeline(const UInterchangeBaseNodeContainer* BaseNodeContainer, const FString& NodeKey, UObject* CreatedAsset, bool bIsAReimport) override;
 
diff --git a/Engine/Plugins/Interchange/Runtime/Source/Pipelines/Public/InterchangeMaterialXPipeline.h b/Engine/Plugins/Interchange/Runtime/Source/Pipelines/Public/InterchangeMaterialXPipeline.h
--- a/Engine/Plugins/Interchange/Runtime/Source/Pipelines/Public/InterchangeMaterialXPipeline.h
+++ b/Engine/Plugins/Interchange/Runtime/Source/Pipelines/Public/InterchangeMaterialXPipeline.h
@@ -101,7 +101,7 @@ public:
 	TObjectPtr<UMaterialXPipelineSettings> MaterialXSettings;
 
 protected:
-	virtual void AdjustSettingsForContext(EInterchangePipelineContext ImportType, TObjectPtr<UObject> ReimportAsset) override;
+	virtual void AdjustSettingsForContext(EInterchangePipelineContext ImportType, TObjectPtr<UObject> ReimportAsset, const UInterchangeBaseNodeContainer* InBaseNodeContainer) override;
 	virtual void ExecutePipeline(UInterchangeBaseNodeContainer* BaseNodeContainer, const TArray<UInterchangeSourceData*>& SourceDatas, const FString& ContentBasePath) override;
 
 	virtual bool CanExecuteOnAnyThread(EInterchangePipelineTask PipelineTask) override
diff --git a/Engine/Plugins/Interchange/Runtime/Source/Pipelines/Public/InterchangeglTFPipeline.h b/Engine/Plugins/Interchange/Runtime/Source/Pipelines/Public/InterchangeglTFPipeline.h
--- a/Engine/Plugins/Interchange/Runtime/Source/Pipelines/Public/InterchangeglTFPipeline.h
+++ b/Engine/Plugins/Interchange/Runtime/Source/Pipelines/Public/InterchangeglTFPipeline.h
@@ -54,7 +54,7 @@ public:
 	TObjectPtr<UGLTFPipelineSettings> GLTFPipelineSettings;
 
 protected:
-	virtual void AdjustSettingsForContext(EInterchangePipelineContext ImportType, TObjectPtr<UObject> ReimportAsset) override;
+	virtual void AdjustSettingsForContext(EInterchangePipelineContext ImportType, TObjectPtr<UObject> ReimportAsset, const UInterchangeBaseNodeContainer* InBaseNodeContainer) override;
 	virtual void ExecutePipeline(UInterchangeBaseNodeContainer* BaseNodeContainer, const TArray<UInterchangeSourceData*>& SourceDatas, const FString& ContentBasePath) override;
 
 	virtual bool CanExecuteOnAnyThread(EInterchangePipelineTask PipelineTask) override
diff --git a/Engine/Source/Editor/UnrealEd/Private/SkinWeightsUtilities.cpp b/Engine/Source/Editor/UnrealEd/Private/SkinWeightsUtilities.cpp
--- a/Engine/Source/Editor/UnrealEd/Private/SkinWeightsUtilities.cpp
+++ b/Engine/Source/Editor/UnrealEd/Private/SkinWeightsUtilities.cpp
@@ -152,7 +152,7 @@ bool FSkinWeightsUtilities::ImportAlternateSkinWeight(USkeletalMesh* SkeletalMes
 			if (GenericPipeline)
 			{
 				GenericPipeline->ClearFlags(EObjectFlags::RF_Standalone | EObjectFlags::RF_Public);
-				GenericPipeline->AdjustSettingsForContext(bIsReimport ? EInterchangePipelineContext::AssetAlternateSkinningReimport : EInterchangePipelineContext::AssetAlternateSkinningImport, nullptr);
+				GenericPipeline->AdjustSettingsForContext(bIsReimport ? EInterchangePipelineContext::AssetAlternateSkinningReimport : EInterchangePipelineContext::AssetAlternateSkinningImport, nullptr, nullptr);
 				ImportAssetParameters.OverridePipelines.Add(GenericPipeline);
 			}
 		}
diff --git a/Engine/Source/Runtime/Interchange/Core/Private/InterchangePipelineBase.cpp b/Engine/Source/Runtime/Interchange/Core/Private/InterchangePipelineBase.cpp
--- a/Engine/Source/Runtime/Interchange/Core/Private/InterchangePipelineBase.cpp
+++ b/Engine/Source/Runtime/Interchange/Core/Private/InterchangePipelineBase.cpp
@@ -37,7 +37,7 @@ void UInterchangePipelineBase::SaveSettings(const FName PipelineStackName)
 	GConfig->Flush(0);
 }
 
-void UInterchangePipelineBase::AdjustSettingsForContext(EInterchangePipelineContext ReimportType, TObjectPtr<UObject> ReimportAsset)
+void UInterchangePipelineBase::AdjustSettingsForContext(EInterchangePipelineContext ReimportType, TObjectPtr<UObject> ReimportAsset, const UInterchangeBaseNodeContainer* BaseNodeContainer)
 {
 	CachePipelineContext = ReimportType;
 	CacheReimportObject = ReimportAsset;
@@ -61,7 +61,7 @@ void UInterchangePipelineBase::AdjustSettingsForContext(EInterchangePipelineCont
 void UInterchangePipelineBase::AdjustSettingsFromCache()
 {
 	PropertiesStates = CachePropertiesStates;
-	AdjustSettingsForContext(CachePipelineContext, CacheReimportObject.Get());
+	AdjustSettingsForContext(CachePipelineContext, CacheReimportObject.Get(), nullptr);
 }
 
 void UInterchangePipelineBase::TransferAdjustSettings(UInterchangePipelineBase* SourcePipeline)
diff --git a/Engine/Source/Runtime/Interchange/Core/Public/InterchangePipelineBase.h b/Engine/Source/Runtime/Interchange/Core/Public/InterchangePipelineBase.h
--- a/Engine/Source/Runtime/Interchange/Core/Public/InterchangePipelineBase.h
+++ b/Engine/Source/Runtime/Interchange/Core/Public/InterchangePipelineBase.h
@@ -299,7 +299,7 @@ public:
 	 * @Param ReimportType - Tells the pipeline what reimport type the user wants to achieve.
 	 * @Param ReimportAsset - This is an optional parameter which is set when reimporting an asset.
 	 */
-	INTERCHANGECORE_API virtual void AdjustSettingsForContext(EInterchangePipelineContext ReimportType, TObjectPtr<UObject> ReimportAsset);
+	INTERCHANGECORE_API virtual void AdjustSettingsForContext(EInterchangePipelineContext ReimportType, TObjectPtr<UObject> ReimportAsset, const UInterchangeBaseNodeContainer* BaseNodeContainer);
 	INTERCHANGECORE_API virtual void AdjustSettingsFromCache();
 
 	/** Transfer the source pipeline adjust settings to this pipeline. */
diff --git a/Engine/Source/Runtime/Interchange/Engine/Private/InterchangeManager.cpp b/Engine/Source/Runtime/Interchange/Engine/Private/InterchangeManager.cpp
--- a/Engine/Source/Runtime/Interchange/Engine/Private/InterchangeManager.cpp
+++ b/Engine/Source/Runtime/Interchange/Engine/Private/InterchangeManager.cpp
@@ -1115,7 +1115,7 @@ bool UInterchangeManager::ConvertImportData(const UObject* SourceImportData, FIm
 				{
 					if (UInterchangePipelineBase* GeneratedPipeline = UE::Interchange::GeneratePipelineInstance(PipelinePath))
 					{
-						GeneratedPipeline->AdjustSettingsForContext(EInterchangePipelineContext::AssetImport, nullptr);
+						GeneratedPipeline->AdjustSettingsForContext(EInterchangePipelineContext::AssetImport, nullptr, nullptr);
 						if (GeneratedPipeline->IsA(ConvertedPipelineClass))
 						{
 							//We found a match, so we will use the default pipeline stacks
@@ -1653,6 +1653,8 @@ UInterchangeManager::ImportInternal(const FString& ContentPath, const UInterchan
 	TSharedRef<UE::Interchange::FImportAsyncHelper, ESPMode::ThreadSafe> AsyncHelper = CreateAsyncHelper(TaskData, ImportAssetParameters);
 	AsyncHelper->UniqueId = UniqueId;
 
+	//We support only one source currently
+
 	//Create a duplicate of the source data, we need to be multithread safe so we copy it to control the life cycle. The async helper will hold it and delete it when the import task will be completed.
 	UInterchangeSourceData* DuplicateSourceData = Cast<UInterchangeSourceData>(StaticDuplicateObject(SourceData, GetTransientPackage()));
 	//Array of source data to build one graph per source
@@ -1678,10 +1680,26 @@ UInterchangeManager::ImportInternal(const FString& ContentPath, const UInterchan
 	for (int32 SourceDataIndex = 0; SourceDataIndex < AsyncHelper->SourceDatas.Num(); ++SourceDataIndex)
 	{
 		AsyncHelper->BaseNodeContainers.Add(TStrongObjectPtr<UInterchangeBaseNodeContainer>(NewObject<UInterchangeBaseNodeContainer>(GetTransientPackage(), NAME_None)));
-		check(AsyncHelper->BaseNodeContainers[SourceDataIndex].IsValid());
+		if (!ensureMsgf(AsyncHelper->BaseNodeContainers[SourceDataIndex].IsValid(), TEXT("Interchange node conatainer is invalid.")))
+		{
+			return EarlyExit();
+		}
 	}
 
+	if (!ensureMsgf(AsyncHelper->Translators.Num() == AsyncHelper->SourceDatas.Num(), TEXT("Each interchange source data should have a translator")))
+	{
+		return EarlyExit();
+	}
 	
+	//We need at least one translator
+	if(!ensure(AsyncHelper->Translators.IsValidIndex(SourceIndex)))
+	{
+		return EarlyExit();
+	}
+	
+	//Make sure we point on the first translator
+	AsyncTranslator = AsyncHelper->Translators[SourceIndex];
+
 	UInterchangePipelineConfigurationBase* RegisteredPipelineConfiguration = nullptr;
 
 	//In runtime we do not have any pipeline configurator
@@ -1698,21 +1716,23 @@ UInterchangeManager::ImportInternal(const FString& ContentPath, const UInterchan
 	}
 #endif
 	
-	auto AdjustPipelineSettingForContext = [bIsReimport, bImportScene, &TaskData](UInterchangePipelineBase* Pipeline)
-	{
-		EInterchangePipelineContext Context;
-		if (bIsReimport)
-		{
-			Context = bImportScene ? EInterchangePipelineContext::SceneReimport : EInterchangePipelineContext::AssetReimport;
-		}
-		else
+	UInterchangeBaseNodeContainer* BaseNodeContainer = AsyncHelper->BaseNodeContainers[SourceIndex].Get();
+
+	auto AdjustPipelineSettingForContext = [bIsReimport, bImportScene, &TaskData, BaseNodeContainer](UInterchangePipelineBase* Pipeline)
 		{
-			Context = bImportScene ? EInterchangePipelineContext::SceneImport : EInterchangePipelineContext::AssetImport;
-		}
+			EInterchangePipelineContext Context;
+			if (bIsReimport)
+			{
+				Context = bImportScene ? EInterchangePipelineContext::SceneReimport : EInterchangePipelineContext::AssetReimport;
+			}
+			else
+			{
+				Context = bImportScene ? EInterchangePipelineContext::SceneImport : EInterchangePipelineContext::AssetImport;
+			}
 
-		Pipeline->AdjustSettingsForContext(Context, TaskData.ReimportObject);
-		Pipeline->DestinationName = TaskData.DestinationName;
-	};
+			Pipeline->AdjustSettingsForContext(Context, TaskData.ReimportObject, BaseNodeContainer);
+			Pipeline->DestinationName = TaskData.DestinationName;
+		};
 
 	const bool bSkipImportDialog = AsyncTranslator ? ImportAllWithSamePipelines.Contains(AsyncTranslator->GetClass()) : false;
 	if (bSkipImportDialog)
@@ -1741,25 +1761,6 @@ UInterchangeManager::ImportInternal(const FString& ContentPath, const UInterchan
 		const bool bShowPipelineStacksConfigurationDialog = false;
 #endif
 
-		auto TranslateSourceFile = [&AsyncHelper]()
-		{
-			LLM_SCOPE_BYNAME(TEXT("Interchange"));
-			FScopedSlowTask Progress(2.f, NSLOCTEXT("InterchangeManager", "TranslatingSourceFile...", "Translating source file..."));
-			Progress.MakeDialog();
-			Progress.EnterProgressFrame(1.f);
-			//Translate the source
-			FGraphEventArray PipelinePrerequistes;
-			check(AsyncHelper->Translators.Num() == AsyncHelper->SourceDatas.Num());
-			for (int32 SourceDataIndex = 0; SourceDataIndex < AsyncHelper->SourceDatas.Num(); ++SourceDataIndex)
-			{
-				//Log the source we begin importing
-				UE_LOG(LogInterchangeEngine, Display, TEXT("Interchange start importing source [%s]"), *AsyncHelper->SourceDatas[SourceDataIndex]->ToDisplayString());
-				int32 TranslatorTaskIndex = AsyncHelper->TranslatorTasks.Add(TGraphTask<UE::Interchange::FTaskTranslator>::CreateTask(nullptr, ENamedThreads::GameThread).ConstructAndDispatchWhenReady(SourceDataIndex, AsyncHelper));
-				AsyncHelper->TranslatorTasks[TranslatorTaskIndex]->Wait();
-			}
-			Progress.EnterProgressFrame(1.f);
-		};
-
 		if (FEngineAnalytics::IsAvailable())
 		{
 			Attribs.Add(FAnalyticsEventAttribute(TEXT("ShowImportDialog"), bShowPipelineStacksConfigurationDialog));
@@ -1771,6 +1772,48 @@ UInterchangeManager::ImportInternal(const FString& ContentPath, const UInterchan
 		TArray<FInterchangeStackInfo> PipelineStacks;
 		TArray<UInterchangePipelineBase*> OutPipelines;
 
+		//If we need to display the dialog we want to translate the source file before showing it
+		if (RegisteredPipelineConfiguration && bShowPipelineStacksConfigurationDialog)
+		{
+			//Make sure we dont translate with a non thread translator that is already lock
+			bool* NonParallelTranslatorLock = nullptr;
+			bool bTranslatorIsThreadSafe = AsyncTranslator->IsThreadSafe();
+			if (!bTranslatorIsThreadSafe)
+			{
+				NonParallelTranslatorLock = NonParallelTranslatorLocks.Find(AsyncTranslator->GetClass());
+				//Add a NonParallelTranslatorLocks for this translator class
+				if (!NonParallelTranslatorLock)
+				{
+					//Create a boolean lock and initialize it to false
+					NonParallelTranslatorLock = &NonParallelTranslatorLocks.FindOrAdd(AsyncTranslator->GetClass());
+					*NonParallelTranslatorLock = false;
+				}
+			}
+
+			if (ensure(NonParallelTranslatorLock && !(*NonParallelTranslatorLock)))
+			{
+				FScopedSlowTask Progress(2.f, NSLOCTEXT("InterchangeManager", "TranslatingSourceFile...", "Translating source file..."));
+				Progress.MakeDialog();
+				Progress.EnterProgressFrame(1.f);
+				//Translate the source
+				FGraphEventArray PipelinePrerequistes;
+				for (int32 SourceDataIndex = 0; SourceDataIndex < AsyncHelper->SourceDatas.Num(); ++SourceDataIndex)
+				{
+					//Log the source we begin importing
+					UE_LOG(LogInterchangeEngine, Display, TEXT("Interchange start importing source [%s]"), *AsyncHelper->SourceDatas[SourceDataIndex]->ToDisplayString());
+					int32 TranslatorTaskIndex = AsyncHelper->TranslatorTasks.Add(TGraphTask<UE::Interchange::FTaskTranslator>::CreateTask(nullptr, ENamedThreads::GameThread).ConstructAndDispatchWhenReady(SourceDataIndex, AsyncHelper));
+					AsyncHelper->TranslatorTasks[TranslatorTaskIndex]->Wait();
+				}
+				Progress.EnterProgressFrame(1.f);
+			}
+			else
+			{
+				//If we already have a running non thread safe translator, we cannot translate here
+				UE_LOG(LogInterchangeEngine, Error, TEXT("Interchange cannot translate simultanously multiple file with a non thread safe translator. Source file [%s]"), *AsyncHelper->SourceDatas[SourceIndex]->ToDisplayString());
+			}
+		}
+		
+
 		//Fill the Stacks before showing the UI
 		if (bIsReimport)
 		{
@@ -1856,8 +1899,6 @@ UInterchangeManager::ImportInternal(const FString& ContentPath, const UInterchan
 		{
 			if (RegisteredPipelineConfiguration && bShowPipelineStacksConfigurationDialog && !bIsUnattended)
 			{
-				TranslateSourceFile();
-				UInterchangeBaseNodeContainer* BaseNodeContainer = AsyncHelper->BaseNodeContainers[SourceIndex].Get();
 				//Show the dialog, a plugin should have registered this dialog. We use a plugin to be able to use editor code when doing UI
 				EInterchangePipelineConfigurationDialogResult DialogResult = RegisteredPipelineConfiguration->ScriptedShowReimportPipelineConfigurationDialog(PipelineStacks, OutPipelines, DuplicateSourceData, AsyncTranslator, BaseNodeContainer, ImportAssetParameters.ReimportAsset);
 				if (DialogResult == EInterchangePipelineConfigurationDialogResult::Cancel)
@@ -1885,8 +1926,6 @@ UInterchangeManager::ImportInternal(const FString& ContentPath, const UInterchan
 		{
 			if (RegisteredPipelineConfiguration && bShowPipelineStacksConfigurationDialog)
 			{
-				TranslateSourceFile();
-				UInterchangeBaseNodeContainer* BaseNodeContainer = AsyncHelper->BaseNodeContainers[SourceIndex].Get();
 				//Show the dialog, a plugin should have register this dialog. We use a plugin to be able to use editor code when doing UI
 				EInterchangePipelineConfigurationDialogResult DialogResult = bImportScene
 					? RegisteredPipelineConfiguration->ScriptedShowScenePipelineConfigurationDialog(PipelineStacks, OutPipelines, DuplicateSourceData, AsyncTranslator, BaseNodeContainer)
diff --git a/Engine/Source/Runtime/Interchange/Engine/Private/InterchangeMeshUtilities.cpp b/Engine/Source/Runtime/Interchange/Engine/Private/InterchangeMeshUtilities.cpp
--- a/Engine/Source/Runtime/Interchange/Engine/Private/InterchangeMeshUtilities.cpp
+++ b/Engine/Source/Runtime/Interchange/Engine/Private/InterchangeMeshUtilities.cpp
@@ -189,7 +189,7 @@ TFuture<bool> UInterchangeMeshUtilities::InternalImportCustomLod(TSharedPtr<TPro
 			}
 			if (ensure(GeneratedPipeline))
 			{
-				GeneratedPipeline->AdjustSettingsForContext(ImportType, nullptr);
+				GeneratedPipeline->AdjustSettingsForContext(ImportType, nullptr, nullptr);
 				ImportAssetParameters.OverridePipelines.Add(GeneratedPipeline);
 			}
 		}
@@ -205,7 +205,7 @@ TFuture<bool> UInterchangeMeshUtilities::InternalImportCustomLod(TSharedPtr<TPro
 			if (UInterchangePipelineBase* GenericPipeline = NewObject<UInterchangePipelineBase>(GetTransientPackage(), GenericPipelineClass))
 			{
 				GenericPipeline->ClearFlags(EObjectFlags::RF_Standalone | EObjectFlags::RF_Public);
-				GenericPipeline->AdjustSettingsForContext(ImportType, nullptr);
+				GenericPipeline->AdjustSettingsForContext(ImportType, nullptr, nullptr);
 				ImportAssetParameters.OverridePipelines.Add(GenericPipeline);
 			}
 		}
