Interchange: make sure there is no temporary asset with keep flag and dirty package.

#rnx

diff --git a/Engine/Source/Runtime/Interchange/Engine/Private/InterchangeMeshUtilities.cpp b/Engine/Source/Runtime/Interchange/Engine/Private/InterchangeMeshUtilities.cpp
--- a/Engine/Source/Runtime/Interchange/Engine/Private/InterchangeMeshUtilities.cpp
+++ b/Engine/Source/Runtime/Interchange/Engine/Private/InterchangeMeshUtilities.cpp
@@ -2,6 +2,8 @@
 
 #include "InterchangeMeshUtilities.h"
 
+#include "AssetRegistry/AssetData.h"
+#include "AssetRegistry/AssetRegistryModule.h"
 #include "Async/Future.h"
 #include "Engine/StaticMesh.h"
 #include "Engine/SkeletalMesh.h"
@@ -211,12 +213,38 @@ TFuture<bool> UInterchangeMeshUtilities::InternalImportCustomLod(TSharedPtr<TPro
 		}
 	}
 
-	FString ImportAssetPath = TEXT("/Engine/TempEditor/Interchange/") + FGuid::NewGuid().ToString(EGuidFormats::Base36Encoded);
+ 	FString ImportAssetPath = TEXT("/Engine/TempEditor/Interchange/") + FGuid::NewGuid().ToString(EGuidFormats::Base36Encoded);
+	
+
+	//Create a functor to remove asset keep flag and package dirty flags
+	auto DeletePathAssets = [ImportAssetPath]()
+		{
+			FAssetRegistryModule& AssetRegistryModule = FModuleManager::LoadModuleChecked<FAssetRegistryModule>("AssetRegistry");
+
+			TArray<FAssetData> AssetsToDelete;
+			AssetRegistryModule.Get().GetAssetsByPath(FName(*ImportAssetPath), AssetsToDelete, true);
+			for (FAssetData AssetData : AssetsToDelete)
+			{
+				UObject* ObjToDelete = AssetData.GetAsset();
+				if (ObjToDelete)
+				{
+					//Avoid temporary package to be saved
+					UPackage* Package = ObjToDelete->GetOutermost();
+					Package->SetDirtyFlag(false);
+					//Avoid gc, remove keep flags
+					ObjToDelete->ClearFlags(RF_Standalone);
+					ObjToDelete->ClearInternalFlags(EInternalObjectFlags::Async);
+					//Make the object transient to prevent saving
+					ObjToDelete->SetFlags(RF_Transient);
+				}
+			}
+		};
+
 	UE::Interchange::FAssetImportResultRef AssetImportResult = InterchangeManager.ImportAssetAsync(ImportAssetPath, SourceData, ImportAssetParameters);
 	FString SourceDataFilename = SourceData->GetFilename();
 	if (SkeletalMesh)
 	{
-		AssetImportResult->OnDone([Promise, SkeletalMesh, LodIndex, SourceDataFilename](UE::Interchange::FImportResult& ImportResult)
+		AssetImportResult->OnDone([Promise, SkeletalMesh, LodIndex, SourceDataFilename, DeletePathAssets](UE::Interchange::FImportResult& ImportResult)
 			{
 				USkeletalMesh* SourceSkeletalMesh = Cast< USkeletalMesh >(ImportResult.GetFirstAssetOfClass(USkeletalMesh::StaticClass()));
 
@@ -225,8 +253,7 @@ TFuture<bool> UInterchangeMeshUtilities::InternalImportCustomLod(TSharedPtr<TPro
 					//Make sure we can modify the skeletalmesh properties
 					FSkinnedAssetAsyncBuildScope AsyncBuildScope(SkeletalMesh);
 					Promise->SetValue(FLODUtilities::SetCustomLOD(SkeletalMesh, SourceSkeletalMesh, LodIndex, SourceDataFilename));
-					SourceSkeletalMesh->ClearFlags(RF_Standalone);
-					SourceSkeletalMesh->ClearInternalFlags(EInternalObjectFlags::Async);
+					DeletePathAssets();
 				}
 				else
 				{
@@ -237,14 +264,13 @@ TFuture<bool> UInterchangeMeshUtilities::InternalImportCustomLod(TSharedPtr<TPro
 	}
 	else if (StaticMesh)
 	{
-		AssetImportResult->OnDone([Promise, StaticMesh, LodIndex, SourceDataFilename](UE::Interchange::FImportResult& ImportResult)
+		AssetImportResult->OnDone([Promise, StaticMesh, LodIndex, SourceDataFilename, DeletePathAssets](UE::Interchange::FImportResult& ImportResult)
 			{
 				UStaticMesh* SourceStaticMesh = Cast< UStaticMesh >(ImportResult.GetFirstAssetOfClass(UStaticMesh::StaticClass()));
 				if(SourceStaticMesh)
 				{
 					Promise->SetValue(StaticMesh->SetCustomLOD(SourceStaticMesh, LodIndex, SourceDataFilename));
-					SourceStaticMesh->ClearFlags(RF_Standalone);
-					SourceStaticMesh->ClearInternalFlags(EInternalObjectFlags::Async);
+					DeletePathAssets();
 				}
 				else
 				{
