Interchange: fix threadsafe translator logic in interchange manager

#rnx

diff --git a/Engine/Source/Runtime/Interchange/Engine/Private/InterchangeManager.cpp b/Engine/Source/Runtime/Interchange/Engine/Private/InterchangeManager.cpp
--- a/Engine/Source/Runtime/Interchange/Engine/Private/InterchangeManager.cpp
+++ b/Engine/Source/Runtime/Interchange/Engine/Private/InterchangeManager.cpp
@@ -1772,12 +1772,12 @@ UInterchangeManager::ImportInternal(const FString& ContentPath, const UInterchan
 		
 		const bool bIsUnattended = FApp::IsUnattended() || GIsAutomationTesting || ImportAssetParameters.bIsAutomated || bSkipImportDialog;
 #if WITH_EDITORONLY_DATA
-		const bool bShowPipelineStacksConfigurationDialog = !bIsUnattended
+		bool bShowPipelineStacksConfigurationDialog = !bIsUnattended
 															&& FInterchangeProjectSettingsUtils::ShouldShowPipelineStacksConfigurationDialog(bImportScene, *SourceData)
 															&& !bImportCanceled
 															&& !IsRunningCommandlet();
 #else
-		const bool bShowPipelineStacksConfigurationDialog = false;
+		bool bShowPipelineStacksConfigurationDialog = false;
 #endif
 
 		if (FEngineAnalytics::IsAvailable())
@@ -1790,10 +1790,11 @@ UInterchangeManager::ImportInternal(const FString& ContentPath, const UInterchan
 		const FName ReimportPipelineName = TEXT("ReimportPipeline");
 		TArray<FInterchangeStackInfo> PipelineStacks;
 		TArray<UInterchangePipelineBase*> OutPipelines;
-
+		
 		//If we need to display the dialog we want to translate the source file before showing it
 		if (RegisteredPipelineConfiguration && bShowPipelineStacksConfigurationDialog)
 		{
+			bool bCanTranslate = true;
 			//Make sure we dont translate with a non thread translator that is already lock
 			bool* NonParallelTranslatorLock = nullptr;
 			bool bTranslatorIsThreadSafe = AsyncTranslator->IsThreadSafe();
@@ -1807,9 +1808,11 @@ UInterchangeManager::ImportInternal(const FString& ContentPath, const UInterchan
 					NonParallelTranslatorLock = &NonParallelTranslatorLocks.FindOrAdd(AsyncTranslator->GetClass());
 					*NonParallelTranslatorLock = false;
 				}
+				//We ensure if we cannot translate
+				bCanTranslate = ensureMsgf(NonParallelTranslatorLock && !(*NonParallelTranslatorLock), TEXT("Interchange, non thread safe translator cannot have multiple instance at the same time."));
 			}
 
-			if (ensure(NonParallelTranslatorLock && !(*NonParallelTranslatorLock)))
+			if (bCanTranslate)
 			{
 				FScopedSlowTask Progress(2.f, NSLOCTEXT("InterchangeManager", "TranslatingSourceFile...", "Translating source file..."));
 				Progress.MakeDialog();
@@ -1829,6 +1832,7 @@ UInterchangeManager::ImportInternal(const FString& ContentPath, const UInterchan
 			{
 				//If we already have a running non thread safe translator, we cannot translate here
 				UE_LOG(LogInterchangeEngine, Error, TEXT("Interchange cannot translate simultanously multiple file with a non thread safe translator. Source file [%s]"), *AsyncHelper->SourceDatas[SourceIndex]->ToDisplayString());
+				bShowPipelineStacksConfigurationDialog = false;
 			}
 		}
 		
