Interchange: fix bad rename when importing animation only, a digit was added to the name.

#rnx

diff --git a/Engine/Source/Runtime/Interchange/Engine/Private/Tasks/InterchangeTaskParsing.cpp b/Engine/Source/Runtime/Interchange/Engine/Private/Tasks/InterchangeTaskParsing.cpp
--- a/Engine/Source/Runtime/Interchange/Engine/Private/Tasks/InterchangeTaskParsing.cpp
+++ b/Engine/Source/Runtime/Interchange/Engine/Private/Tasks/InterchangeTaskParsing.cpp
@@ -259,7 +259,7 @@ void UE::Interchange::FTaskParsing::DoTask(ENamedThreads::Type CurrentThread, co
 			//Make sure there is no duplicate name full path
 			uint32 NameIndex = 1;
 			FString NewName = AssetFullPath;
-			bool NameClash = true;
+			bool NameClash = FactoryNode->IsEnabled();
 			
 			while (NameClash)
 			{
@@ -286,7 +286,7 @@ void UE::Interchange::FTaskParsing::DoTask(ENamedThreads::Type CurrentThread, co
 				}
 			}
 
-			if (ensureMsgf(!CreatedTasksAssetNames.Contains(AssetFullPath),
+			if (!FactoryNode->IsEnabled() || ensureMsgf(!CreatedTasksAssetNames.Contains(AssetFullPath),
 				TEXT("Found multiple task data with the same asset name (%s). Only one will be executed."), *AssetFullPath))
 			{
 				LLM_SCOPE_BYNAME(TEXT("Interchange"));
@@ -307,7 +307,11 @@ void UE::Interchange::FTaskParsing::DoTask(ENamedThreads::Type CurrentThread, co
 					TGraphTask<FTaskImportObjectFinalize_GameThread>::CreateTask(&(FinalizeImportObjectTasksPrerequistes)).ConstructAndDispatchWhenReady(AsyncHelper->ContentBasePath, SourceIndex, WeakAsyncHelper, FactoryNode)
 				);
 
-				CreatedTasksAssetNames.Add(AssetFullPath);
+				//Only add the name if the factory node is enabled
+				if (FactoryNode->IsEnabled())
+				{
+					CreatedTasksAssetNames.Add(AssetFullPath);
+				}
 
 				return AsyncHelper->FinalizeImportObjectTasks[FinalizeCreateTaskIndex];
 			}
