[Interchange] add a boolean to control if we show the options dialog when doing re-import

#rnx

diff --git a/Engine/Source/Runtime/Interchange/Engine/Private/InterchangeManager.cpp b/Engine/Source/Runtime/Interchange/Engine/Private/InterchangeManager.cpp
--- a/Engine/Source/Runtime/Interchange/Engine/Private/InterchangeManager.cpp
+++ b/Engine/Source/Runtime/Interchange/Engine/Private/InterchangeManager.cpp
@@ -1773,7 +1773,7 @@ UInterchangeManager::ImportInternal(const FString& ContentPath, const UInterchan
 		const bool bIsUnattended = FApp::IsUnattended() || GIsAutomationTesting || ImportAssetParameters.bIsAutomated || bSkipImportDialog;
 #if WITH_EDITORONLY_DATA
 		bool bShowPipelineStacksConfigurationDialog = !bIsUnattended
-															&& FInterchangeProjectSettingsUtils::ShouldShowPipelineStacksConfigurationDialog(bImportScene, *SourceData)
+															&& FInterchangeProjectSettingsUtils::ShouldShowPipelineStacksConfigurationDialog(bImportScene, bIsReimport, *SourceData)
 															&& !bImportCanceled
 															&& !IsRunningCommandlet();
 #else
diff --git a/Engine/Source/Runtime/Interchange/Engine/Private/InterchangeProjectSettings.cpp b/Engine/Source/Runtime/Interchange/Engine/Private/InterchangeProjectSettings.cpp
--- a/Engine/Source/Runtime/Interchange/Engine/Private/InterchangeProjectSettings.cpp
+++ b/Engine/Source/Runtime/Interchange/Engine/Private/InterchangeProjectSettings.cpp
@@ -106,7 +106,7 @@ void FInterchangeProjectSettingsUtils::SetDefaultPipelineStackName(const bool bI
 	GetMutableDefault<UInterchangeProjectSettings>()->SaveConfig();
 }
 
-bool FInterchangeProjectSettingsUtils::ShouldShowPipelineStacksConfigurationDialog(const bool bIsSceneImport, const UInterchangeSourceData& SourceData)
+bool FInterchangeProjectSettingsUtils::ShouldShowPipelineStacksConfigurationDialog(const bool bIsSceneImport, const bool bReImport, const UInterchangeSourceData& SourceData)
 {
 	const FInterchangeImportSettings& ImportSettings = GetDefaultImportSettings(bIsSceneImport);
 
@@ -155,5 +155,7 @@ bool FInterchangeProjectSettingsUtils::ShouldShowPipelineStacksConfigurationDial
 		}
 	}
 
+	bShowImportDialog &= (bReImport ? ImportSettings.bShowImportDialogAtReimport : true);
+
 	return bShowImportDialog;
 }
\ No newline at end of file
diff --git a/Engine/Source/Runtime/Interchange/Engine/Public/InterchangeProjectSettings.h b/Engine/Source/Runtime/Interchange/Engine/Public/InterchangeProjectSettings.h
--- a/Engine/Source/Runtime/Interchange/Engine/Public/InterchangeProjectSettings.h
+++ b/Engine/Source/Runtime/Interchange/Engine/Public/InterchangeProjectSettings.h
@@ -56,9 +56,13 @@ struct FInterchangeImportSettings
 	UPROPERTY(EditAnywhere, Category = "Pipeline")
 	TSoftClassPtr <UInterchangePipelineConfigurationBase> ImportDialogClass;
 
-	/** If enabled, the import option dialog will show when interchange import or re-import.*/
+	/** If enabled, the import option dialog will show when interchange import.*/
 	UPROPERTY(EditAnywhere, Category = "Pipeline")
 	bool bShowImportDialog = true;
+
+	/** If enabled, the import option dialog will show when interchange re-import.*/
+	UPROPERTY(EditAnywhere, Category = "Pipeline", meta = (EditCondition = "bShowImportDialog", EditConditionHides))
+	bool bShowImportDialogAtReimport = false;
 };
 
 USTRUCT()
@@ -95,7 +99,7 @@ struct FInterchangeContentImportSettings : public FInterchangeImportSettings
 	TMap<EInterchangeTranslatorAssetType, FName> DefaultPipelineStackOverride;
 
 	/** This tell interchange if the import dialog should show or not when importing a particular type of asset.*/
-	UPROPERTY(EditAnywhere, Category = "Pipeline", Meta=(DisplayAfter="bShowImportDialog"))
+	UPROPERTY(EditAnywhere, Category = "Pipeline", Meta=(DisplayAfter="bShowImportDialogAtReimport"))
 	TMap<EInterchangeTranslatorAssetType, FInterchangeDialogOverride> ShowImportDialogOverride;
 };
 
@@ -154,5 +158,5 @@ public:
 	static INTERCHANGEENGINE_API FName GetDefaultPipelineStackName(const bool bIsSceneImport, const UInterchangeSourceData& SourceData);
 	static INTERCHANGEENGINE_API void SetDefaultPipelineStackName(const bool bIsSceneImport, const UInterchangeSourceData& SourceData, const FName StackName);
 
-	static INTERCHANGEENGINE_API bool ShouldShowPipelineStacksConfigurationDialog(const bool bIsSceneImport, const UInterchangeSourceData& SourceData);
+	static INTERCHANGEENGINE_API bool ShouldShowPipelineStacksConfigurationDialog(const bool bIsSceneImport, const bool bReImport, const UInterchangeSourceData& SourceData);
 };
