Interchange: fix blueprint starting interchange import.

#rnx

diff --git a/Engine/Source/Runtime/Interchange/Engine/Private/InterchangeManager.cpp b/Engine/Source/Runtime/Interchange/Engine/Private/InterchangeManager.cpp
--- a/Engine/Source/Runtime/Interchange/Engine/Private/InterchangeManager.cpp
+++ b/Engine/Source/Runtime/Interchange/Engine/Private/InterchangeManager.cpp
@@ -710,11 +710,18 @@ void UE::Interchange::FImportResult::SetDone()
 	}
 }
 
-void UE::Interchange::FImportResult::WaitUntilDone()
+void UE::Interchange::FImportResult::WaitUntilDone(bool bSynchronous /*= false*/)
 {
 	if (ImportStatus == EStatus::InProgress)
 	{
-		FTaskGraphInterface::Get().WaitUntilTaskCompletes(GraphEvent);
+		if (bSynchronous)
+		{
+			FTaskGraphInterface::Get().ProcessThreadUntilIdle(ENamedThreads::GameThread_Local);
+		}
+		else
+		{
+			FTaskGraphInterface::Get().WaitUntilTaskCompletes(GraphEvent);
+		}
 	}
 }
 
@@ -935,6 +942,14 @@ UInterchangeManager& UInterchangeManager::GetInterchangeManager()
 					InterchangeManager->StartQueuedTasks(InterchangeManager->bGCEndDelegateCancellAllTask);
 				}
 			});
+		InterchangeManager->GCPreDelegate = FCoreUObjectDelegates::GetPreGarbageCollectDelegate().AddLambda([]()
+			{
+				if (IsInterchangeImportEnabled() && InterchangeManager.IsValid())
+				{
+					const bool bCancel = !GIsEditor;
+					InterchangeManager->WaitUntilAllTasksDone(bCancel);
+				}
+			});
 
 		//We cancel any running task when we pre exit the engine
 		FCoreDelegates::OnEnginePreExit.AddLambda([]()
@@ -961,6 +976,11 @@ UInterchangeManager& UInterchangeManager::GetInterchangeManager()
 				FCoreUObjectDelegates::GetPostGarbageCollect().Remove(InterchangeManager->GCEndDelegate);
 				InterchangeManager->GCEndDelegate.Reset();
 			}
+			if(InterchangeManager->GCPreDelegate.IsValid())
+			{
+				FCoreUObjectDelegates::GetPreGarbageCollectDelegate().Remove(InterchangeManager->GCPreDelegate);
+				InterchangeManager->GCPreDelegate.Reset();
+			}
 			//Task should have been cancel in the Engine pre exit callback
 			ensure(InterchangeManager->ImportTasks.Num() == 0);
 			InterchangeManager->OnPreDestroyInterchangeManager.Broadcast();
@@ -1507,8 +1527,9 @@ void UInterchangeManager::StartQueuedTasks(bool bCancelAllTasks /*= false*/)
 
 bool UInterchangeManager::ImportAsset(const FString& ContentPath, const UInterchangeSourceData* SourceData, const FImportAssetParameters& ImportAssetParameters, TArray<UObject*>& OutImportedObjects)
 {
-	UE::Interchange::FAssetImportResultRef InterchangeResult = ImportAssetAsync(ContentPath, SourceData, ImportAssetParameters);
-	InterchangeResult->WaitUntilDone();
+	ImportAssetParameters.bRunSynchronous = true;
+	UE::Interchange::FAssetImportResultRef InterchangeResult = ImportInternal(ContentPath, SourceData, ImportAssetParameters, UE::Interchange::EImportType::ImportType_Asset).Get<0>();
+	InterchangeResult->WaitUntilDone(ImportAssetParameters.bRunSynchronous);
 	OutImportedObjects = InterchangeResult->GetImportedObjects();
 	return InterchangeResult->IsValid();
 }
@@ -1521,23 +1542,37 @@ bool UInterchangeManager::ImportAsset(const FString& ContentPath, const UInterch
 
 UE::Interchange::FAssetImportResultRef UInterchangeManager::ImportAssetAsync(const FString& ContentPath, const UInterchangeSourceData* SourceData, const FImportAssetParameters& ImportAssetParameters)
 {
+	ImportAssetParameters.bRunSynchronous = false;
 	return ImportInternal(ContentPath, SourceData, ImportAssetParameters, UE::Interchange::EImportType::ImportType_Asset).Get<0>();
 }
 
+bool UInterchangeManager::ScriptedImportAssetAsync(const FString& ContentPath, const UInterchangeSourceData* SourceData, const FImportAssetParameters& ImportAssetParameters)
+{
+	UE::Interchange::FAssetImportResultRef InterchangeResult = ImportAssetAsync(ContentPath, SourceData, ImportAssetParameters);
+	return InterchangeResult->IsValid();
+}
+
 bool UInterchangeManager::ImportScene(const FString& ContentPath, const UInterchangeSourceData* SourceData, const FImportAssetParameters& ImportAssetParameters)
 {
+	ImportAssetParameters.bRunSynchronous = true;
 	using namespace UE::Interchange;
-
 	TTuple<FAssetImportResultRef, FSceneImportResultRef> ImportResults = ImportInternal(ContentPath, SourceData, ImportAssetParameters, UE::Interchange::EImportType::ImportType_Scene);
-	
-	ImportResults.Get<0>()->WaitUntilDone();
-	ImportResults.Get<1>()->WaitUntilDone();
+	ImportResults.Get<0>()->WaitUntilDone(ImportAssetParameters.bRunSynchronous);
+	ImportResults.Get<1>()->WaitUntilDone(ImportAssetParameters.bRunSynchronous);
+	return ImportResults.Get<0>()->IsValid() && ImportResults.Get<1>()->IsValid();
+}
+
+bool UInterchangeManager::ScriptedImportSceneAsync(const FString& ContentPath, const UInterchangeSourceData* SourceData, const FImportAssetParameters& ImportAssetParameters)
+{
+	using namespace UE::Interchange;
+	TTuple<FAssetImportResultRef, FSceneImportResultRef> ImportResults = ImportSceneAsync(ContentPath, SourceData, ImportAssetParameters);
 	return ImportResults.Get<0>()->IsValid() && ImportResults.Get<1>()->IsValid();
 }
 
 TTuple<UE::Interchange::FAssetImportResultRef, UE::Interchange::FSceneImportResultRef>
 UInterchangeManager::ImportSceneAsync(const FString& ContentPath, const UInterchangeSourceData* SourceData, const FImportAssetParameters& ImportAssetParameters)
 {
+	ImportAssetParameters.bRunSynchronous = false;
 	return ImportInternal(ContentPath, SourceData, ImportAssetParameters, UE::Interchange::EImportType::ImportType_Scene);
 }
 
@@ -1671,6 +1706,7 @@ UInterchangeManager::ImportInternal(const FString& ContentPath, const UInterchan
 
 	TSharedRef<UE::Interchange::FImportAsyncHelper, ESPMode::ThreadSafe> AsyncHelper = CreateAsyncHelper(TaskData, ImportAssetParameters);
 	AsyncHelper->UniqueId = UniqueId;
+	AsyncHelper->bRunSynchronous = ImportAssetParameters.bRunSynchronous;
 
 	//We support only one source currently
 
@@ -1823,7 +1859,7 @@ UInterchangeManager::ImportInternal(const FString& ContentPath, const UInterchan
 				{
 					//Log the source we begin importing
 					UE_LOG(LogInterchangeEngine, Display, TEXT("Interchange start importing source [%s]"), *AsyncHelper->SourceDatas[SourceDataIndex]->ToDisplayString());
-					int32 TranslatorTaskIndex = AsyncHelper->TranslatorTasks.Add(TGraphTask<UE::Interchange::FTaskTranslator>::CreateTask(nullptr, ENamedThreads::GameThread).ConstructAndDispatchWhenReady(SourceDataIndex, AsyncHelper));
+					int32 TranslatorTaskIndex = AsyncHelper->TranslatorTasks.Add(TGraphTask<UE::Interchange::FTaskTranslator>::CreateTask(nullptr, ENamedThreads::GameThread_Local).ConstructAndDispatchWhenReady(SourceDataIndex, AsyncHelper));
 					AsyncHelper->TranslatorTasks[TranslatorTaskIndex]->Wait();
 				}
 				Progress.EnterProgressFrame(1.f);
diff --git a/Engine/Source/Runtime/Interchange/Engine/Private/Tasks/InterchangeTaskCompletion.h b/Engine/Source/Runtime/Interchange/Engine/Private/Tasks/InterchangeTaskCompletion.h
--- a/Engine/Source/Runtime/Interchange/Engine/Private/Tasks/InterchangeTaskCompletion.h
+++ b/Engine/Source/Runtime/Interchange/Engine/Private/Tasks/InterchangeTaskCompletion.h
@@ -24,8 +24,13 @@ namespace UE
 			{
 			}
 
-			static FORCEINLINE ENamedThreads::Type GetDesiredThread()
+			FORCEINLINE ENamedThreads::Type GetDesiredThread()
 			{
+				TSharedPtr<FImportAsyncHelper, ESPMode::ThreadSafe> AsyncHelper = WeakAsyncHelper.Pin();
+				if (AsyncHelper.IsValid() && AsyncHelper->bRunSynchronous)
+				{
+					return ENamedThreads::GameThread_Local;
+				}
 				return ENamedThreads::GameThread;
 			}
 			static FORCEINLINE ESubsequentsMode::Type GetSubsequentsMode()
@@ -53,8 +58,13 @@ namespace UE
 			{
 			}
 
-			static FORCEINLINE ENamedThreads::Type GetDesiredThread()
+			FORCEINLINE ENamedThreads::Type GetDesiredThread()
 			{
+				TSharedPtr<FImportAsyncHelper, ESPMode::ThreadSafe> AsyncHelper = WeakAsyncHelper.Pin();
+				if (AsyncHelper.IsValid() && AsyncHelper->bRunSynchronous)
+				{
+					return ENamedThreads::GameThread_Local;
+				}
 				return ENamedThreads::GameThread;
 			}
 			static FORCEINLINE ESubsequentsMode::Type GetSubsequentsMode()
diff --git a/Engine/Source/Runtime/Interchange/Engine/Private/Tasks/InterchangeTaskCreateSceneObjects.cpp b/Engine/Source/Runtime/Interchange/Engine/Private/Tasks/InterchangeTaskCreateSceneObjects.cpp
--- a/Engine/Source/Runtime/Interchange/Engine/Private/Tasks/InterchangeTaskCreateSceneObjects.cpp
+++ b/Engine/Source/Runtime/Interchange/Engine/Private/Tasks/InterchangeTaskCreateSceneObjects.cpp
@@ -33,6 +33,18 @@ UE::Interchange::FTaskCreateSceneObjects::FTaskCreateSceneObjects(const FString&
 	check(FactoryClass);
 }
 
+ENamedThreads::Type UE::Interchange::FTaskCreateSceneObjects::GetDesiredThread()
+{
+	// We are creating the factories in this task so it must execute on the GameThread.
+	// Also, there are no "CreatePackage Task" equivalent for scene objects right now, so the factories must create those on the game thread.
+	TSharedPtr<FImportAsyncHelper, ESPMode::ThreadSafe> AsyncHelper = WeakAsyncHelper.Pin();
+	if (AsyncHelper.IsValid() && AsyncHelper->bRunSynchronous)
+	{
+		return ENamedThreads::GameThread_Local;
+	}
+	return ENamedThreads::GameThread;
+}
+
 void UE::Interchange::FTaskCreateSceneObjects::DoTask(ENamedThreads::Type CurrentThread, const FGraphEventRef& MyCompletionGraphEvent)
 {
 	TRACE_CPUPROFILER_EVENT_SCOPE(UE::Interchange::FTaskCreateSceneObjects::DoTask)
diff --git a/Engine/Source/Runtime/Interchange/Engine/Private/Tasks/InterchangeTaskCreateSceneObjects.h b/Engine/Source/Runtime/Interchange/Engine/Private/Tasks/InterchangeTaskCreateSceneObjects.h
--- a/Engine/Source/Runtime/Interchange/Engine/Private/Tasks/InterchangeTaskCreateSceneObjects.h
+++ b/Engine/Source/Runtime/Interchange/Engine/Private/Tasks/InterchangeTaskCreateSceneObjects.h
@@ -29,12 +29,7 @@ namespace UE
 		public:
 			explicit FTaskCreateSceneObjects(const FString& InPackageBasePath, const int32 InSourceIndex, TWeakPtr<FImportAsyncHelper> InAsyncHelper, TArrayView<UInterchangeFactoryBaseNode*> InNodes, const UClass* InFactoryClass);
 
-			ENamedThreads::Type GetDesiredThread()
-			{
-				// We are creating the factories in this task so it must execute on the GameThread.
-				// Also, there are no "CreatePackage Task" equivalent for scene objects right now, so the factories must create those on the game thread.
-				return ENamedThreads::GameThread;
-			}
+			ENamedThreads::Type GetDesiredThread();
 
 			static ESubsequentsMode::Type GetSubsequentsMode()
 			{
diff --git a/Engine/Source/Runtime/Interchange/Engine/Private/Tasks/InterchangeTaskImportObject.h b/Engine/Source/Runtime/Interchange/Engine/Private/Tasks/InterchangeTaskImportObject.h
--- a/Engine/Source/Runtime/Interchange/Engine/Private/Tasks/InterchangeTaskImportObject.h
+++ b/Engine/Source/Runtime/Interchange/Engine/Private/Tasks/InterchangeTaskImportObject.h
@@ -38,6 +38,11 @@ namespace UE::Interchange
 
 		ENamedThreads::Type GetDesiredThread() const
 		{
+			TSharedPtr<FImportAsyncHelper, ESPMode::ThreadSafe> AsyncHelper = WeakAsyncHelper.Pin();
+			if (AsyncHelper.IsValid() && AsyncHelper->bRunSynchronous)
+			{
+				return ENamedThreads::GameThread_Local;
+			}
 			return ENamedThreads::GameThread;
 		}
 
@@ -74,6 +79,11 @@ namespace UE::Interchange
 
 		ENamedThreads::Type GetDesiredThread() const
 		{
+			TSharedPtr<FImportAsyncHelper, ESPMode::ThreadSafe> AsyncHelper = WeakAsyncHelper.Pin();
+			if (AsyncHelper.IsValid() && AsyncHelper->bRunSynchronous)
+			{
+				return ENamedThreads::GameThread_Local;
+			}
 			return ENamedThreads::AnyBackgroundThreadNormalTask;
 		}
 
@@ -114,6 +124,11 @@ namespace UE::Interchange
 
 		ENamedThreads::Type GetDesiredThread() const
 		{
+			TSharedPtr<FImportAsyncHelper, ESPMode::ThreadSafe> AsyncHelper = WeakAsyncHelper.Pin();
+			if (AsyncHelper.IsValid() && AsyncHelper->bRunSynchronous)
+			{
+				return ENamedThreads::GameThread_Local;
+			}
 			return ENamedThreads::GameThread;
 		}
 
diff --git a/Engine/Source/Runtime/Interchange/Engine/Private/Tasks/InterchangeTaskParsing.h b/Engine/Source/Runtime/Interchange/Engine/Private/Tasks/InterchangeTaskParsing.h
--- a/Engine/Source/Runtime/Interchange/Engine/Private/Tasks/InterchangeTaskParsing.h
+++ b/Engine/Source/Runtime/Interchange/Engine/Private/Tasks/InterchangeTaskParsing.h
@@ -26,8 +26,13 @@ namespace UE
 				check(InterchangeManager);
 			}
 
-			static FORCEINLINE ENamedThreads::Type GetDesiredThread()
+			FORCEINLINE ENamedThreads::Type GetDesiredThread()
 			{
+				TSharedPtr<FImportAsyncHelper, ESPMode::ThreadSafe> AsyncHelper = WeakAsyncHelper.Pin();
+				if (AsyncHelper.IsValid() && AsyncHelper->bRunSynchronous)
+				{
+					return ENamedThreads::GameThread_Local;
+				}
 				return ENamedThreads::AnyBackgroundThreadNormalTask;
 			}
 			static FORCEINLINE ESubsequentsMode::Type GetSubsequentsMode()
diff --git a/Engine/Source/Runtime/Interchange/Engine/Private/Tasks/InterchangeTaskPipeline.cpp b/Engine/Source/Runtime/Interchange/Engine/Private/Tasks/InterchangeTaskPipeline.cpp
--- a/Engine/Source/Runtime/Interchange/Engine/Private/Tasks/InterchangeTaskPipeline.cpp
+++ b/Engine/Source/Runtime/Interchange/Engine/Private/Tasks/InterchangeTaskPipeline.cpp
@@ -100,9 +100,10 @@ void UE::Interchange::FTaskWaitAssetCompilation::DoTask(ENamedThreads::Type Curr
 		{
 			LLM_SCOPE_BYNAME(TEXT("Interchange"));
 			//Compilation status cannot be ask in async thread, query the compile status on the main thread with a small fast function
-			//This ensure we dont stall the main thread until all assets are compile.
+			//This ensure we don't stall the main thread until all assets are compile.
 			bool bCompilationFinish = false;
-			Async(EAsyncExecution::TaskGraphMainThread, [&bCompilationFinish, &ImportedObjects]()
+
+			auto VerifyAssetCompiling = [&bCompilationFinish, &ImportedObjects]()
 				{
 					TRACE_CPUPROFILER_EVENT_SCOPE(UE::Interchange::FTaskWaitAssetCompilation::DoTask::IsCompilingLambda_GameThread);
 					LLM_SCOPE_BYNAME(TEXT("Interchange"));
@@ -130,7 +131,19 @@ void UE::Interchange::FTaskWaitAssetCompilation::DoTask(ENamedThreads::Type Curr
 							}
 						}
 					}
-				}).Wait();
+				};
+
+			if (IsInGameThread())
+			{
+				VerifyAssetCompiling();
+			}
+			else
+			{
+				Async(EAsyncExecution::TaskGraphMainThread, [&VerifyAssetCompiling]()
+					{
+						VerifyAssetCompiling();
+					}).Wait();
+			}
 			return bCompilationFinish;
 		}, 0.05f);
 #endif //WITH_EDITOR
diff --git a/Engine/Source/Runtime/Interchange/Engine/Private/Tasks/InterchangeTaskPipeline.h b/Engine/Source/Runtime/Interchange/Engine/Private/Tasks/InterchangeTaskPipeline.h
--- a/Engine/Source/Runtime/Interchange/Engine/Private/Tasks/InterchangeTaskPipeline.h
+++ b/Engine/Source/Runtime/Interchange/Engine/Private/Tasks/InterchangeTaskPipeline.h
@@ -28,6 +28,12 @@ namespace UE
 
 			ENamedThreads::Type GetDesiredThread()
 			{
+				TSharedPtr<FImportAsyncHelper, ESPMode::ThreadSafe> AsyncHelper = WeakAsyncHelper.Pin();
+				if (AsyncHelper.IsValid() && AsyncHelper->bRunSynchronous)
+				{
+					return ENamedThreads::GameThread_Local;
+				}
+
 				if (!ensure(PipelineBase.IsValid()))
 				{
 					return ENamedThreads::GameThread;
@@ -71,6 +77,12 @@ namespace UE
 
 			ENamedThreads::Type GetDesiredThread()
 			{
+				TSharedPtr<FImportAsyncHelper, ESPMode::ThreadSafe> AsyncHelper = WeakAsyncHelper.Pin();
+				if (AsyncHelper.IsValid() && AsyncHelper->bRunSynchronous)
+				{
+					return ENamedThreads::GameThread_Local;
+				}
+
 				return ENamedThreads::AnyBackgroundThreadNormalTask;
 			}
 
@@ -102,6 +114,12 @@ namespace UE
 
 			ENamedThreads::Type GetDesiredThread()
 			{
+				TSharedPtr<FImportAsyncHelper, ESPMode::ThreadSafe> AsyncHelper = WeakAsyncHelper.Pin();
+				if (AsyncHelper.IsValid() && AsyncHelper->bRunSynchronous)
+				{
+					return ENamedThreads::GameThread_Local;
+				}
+
 				return ENamedThreads::GameThread;
 			}
 
diff --git a/Engine/Source/Runtime/Interchange/Engine/Private/Tasks/InterchangeTaskTranslator.h b/Engine/Source/Runtime/Interchange/Engine/Private/Tasks/InterchangeTaskTranslator.h
--- a/Engine/Source/Runtime/Interchange/Engine/Private/Tasks/InterchangeTaskTranslator.h
+++ b/Engine/Source/Runtime/Interchange/Engine/Private/Tasks/InterchangeTaskTranslator.h
@@ -26,8 +26,14 @@ namespace UE
 			{
 			}
 
-			static FORCEINLINE ENamedThreads::Type GetDesiredThread()
+			FORCEINLINE ENamedThreads::Type GetDesiredThread()
 			{
+				TSharedPtr<FImportAsyncHelper, ESPMode::ThreadSafe> AsyncHelper = WeakAsyncHelper.Pin();
+				if (AsyncHelper.IsValid() && AsyncHelper->bRunSynchronous)
+				{
+					return ENamedThreads::GameThread_Local;
+				}
+
 				return ENamedThreads::AnyBackgroundThreadNormalTask;
 			}
 			static FORCEINLINE ESubsequentsMode::Type GetSubsequentsMode()
diff --git a/Engine/Source/Runtime/Interchange/Engine/Public/InterchangeManager.h b/Engine/Source/Runtime/Interchange/Engine/Public/InterchangeManager.h
--- a/Engine/Source/Runtime/Interchange/Engine/Public/InterchangeManager.h
+++ b/Engine/Source/Runtime/Interchange/Engine/Public/InterchangeManager.h
@@ -146,7 +146,7 @@ namespace UE
 
 			INTERCHANGEENGINE_API void SetInProgress();
 			INTERCHANGEENGINE_API void SetDone();
-			INTERCHANGEENGINE_API void WaitUntilDone();
+			INTERCHANGEENGINE_API void WaitUntilDone(bool bSynchronous = false);
 
 			// Assets are only made available once they have been completely imported (passed through the entire import pipeline).
 			// While the status isn't EStatus::Done, the list can grow between subsequent calls.
@@ -211,6 +211,8 @@ namespace UE
 				return TEXT("UE::Interchange::FImportAsyncHelper");
 			}
 
+			bool bRunSynchronous = false;
+
 			/** Unique ID for this async helper. */
 			int32 UniqueId;
 
@@ -405,6 +407,10 @@ struct FImportAssetParameters
 	UPROPERTY(EditAnywhere, BlueprintReadWrite, Category = "Interchange|ImportAsset", meta=(PinHiddenByDefault))
 	FOnImportDoneDynamic OnSceneImportDone;
 	FOnImportDoneNative OnSceneImportDoneNative;
+
+	// Tell Interchange that import must run synchronously on the game thread.
+	//This is an internal property set by the Import API
+	mutable bool bRunSynchronous;
 };
 
 UCLASS(Transient, BlueprintType, CustomConstructor, MinimalAPI)
@@ -551,32 +557,89 @@ public:
 	INTERCHANGEENGINE_API bool CanReimport(const UObject* Object, TArray<FString>& OutFilenames) const;
 
 	/**
-	 * Call this to start an asset import process. The caller must specify the source data.
+	 * Call this to start a synchronous asset import process.
 	 * This process can import many different assets into the game content.
 	 *
 	 * @Param ContentPath - The path where the imported assets will be created.
 	 * @Param SourceData - The source data input to translate.
 	 * @param ImportAssetParameters - All parameters that need to be passed to the import asset function.
 	 * @return true if the import succeeds, or false otherwise.
+	 * 
+	 * @Note - In blueprint depending on the event you use to start the import its possible to have a deadlock, use the async function if its what you are experimenting
 	 */
 	UFUNCTION(BlueprintCallable, Category = "Interchange | Import Manager")
 	INTERCHANGEENGINE_API bool ImportAsset(const FString& ContentPath, const UInterchangeSourceData* SourceData, const FImportAssetParameters& ImportAssetParameters, TArray<UObject*>& OutImportedObjects);
+
+	/**
+	 * Call this to start a synchronous asset import process.
+	 * This process can import many different assets into the game content.
+	 *
+	 * @Param ContentPath - The path where the imported assets will be created.
+	 * @Param SourceData - The source data input to translate.
+	 * @param ImportAssetParameters - All parameters that need to be passed to the import asset function.
+	 * @return true if the import succeeds, or false otherwise.
+	 *
+	 */
 	INTERCHANGEENGINE_API bool ImportAsset(const FString& ContentPath, const UInterchangeSourceData* SourceData, const FImportAssetParameters& ImportAssetParameters);
 
+	/**
+	 * Call this to start an asynchronous asset import process.
+	 * This process can import many different assets into the game content.
+	 *
+	 * @Param ContentPath - The path where the imported assets will be created.
+	 * @Param SourceData - The source data input to translate.
+	 * @param ImportAssetParameters - All parameters that need to be passed to the import asset function.
+	 * @return return an import result which can be use to know when the asynchronous import is terminate.
+	 */
 	INTERCHANGEENGINE_API UE::Interchange::FAssetImportResultRef ImportAssetAsync(const FString& ContentPath, const UInterchangeSourceData* SourceData, const FImportAssetParameters& ImportAssetParameters);
+	
+	/**
+	 * Call this from blueprint or python to start an asynchronous asset import process.
+	 * This process can import many different assets into the game content.
+	 *
+	 * @Param ContentPath - The path where the imported assets will be created.
+	 * @Param SourceData - The source data input to translate.
+	 * @param ImportAssetParameters - All parameters that need to be passed to the import asset function.
+	 * @return true if the import was started, or false otherwise.
+	 */
+	UFUNCTION(BlueprintCallable, Category = "Interchange | Import Manager")
+	INTERCHANGEENGINE_API bool ScriptedImportAssetAsync(const FString& ContentPath, const UInterchangeSourceData* SourceData, const FImportAssetParameters& ImportAssetParameters);
 
 	/**
-	 * Call this to start a scene import process. The caller must specify the source data.
-	 * This process can import many different assets and their transforms (USceneComponent), store the result in a Blueprint, and add the Blueprint to the level.
+	 * Call this to start a synchronous scene import process.
+	 * This process can import many different assets and their transforms (USceneComponent).
 	 *
 	 * @Param ContentPath - The path where the imported assets will be created.
 	 * @Param SourceData - The source data input to translate. This object will be duplicated to allow thread-safe operations.
 	 * @param ImportAssetParameters - All parameters that need to be passed to the import asset function.
 	 * @return true if the import succeeds, or false otherwise.
+	 * 
+	 * @Note - In blueprint depending on the event you use to start the import its possible to have a deadlock, use the async function if its what you are experimenting
 	 */
 	UFUNCTION(BlueprintCallable, Category = "Interchange | Import Manager")
 	INTERCHANGEENGINE_API bool ImportScene(const FString& ContentPath, const UInterchangeSourceData* SourceData, const FImportAssetParameters& ImportAssetParameters);
 
+	/**
+	 * Call this to start a asynchronous scene import process.
+	 * This process can import many different assets and their transforms (USceneComponent).
+	 *
+	 * @Param ContentPath - The path where the imported assets will be created.
+	 * @Param SourceData - The source data input to translate. This object will be duplicated to allow thread-safe operations.
+	 * @param ImportAssetParameters - All parameters that need to be passed to the import asset function.
+	 * @return true if the import was started, or false otherwise.
+	 */
+	UFUNCTION(BlueprintCallable, Category = "Interchange | Import Manager")
+	INTERCHANGEENGINE_API bool ScriptedImportSceneAsync(const FString& ContentPath, const UInterchangeSourceData* SourceData, const FImportAssetParameters& ImportAssetParameters);
+
+	/**
+	 * Call this to start a asynchronous scene import process.
+	 * This process can import many different assets and their transforms (USceneComponent), store the result in a Blueprint, and add the Blueprint to the level.
+	 *
+	 * @Param ContentPath - The path where the imported assets will be created.
+	 * @Param SourceData - The source data input to translate. This object will be duplicated to allow thread-safe operations.
+	 * @param ImportAssetParameters - All parameters that need to be passed to the import asset function.
+	 * @return return a pair of import result which can be use to know when the asynchronous import is terminate.
+	 */
 	INTERCHANGEENGINE_API TTuple<UE::Interchange::FAssetImportResultRef, UE::Interchange::FSceneImportResultRef>
 	ImportSceneAsync(const FString& ContentPath, const UInterchangeSourceData* SourceData, const FImportAssetParameters& ImportAssetParameters);
 
@@ -782,6 +845,7 @@ private:
 
 	//We want to avoid starting an import task during a GC.
 	FDelegateHandle GCEndDelegate;
+	FDelegateHandle GCPreDelegate;
 	bool bGCEndDelegateCancellAllTask = false;
 
 	friend class UE::Interchange::FScopedTranslator;
