Interchange: Move show dialog at re-import in EditorPerProjectUserSettings config file

#rnx

diff --git a/Engine/Source/Runtime/Interchange/Engine/Private/InterchangeProjectSettings.cpp b/Engine/Source/Runtime/Interchange/Engine/Private/InterchangeProjectSettings.cpp
--- a/Engine/Source/Runtime/Interchange/Engine/Private/InterchangeProjectSettings.cpp
+++ b/Engine/Source/Runtime/Interchange/Engine/Private/InterchangeProjectSettings.cpp
@@ -155,7 +155,10 @@ bool FInterchangeProjectSettingsUtils::ShouldShowPipelineStacksConfigurationDial
 		}
 	}
 
-	bShowImportDialog &= (bReImport ? ImportSettings.bShowImportDialogAtReimport : true);
+	if (const UInterchangeEditorSettings* InterchangeEditorSettings = GetDefault<UInterchangeEditorSettings>())
+	{
+		bShowImportDialog &= (bReImport ? InterchangeEditorSettings->bShowImportDialogAtReimport : true);
+	}
 
 	return bShowImportDialog;
 }
\ No newline at end of file
diff --git a/Engine/Source/Runtime/Interchange/Engine/Public/InterchangeProjectSettings.h b/Engine/Source/Runtime/Interchange/Engine/Public/InterchangeProjectSettings.h
--- a/Engine/Source/Runtime/Interchange/Engine/Public/InterchangeProjectSettings.h
+++ b/Engine/Source/Runtime/Interchange/Engine/Public/InterchangeProjectSettings.h
@@ -59,10 +59,6 @@ struct FInterchangeImportSettings
 	/** If enabled, the import option dialog will show when interchange import.*/
 	UPROPERTY(EditAnywhere, Category = "Pipeline")
 	bool bShowImportDialog = true;
-
-	/** If enabled, the import option dialog will show when interchange re-import.*/
-	UPROPERTY(EditAnywhere, Category = "Pipeline", meta = (EditCondition = "bShowImportDialog", EditConditionHides))
-	bool bShowImportDialogAtReimport = false;
 };
 
 USTRUCT()
@@ -160,3 +156,14 @@ public:
 
 	static INTERCHANGEENGINE_API bool ShouldShowPipelineStacksConfigurationDialog(const bool bIsSceneImport, const bool bReImport, const UInterchangeSourceData& SourceData);
 };
+
+UCLASS(config = EditorPerProjectUserSettings, meta = (DisplayName = Interchange), MinimalAPI)
+class UInterchangeEditorSettings : public UDeveloperSettings
+{
+	GENERATED_BODY()
+
+public:
+	/** If enabled, the import option dialog will show when interchange re-import.*/
+	UPROPERTY(EditAnywhere, config, Category = "Show Dialog")
+	bool bShowImportDialogAtReimport = false;
+};
