Interchange: Make interchange python script friendly.
- Make sure the graph inspector do not show when GIsRunningUnattendedScript is true
- Remove unecessary LOG or downgrade the level to verbose
- Allow scripted with override pipeline to show the UI to the user if bAutomated is false
- Pipeline API can tell if a pipeline is for a re-import or an override.
- Fix DestinationName and ContentImportPath to support pipeline duplication
- Fix the converter to convert properly between a animation only and skeletalmesh import

#rnx

diff --git a/Engine/Plugins/Interchange/Editor/Source/InterchangeEditor/Private/InterchangeFbxAssetImportDataConverter.cpp b/Engine/Plugins/Interchange/Editor/Source/InterchangeEditor/Private/InterchangeFbxAssetImportDataConverter.cpp
--- a/Engine/Plugins/Interchange/Editor/Source/InterchangeEditor/Private/InterchangeFbxAssetImportDataConverter.cpp
+++ b/Engine/Plugins/Interchange/Editor/Source/InterchangeEditor/Private/InterchangeFbxAssetImportDataConverter.cpp
@@ -804,7 +804,7 @@ namespace UE::Interchange::Private
 
 		//Discover if we must import something in particular
 		if (FbxImportUI->MeshTypeToImport == EFBXImportType::FBXIT_SkeletalMesh
-			|| FbxImportUI->bImportAsSkeletal)
+			|| (FbxImportUI->bImportAsSkeletal && FbxImportUI->bImportMesh))
 		{
 			GenericAssetPipeline->MeshPipeline->bImportSkeletalMeshes = true;
 			GenericAssetPipeline->MeshPipeline->CommonMeshesProperties->ForceAllMeshAsType = EInterchangeForceMeshType::IFMT_SkeletalMesh;
@@ -840,7 +840,7 @@ namespace UE::Interchange::Private
 
 			FillInterchangeGenericAssetsPipelineFromFbxMeshImportData(GenericAssetPipeline, Cast<UFbxStaticMeshImportData>(FbxImportUI->StaticMeshImportData));
 		}
-		else if (FbxImportUI->MeshTypeToImport == EFBXImportType::FBXIT_Animation)
+		else if (FbxImportUI->MeshTypeToImport == EFBXImportType::FBXIT_Animation || (FbxImportUI->bImportAsSkeletal && !FbxImportUI->bImportMesh && FbxImportUI->bImportAnimations))
 		{
 			GenericAssetPipeline->AnimationPipeline->bImportAnimations = true;
 			if (FbxImportUI->Skeleton)
diff --git a/Engine/Plugins/Interchange/Editor/Source/Pipelines/Private/InterchangeGraphInspectorPipeline.cpp b/Engine/Plugins/Interchange/Editor/Source/Pipelines/Private/InterchangeGraphInspectorPipeline.cpp
--- a/Engine/Plugins/Interchange/Editor/Source/Pipelines/Private/InterchangeGraphInspectorPipeline.cpp
+++ b/Engine/Plugins/Interchange/Editor/Source/Pipelines/Private/InterchangeGraphInspectorPipeline.cpp
@@ -14,7 +14,7 @@
 void UInterchangeGraphInspectorPipeline::ExecutePipeline(UInterchangeBaseNodeContainer* BaseNodeContainer, const TArray<UInterchangeSourceData*>& SourceDatas, const FString& ContentBasePath)
 {
 	//Do not pop dialog if we are unattended or doing automation
-	if (FApp::IsUnattended() || IsRunningCommandlet())
+	if (FApp::IsUnattended() || IsRunningCommandlet() || GIsRunningUnattendedScript)
 	{
 		return;
 	}
diff --git a/Engine/Plugins/Interchange/Editor/Source/Pipelines/Private/SInterchangePipelineConfigurationDialog.cpp b/Engine/Plugins/Interchange/Editor/Source/Pipelines/Private/SInterchangePipelineConfigurationDialog.cpp
--- a/Engine/Plugins/Interchange/Editor/Source/Pipelines/Private/SInterchangePipelineConfigurationDialog.cpp
+++ b/Engine/Plugins/Interchange/Editor/Source/Pipelines/Private/SInterchangePipelineConfigurationDialog.cpp
@@ -306,9 +306,8 @@ TSharedRef<SBox> SInterchangePipelineConfigurationDialog::SpawnPipelineConfigura
 				if (UInterchangePipelineBase* GeneratedPipeline = UE::Interchange::GeneratePipelineInstance(DefaultPipeline))
 				{
 					GeneratedPipeline->TransferAdjustSettings(DefaultPipeline);
-					if (Stack.StackName == ReimportStackName)
+					if (GeneratedPipeline->IsFromReimportOrOverride())
 					{
-						
 						//We save the pipeline settings to allow Reset to Default to work
 						GeneratedPipeline->SaveSettings(Stack.StackName);
 					}
@@ -984,7 +983,7 @@ void SInterchangePipelineConfigurationDialog::RefreshStack(bool bStackSelectionC
 			if (UInterchangePipelineBase* GeneratedPipeline = UE::Interchange::GeneratePipelineInstance(DefaultPipeline))
 			{
 				GeneratedPipeline->TransferAdjustSettings(DefaultPipeline);
-				if (Stack.StackName != ReimportStackName || !bStackSelectionChange)
+				if (!GeneratedPipeline->IsFromReimportOrOverride() || !bStackSelectionChange)
 				{
 					//Load the settings for this pipeline
 					GeneratedPipeline->LoadSettings(Stack.StackName);
diff --git a/Engine/Plugins/Interchange/Runtime/Source/Import/Private/Animation/InterchangeAnimSequenceFactory.cpp b/Engine/Plugins/Interchange/Runtime/Source/Import/Private/Animation/InterchangeAnimSequenceFactory.cpp
--- a/Engine/Plugins/Interchange/Runtime/Source/Import/Private/Animation/InterchangeAnimSequenceFactory.cpp
+++ b/Engine/Plugins/Interchange/Runtime/Source/Import/Private/Animation/InterchangeAnimSequenceFactory.cpp
@@ -1348,11 +1348,8 @@ UInterchangeFactoryBase::FImportAssetResult UInterchangeAnimSequenceFactory::End
 		{
 			for (const FString& CurveName : CurvesNotFound)
 			{
-				UInterchangeResultError_Generic* Message = AddMessage<UInterchangeResultError_Generic>();
-				Message->SourceAssetName = Arguments.SourceData->GetFilename();
-				Message->DestinationAssetName = Arguments.AssetName;
-				Message->AssetType = UAnimSequence::StaticClass();
-				Message->Text = FText::Format(NSLOCTEXT("UInterchangeAnimSequenceFactory", "NonExistingCurve", "Curve ({0}) was not found in the new Animation."), FText::FromString(CurveName));
+				//This is only a verbose log
+				UE_LOG(LogInterchangeImport, Verbose, TEXT("Curve (%s) was not found in the new Animation"), *CurveName);
 			}
 		}
 		Controller.NotifyPopulated();
diff --git a/Engine/Plugins/Interchange/Runtime/Source/Parsers/Fbx/Private/FbxMesh.cpp b/Engine/Plugins/Interchange/Runtime/Source/Parsers/Fbx/Private/FbxMesh.cpp
--- a/Engine/Plugins/Interchange/Runtime/Source/Parsers/Fbx/Private/FbxMesh.cpp
+++ b/Engine/Plugins/Interchange/Runtime/Source/Parsers/Fbx/Private/FbxMesh.cpp
@@ -2120,13 +2120,6 @@ bool FFbxMesh::GetGlobalJointBindPoseTransform(FFbxParser* Parser, FbxScene* SDK
 	// get bind pose
 	if (!bRerieveBindPoseResult)
 	{
-		bool bOriginalBadBindPoseMessageDisplay = bBadBindPoseMessageDisplay;
-		if (!GIsAutomationTesting && !bOriginalBadBindPoseMessageDisplay)
-		{
-			UInterchangeResultWarning_Generic* Message = Parser->AddMessage<UInterchangeResultWarning_Generic>();
-			Message->Text = LOCTEXT("MissingBindPose", "Missing bind pose. Try to recreate bind pose.");
-			bBadBindPoseMessageDisplay = true;
-		}
 		// if failed, delete bind pose, and retry.
 		const int32 PoseCount = SDKScene->GetPoseCount();
 		for (int32 PoseIndex = PoseCount - 1; PoseIndex >= 0; --PoseIndex)
@@ -2143,22 +2136,6 @@ bool FFbxMesh::GetGlobalJointBindPoseTransform(FFbxParser* Parser, FbxScene* SDK
 
 		SDKManager->CreateMissingBindPoses(SDKScene);
 		bRerieveBindPoseResult = RetrievePoseFromBindPose();
-		if (!bRerieveBindPoseResult)
-		{
-			if (!GIsAutomationTesting && !bOriginalBadBindPoseMessageDisplay)
-			{
-				UInterchangeResultWarning_Generic* Message = Parser->AddMessage<UInterchangeResultWarning_Generic>();
-				Message->Text = LOCTEXT("MissingBindPose_RecreateFailed", "Missing bind pose. Recreating bind pose failed.");
-			}
-		}
-		else
-		{
-			if (!GIsAutomationTesting && !bOriginalBadBindPoseMessageDisplay)
-			{
-				UInterchangeResultWarning_Generic* Message = Parser->AddMessage<UInterchangeResultWarning_Generic>();
-				Message->Text = LOCTEXT("MissingBindPose_RecreateSuccess", "Missing bind pose. Recreating bind pose succeeded.");
-			}
-		}
 	}
 	return bRerieveBindPoseResult;
 }
diff --git a/Engine/Plugins/Interchange/Runtime/Source/Pipelines/Private/InterchangeGenericAssetsPipeline.cpp b/Engine/Plugins/Interchange/Runtime/Source/Pipelines/Private/InterchangeGenericAssetsPipeline.cpp
--- a/Engine/Plugins/Interchange/Runtime/Source/Pipelines/Private/InterchangeGenericAssetsPipeline.cpp
+++ b/Engine/Plugins/Interchange/Runtime/Source/Pipelines/Private/InterchangeGenericAssetsPipeline.cpp
@@ -75,11 +75,15 @@ UInterchangeGenericAssetsPipeline::UInterchangeGenericAssetsPipeline()
 void UInterchangeGenericAssetsPipeline::PreDialogCleanup(const FName PipelineStackName)
 {
 	check(CommonSkeletalMeshesAndAnimationsProperties);
-	//Set the pipeline skeleton using what we found in AdjustSettingsForContext
-	CommonSkeletalMeshesAndAnimationsProperties->Skeleton = nullptr;
-	if (ContentPathExistingSkeleton.IsValid())
+	//Don't touch the skeleton property if we are a override pipeline or a re-import pipeline
+	if (!IsFromReimportOrOverride())
 	{
-		CommonSkeletalMeshesAndAnimationsProperties->Skeleton = Cast<USkeleton>(ContentPathExistingSkeleton.TryLoad());
+		//Set the pipeline skeleton using what we found in AdjustSettingsForContext
+		CommonSkeletalMeshesAndAnimationsProperties->Skeleton = nullptr;
+		if (ContentPathExistingSkeleton.IsValid())
+		{
+			CommonSkeletalMeshesAndAnimationsProperties->Skeleton = Cast<USkeleton>(ContentPathExistingSkeleton.TryLoad());
+		}
 	}
 
 	CommonSkeletalMeshesAndAnimationsProperties->bImportOnlyAnimations = bImportOnlyAnimationAdjusted;
diff --git a/Engine/Plugins/Interchange/Runtime/Source/Pipelines/Private/InterchangeGenericMeshPipeline.cpp b/Engine/Plugins/Interchange/Runtime/Source/Pipelines/Private/InterchangeGenericMeshPipeline.cpp
--- a/Engine/Plugins/Interchange/Runtime/Source/Pipelines/Private/InterchangeGenericMeshPipeline.cpp
+++ b/Engine/Plugins/Interchange/Runtime/Source/Pipelines/Private/InterchangeGenericMeshPipeline.cpp
@@ -178,7 +178,11 @@ bool UInterchangeGenericMeshPipeline::IsPropertyChangeNeedRefresh(const FPropert
 
 void UInterchangeGenericMeshPipeline::PreDialogCleanup(const FName PipelineStackName)
 {
-	PhysicsAsset = nullptr;
+	//Do not change the physics asset if this pipeline is a re-import or an override pipeline
+	if (!IsFromReimportOrOverride())
+	{
+		PhysicsAsset = nullptr;
+	}
 }
 
 #if WITH_EDITOR
diff --git a/Engine/Source/Runtime/Interchange/Core/Public/InterchangePipelineBase.h b/Engine/Source/Runtime/Interchange/Core/Public/InterchangePipelineBase.h
--- a/Engine/Source/Runtime/Interchange/Core/Public/InterchangePipelineBase.h
+++ b/Engine/Source/Runtime/Interchange/Core/Public/InterchangePipelineBase.h
@@ -310,6 +310,22 @@ public:
 		bIsBasicLayout = bBasicLayoutModeValue;
 	}
 
+	/*
+	 * Set to true if this pipeline is from a re-import or an override pipelines stack
+	 */
+	INTERCHANGECORE_API void SetFromReimportOrOverride(bool bInFromReimportOrOverride)
+	{
+		bFromReimportOrOverride = bInFromReimportOrOverride;
+	}
+
+	/*
+	 * Return true if the pipeline was created for a re-import or an override pipelines stack
+	 */
+	INTERCHANGECORE_API bool IsFromReimportOrOverride()
+	{
+		return bFromReimportOrOverride;
+	}
+	
 	/**
 	 * This function is called before showing the import dialog. It is not called when doing a reimport.
 	 */
@@ -460,12 +476,13 @@ public:
 	 * The Unreal import system has an option to force a name if we import only one main asset (one texture, one mesh or one animation).
 	 * The generic asset pipeline uses this information to behave as expected.
 	 */
+	UPROPERTY()
 	FString DestinationName;
 
 	/*
 	 * The content path where asset should be created.
 	 */
-	UPROPERTY(Transient)
+	UPROPERTY()
 	FString ContentImportPath;
 
 protected:
@@ -540,6 +557,13 @@ protected:
 	 */
 	bool bIsBasicLayout = false;
 
+	/*
+	 * If true, this pipeline was create to re-import an asset or override the project settings pipelines.
+	 * That kind of pipeline will not be treat like project settings pipeline in the UI. PredialogCleanup will not be called.
+	 */
+	UPROPERTY()
+	bool bFromReimportOrOverride = false;
+
 	UPROPERTY()
 	TObjectPtr<UInterchangeResultsContainer> Results;
 
diff --git a/Engine/Source/Runtime/Interchange/Engine/Private/InterchangeManager.cpp b/Engine/Source/Runtime/Interchange/Engine/Private/InterchangeManager.cpp
--- a/Engine/Source/Runtime/Interchange/Engine/Private/InterchangeManager.cpp
+++ b/Engine/Source/Runtime/Interchange/Engine/Private/InterchangeManager.cpp
@@ -1803,30 +1803,22 @@ UInterchangeManager::ImportInternal(const FString& ContentPath, const UInterchan
 			}
 		}
 	}
-	else if ( ImportAssetParameters.OverridePipelines.Num() == 0 )
+	else
 	{
-		
 		const bool bIsUnattended = FApp::IsUnattended() || GIsAutomationTesting || ImportAssetParameters.bIsAutomated || bSkipImportDialog;
 #if WITH_EDITORONLY_DATA
 		bool bShowPipelineStacksConfigurationDialog = !bIsUnattended
-															&& FInterchangeProjectSettingsUtils::ShouldShowPipelineStacksConfigurationDialog(bImportScene, bIsReimport, *SourceData)
-															&& !bImportCanceled
-															&& !IsRunningCommandlet();
+			&& FInterchangeProjectSettingsUtils::ShouldShowPipelineStacksConfigurationDialog(bImportScene, bIsReimport, *SourceData)
+			&& !bImportCanceled
+			&& !IsRunningCommandlet();
 #else
 		bool bShowPipelineStacksConfigurationDialog = false;
 #endif
-
 		if (FEngineAnalytics::IsAvailable())
 		{
 			Attribs.Add(FAnalyticsEventAttribute(TEXT("ShowImportDialog"), bShowPipelineStacksConfigurationDialog));
 		}
 
-		const TMap<FName, FInterchangePipelineStack>& DefaultPipelineStacks = InterchangeImportSettings.PipelineStacks;
-
-		const FName ReimportPipelineName = TEXT("ReimportPipeline");
-		TArray<FInterchangeStackInfo> PipelineStacks;
-		TArray<UInterchangePipelineBase*> OutPipelines;
-		
 		//If we need to display the dialog we want to translate the source file before showing it
 		if (RegisteredPipelineConfiguration && bShowPipelineStacksConfigurationDialog)
 		{
@@ -1871,45 +1863,53 @@ UInterchangeManager::ImportInternal(const FString& ContentPath, const UInterchan
 				bShowPipelineStacksConfigurationDialog = false;
 			}
 		}
-		
 
-		//Fill the Stacks before showing the UI
-		if (bIsReimport)
-		{
-			FInterchangeStackInfo& StackInfo = PipelineStacks.AddDefaulted_GetRef();
-			StackInfo.StackName = ReimportPipelineName;
+		const TMap<FName, FInterchangePipelineStack>& DefaultPipelineStacks = InterchangeImportSettings.PipelineStacks;
+
+		const FName ReimportPipelineName = TEXT("ReimportPipeline");
+		const FName OverridePipelineName = TEXT("OverridePipeline");
+		TArray<FInterchangeStackInfo> PipelineStacks;
+		TArray<UInterchangePipelineBase*> OutPipelines;
 
-			TArray<UObject*> Pipelines = OriginalAssetImportData->GetPipelines();
-			for (UObject* CurrentPipeline : Pipelines)
+		if (ImportAssetParameters.OverridePipelines.IsEmpty())
+		{
+			//Fill the Stacks before showing the UI
+			if (bIsReimport)
 			{
-				UInterchangePipelineBase* SourcePipeline = Cast<UInterchangePipelineBase>(CurrentPipeline);
-				if (!SourcePipeline)
+				FInterchangeStackInfo& StackInfo = PipelineStacks.AddDefaulted_GetRef();
+				StackInfo.StackName = ReimportPipelineName;
+
+				TArray<UObject*> Pipelines = OriginalAssetImportData->GetPipelines();
+				for (UObject* CurrentPipeline : Pipelines)
 				{
-					if (UInterchangePythonPipelineAsset* PythonPipelineAsset = Cast<UInterchangePythonPipelineAsset>(CurrentPipeline))
+					UInterchangePipelineBase* SourcePipeline = Cast<UInterchangePipelineBase>(CurrentPipeline);
+					if (!SourcePipeline)
 					{
-						SourcePipeline = PythonPipelineAsset->GeneratedPipeline;
+						if (UInterchangePythonPipelineAsset* PythonPipelineAsset = Cast<UInterchangePythonPipelineAsset>(CurrentPipeline))
+						{
+							SourcePipeline = PythonPipelineAsset->GeneratedPipeline;
+						}
+					}
+					if (SourcePipeline && SourcePipeline->SupportReimport()) //Its possible a pipeline doesnt exist anymore so it wont load into memory when we loading the outer asset
+					{
+						//Duplicate the pipeline saved in the asset import data
+						UInterchangePipelineBase* GeneratedPipeline = Cast<UInterchangePipelineBase>(StaticDuplicateObject(SourcePipeline, GetTransientPackage()));
+						// Make sure that the instance does not carry over standalone and public flags as they are not actual assets to be persisted
+						GeneratedPipeline->ClearFlags(EObjectFlags::RF_Standalone | EObjectFlags::RF_Public);
+						constexpr bool bFromReImportOrOverride = true;
+						GeneratedPipeline->SetFromReimportOrOverride(bFromReImportOrOverride);
+						AdjustPipelineSettingForContext(GeneratedPipeline);
+						StackInfo.Pipelines.Add(GeneratedPipeline);
+					}
+					else if (!SourcePipeline)
+					{
+						//A pipeline was not loaded
+						UE_LOG(LogInterchangeEngine, Warning, TEXT("Interchange Reimport: Missing import pipeline from the reimporting asset. The reimport might fail."));
 					}
-				}
-				if (SourcePipeline && SourcePipeline->SupportReimport()) //Its possible a pipeline doesnt exist anymore so it wont load into memory when we loading the outer asset
-				{
-					//Duplicate the pipeline saved in the asset import data
-					UInterchangePipelineBase* GeneratedPipeline = Cast<UInterchangePipelineBase>(StaticDuplicateObject(SourcePipeline, GetTransientPackage()));
-					// Make sure that the instance does not carry over standalone and public flags as they are not actual assets to be persisted
-					GeneratedPipeline->ClearFlags(EObjectFlags::RF_Standalone | EObjectFlags::RF_Public);
-					AdjustPipelineSettingForContext(GeneratedPipeline);
-					StackInfo.Pipelines.Add(GeneratedPipeline);
-				}
-				else if(!SourcePipeline)
-				{
-					//A pipeline was not loaded
-					UE_LOG(LogInterchangeEngine, Warning, TEXT("Interchange Reimport: Missing import pipeline from the reimporting asset. The reimport might fail."));
 				}
 			}
-		}
 
-		{
 			UE::Interchange::FScopedTranslator ScopedTranslator(SourceData);
-
 			for (const TPair<FName, FInterchangePipelineStack>& PipelineStackInfo : DefaultPipelineStacks)
 			{
 				FName StackName = PipelineStackInfo.Key;
@@ -1940,19 +1940,42 @@ UInterchangeManager::ImportInternal(const FString& ContentPath, const UInterchan
 				}
 			}
 		}
-
-		auto SetImportAllWithSamePipelines = [this, &AsyncTranslator](TArray<UInterchangePipelineBase*>& ToDuplicatePipelines)
+		else
 		{
-			TArray<UInterchangePipelineBase*>& PipelineList = ImportAllWithSamePipelines.FindOrAdd(AsyncTranslator->GetClass());
-			for (const UInterchangePipelineBase* Pipeline : ToDuplicatePipelines)
+			//If we receive an override list of pipelines add them to the stack
+			FInterchangeStackInfo& StackInfo = PipelineStacks.AddDefaulted_GetRef();
+			StackInfo.StackName = OverridePipelineName;
+			for (int32 GraphPipelineIndex = 0; GraphPipelineIndex < ImportAssetParameters.OverridePipelines.Num(); ++GraphPipelineIndex)
 			{
-				if (UInterchangePipelineBase* DupPipeline = DuplicateObject<UInterchangePipelineBase>(Pipeline, GetTransientPackage()))
+				UInterchangePipelineBase* GeneratedPipeline = UE::Interchange::GeneratePipelineInstance(ImportAssetParameters.OverridePipelines[GraphPipelineIndex]);
+				if (!GeneratedPipeline)
+				{
+					UE_LOG(LogInterchangeEngine, Error, TEXT("Interchange Import: Overridden pipeline array contains a NULL pipeline. Fix your script or code to avoid this issue."));
+					continue;
+				}
+				else
 				{
-					DupPipeline->SetInternalFlags(EInternalObjectFlags::Async);
-					PipelineList.Add(DupPipeline);
+					GeneratedPipeline->ClearFlags(EObjectFlags::RF_Standalone | EObjectFlags::RF_Public);
+					constexpr bool bFromReImportOrOverride = true;
+					GeneratedPipeline->SetFromReimportOrOverride(bFromReImportOrOverride);
+					AdjustPipelineSettingForContext(GeneratedPipeline);
+					StackInfo.Pipelines.Add(GeneratedPipeline);
 				}
 			}
-		};
+		}
+
+		auto SetImportAllWithSamePipelines = [this, &AsyncTranslator](TArray<UInterchangePipelineBase*>& ToDuplicatePipelines)
+			{
+				TArray<UInterchangePipelineBase*>& PipelineList = ImportAllWithSamePipelines.FindOrAdd(AsyncTranslator->GetClass());
+				for (const UInterchangePipelineBase* Pipeline : ToDuplicatePipelines)
+				{
+					if (UInterchangePipelineBase* DupPipeline = DuplicateObject<UInterchangePipelineBase>(Pipeline, GetTransientPackage()))
+					{
+						DupPipeline->SetInternalFlags(EInternalObjectFlags::Async);
+						PipelineList.Add(DupPipeline);
+					}
+				}
+			};
 
 		if (bIsReimport)
 		{
@@ -1976,7 +1999,7 @@ UInterchangeManager::ImportInternal(const FString& ContentPath, const UInterchan
 					{
 						return StackInfo.StackName == ReimportPipelineName;
 					});
-				
+
 				check(StackInfoPtr);
 				OutPipelines = StackInfoPtr->Pipelines;
 			}
@@ -2051,26 +2074,6 @@ UInterchangeManager::ImportInternal(const FString& ContentPath, const UInterchan
 			}
 		}
 	}
-	else
-	{
-		for (int32 GraphPipelineIndex = 0; GraphPipelineIndex < ImportAssetParameters.OverridePipelines.Num(); ++GraphPipelineIndex)
-		{
-			UInterchangePipelineBase* GeneratedPipeline = UE::Interchange::GeneratePipelineInstance(ImportAssetParameters.OverridePipelines[GraphPipelineIndex]);
-			if (!GeneratedPipeline)
-			{
-				UE_LOG(LogInterchangeEngine, Error, TEXT("Interchange Import: Overridden pipeline array contains a NULL pipeline. Fix your script or code to avoid this issue."));
-				continue;
-			}
-			else
-			{
-				// Duplicate the override pipelines to protect the scripted users form making race conditions
-				AdjustPipelineSettingForContext(GeneratedPipeline);
-				AsyncHelper->Pipelines.Add(GeneratedPipeline);
-				AsyncHelper->OriginalPipelines.Add(GeneratedPipeline);
-				UE::Interchange::Private::FillPipelineAnalyticData(GeneratedPipeline, UniqueId, FString());
-			}
-		}
-	}
 
 	//Cancel the import do not queue task
 	if (bImportCanceled || bImportAborted)
