Interchange: remove a warning when using the override stack in automation.

#rnx

diff --git a/Engine/Source/Runtime/Interchange/Engine/Private/InterchangeManager.cpp b/Engine/Source/Runtime/Interchange/Engine/Private/InterchangeManager.cpp
--- a/Engine/Source/Runtime/Interchange/Engine/Private/InterchangeManager.cpp
+++ b/Engine/Source/Runtime/Interchange/Engine/Private/InterchangeManager.cpp
@@ -2036,16 +2036,20 @@ UInterchangeManager::ImportInternal(const FString& ContentPath, const UInterchan
 				}
 				else if (PipelineStacks.Num() > 0)
 				{
+					//Take the first valid stack
 					for (FInterchangeStackInfo& StackInfo : PipelineStacks)
 					{
-						OutPipelines = StackInfo.Pipelines;
-						UE_LOG(LogInterchangeEngine, Warning, TEXT("Interchange import: Invalid Default stack. using stack [%s] to import."), *StackInfo.StackName.ToString());
-						break;
+						if (!StackInfo.Pipelines.IsEmpty())
+						{
+							OutPipelines = StackInfo.Pipelines;
+							UE_LOG(LogInterchangeEngine, Display, TEXT("Interchange import: Using stack [%s] to import."), *StackInfo.StackName.ToString());
+							break;
+						}
 					}
 				}
 				else
 				{
-					UE_LOG(LogInterchangeEngine, Warning, TEXT("Interchange Import: Cannot find any valid stack, canceling import."));
+					UE_LOG(LogInterchangeEngine, Error, TEXT("Interchange Import: Cannot find any valid stack, canceling import."));
 					bImportCanceled = true;
 				}
 			}
