Interchange: If we discover we override an existing asset before starting the import, enqueue a real re-import.

#rnx

diff --git a/Engine/Source/Developer/AssetTools/Private/AssetTools.cpp b/Engine/Source/Developer/AssetTools/Private/AssetTools.cpp
--- a/Engine/Source/Developer/AssetTools/Private/AssetTools.cpp
+++ b/Engine/Source/Developer/AssetTools/Private/AssetTools.cpp
@@ -3771,6 +3771,14 @@ TArray<UObject*> UAssetToolsImpl::ImportAssetsInternal(const TArray<FString>& Fi
 									}
 								}
 							}
+							for (UObject* Object : MainAssets)
+							{
+								if (Object->IsAsset())
+								{
+									FAssetRegistryModule::AssetCreated(Object);
+									GEditor->BroadcastObjectReimported(Object);
+								}
+							}
 
 							//Force browser to sync to the import asset if there is only one asset imported
 							if (bSyncToBrowser || (bForceContentBrowserSyncIfOnlyOneMainAsset && MainAssets.Num() == 1))
@@ -3817,20 +3825,17 @@ TArray<UObject*> UAssetToolsImpl::ImportAssetsInternal(const TArray<FString>& Fi
 							InterchangeManager.ConvertImportData(Params.AssetImportTask->Options, ImportAssetParameters);
 						}
 					}
-					UE::Interchange::FAssetImportResultRef InterchangeResult = (InterchangeManager.ImportAssetAsync(DestinationPath, ScopedSourceData.GetSourceData(), ImportAssetParameters));
+					UE::Interchange::FAssetImportResultRef InterchangeResult = Params.bAllowAsyncImport
+						? (InterchangeManager.ImportAssetAsync(DestinationPath, ScopedSourceData.GetSourceData(), ImportAssetParameters))
+						: (InterchangeManager.ImportAssetWithResult(DestinationPath, ScopedSourceData.GetSourceData(), ImportAssetParameters));
+
+					InterchangeResult->OnDone(AppendAndBroadcastImportResultIfNeeded);
 
 					// If we have an ImportTask, fill out the asynchronous results object here so the caller can see when the results are ready
 					if ( Params.AssetImportTask)
 					{
 						 Params.AssetImportTask->AsyncResults = InterchangeResult;
 					}
-
-					InterchangeResult->OnDone(AppendAndBroadcastImportResultIfNeeded);
-
-					if (!Params.bAllowAsyncImport)
-					{
-						InterchangeResult->WaitUntilDone();
-					}
 				}
 
 				if (!Params.bAllowAsyncImport)
diff --git a/Engine/Source/Runtime/Interchange/Engine/Private/InterchangeManager.cpp b/Engine/Source/Runtime/Interchange/Engine/Private/InterchangeManager.cpp
--- a/Engine/Source/Runtime/Interchange/Engine/Private/InterchangeManager.cpp
+++ b/Engine/Source/Runtime/Interchange/Engine/Private/InterchangeManager.cpp
@@ -1540,6 +1540,14 @@ bool UInterchangeManager::ImportAsset(const FString& ContentPath, const UInterch
 	return ImportAsset(ContentPath, SourceData, ImportAssetParameters, ImportedObjects);
 }
 
+UE::Interchange::FAssetImportResultRef UInterchangeManager::ImportAssetWithResult(const FString& ContentPath, const UInterchangeSourceData* SourceData, const FImportAssetParameters& ImportAssetParameters)
+{
+	ImportAssetParameters.bRunSynchronous = true;
+	UE::Interchange::FAssetImportResultRef InterchangeResult = ImportInternal(ContentPath, SourceData, ImportAssetParameters, UE::Interchange::EImportType::ImportType_Asset).Get<0>();
+	InterchangeResult->WaitUntilDone(ImportAssetParameters.bRunSynchronous);
+	return InterchangeResult;
+}
+
 UE::Interchange::FAssetImportResultRef UInterchangeManager::ImportAssetAsync(const FString& ContentPath, const UInterchangeSourceData* SourceData, const FImportAssetParameters& ImportAssetParameters)
 {
 	ImportAssetParameters.bRunSynchronous = false;
@@ -1668,7 +1676,77 @@ UInterchangeManager::ImportInternal(const FString& ContentPath, const UInterchan
 		UE_LOG(LogInterchangeEngine, Error, TEXT("Cannot import file. There is no pipeline stack defined for the %s import type."), bImportScene ? TEXT("scene") : TEXT("content"));
 		return EarlyExit();
 	}
-	
+
+	const bool bIsUnattended = FApp::IsUnattended() || FApp::IsGame() || GIsAutomationTesting || ImportAssetParameters.bIsAutomated;
+
+	//If we detect the import will do a re-import of an existing asset, the ReimportAsset parameter should be filled
+	//And we will also convert the import data. With this code this import will be considered a re-import.
+	if (ImportType == UE::Interchange::EImportType::ImportType_Asset && !ImportAssetParameters.ReimportAsset)
+	{
+		FString AssetpackageName = ImportAssetParameters.DestinationName.IsEmpty() ? FPaths::GetCleanFilename(SourceData->GetFilename()) : ImportAssetParameters.DestinationName;
+		FString PackageBasePath = ContentPath;
+		UE::Interchange::SanitizeObjectName(AssetpackageName);
+		UE::Interchange::SanitizeObjectPath(PackageBasePath);
+		FString FullPackagePath = FPaths::Combine(*PackageBasePath, *AssetpackageName);
+		if (!UE::Interchange::FPackageUtils::IsMapPackageAsset(FullPackagePath))
+		{
+			UPackage* Pkg = FindPackage(nullptr, *FullPackagePath);
+			if (!Pkg)
+			{
+				//Try to load the package from disk
+				Pkg = LoadPackage(nullptr, *FullPackagePath, LOAD_NoWarn | LOAD_Quiet);
+			}
+
+			if (Pkg)
+			{
+				UObject* ExistingAsset = StaticFindObject(nullptr, Pkg, *AssetpackageName);
+				//If we found an asset we can re-import, convert set the param to re-import instead of override
+				if (ExistingAsset && GetRegisteredFactoryClass(ExistingAsset->GetClass()))
+				{
+					bool bConvertToReimport = ImportAssetParameters.bReplaceExisting;
+					if (!bIsUnattended && !ImportAssetParameters.bReplaceExisting)
+					{
+						FText OverrideDialogMessage = FText::Format(NSLOCTEXT("InterchangeManager", "OverrideAssetMessage", "This import will override asset '{0}', Do you want to convert this import to a re-import?")
+							, FText::FromString(ExistingAsset->GetFullName()));
+						EAppReturnType::Type DialogResult = FMessageDialog::Open(EAppMsgType::YesNoCancel, OverrideDialogMessage);
+						switch (DialogResult)
+						{
+							case EAppReturnType::Yes:
+							{
+								bConvertToReimport = true;
+								break;
+							}
+							case EAppReturnType::No:
+							{
+								bConvertToReimport = false;
+								break;
+							}
+							case EAppReturnType::Cancel:
+							{
+								UE_LOG(LogInterchangeEngine, Display, TEXT("User cancel import of file [%s]."), *(SourceData->GetFilename()));
+								return EarlyExit();
+							}
+						}
+					}
+
+					if (bConvertToReimport)
+					{
+						//Make sure we have the correct flag
+						ExistingAsset->SetFlags(RF_Standalone | RF_Public);
+						FImportAssetParameters* MutableImportAssetParameters = const_cast<FImportAssetParameters*>(&ImportAssetParameters);
+						MutableImportAssetParameters->ReimportAsset = ExistingAsset;
+						UInterchangeAssetImportData* OriginalAssetImportData = UInterchangeAssetImportData::GetFromObject(ImportAssetParameters.ReimportAsset);
+						if (!OriginalAssetImportData)
+						{
+							//Convert the object asset import data to interchange
+							ConvertImportData(MutableImportAssetParameters->ReimportAsset, FPaths::GetExtension(SourceData->GetFilename()));
+						}
+					}
+				}
+			}
+		}
+	}
+
 	//Set a default pipeline stack if none is valid
 	if (!InterchangeImportSettings.PipelineStacks.Contains(InterchangeImportSettings.DefaultPipelineStack))
 	{
@@ -1805,7 +1883,6 @@ UInterchangeManager::ImportInternal(const FString& ContentPath, const UInterchan
 	}
 	else
 	{
-		const bool bIsUnattended = FApp::IsUnattended() || GIsAutomationTesting || ImportAssetParameters.bIsAutomated || bSkipImportDialog;
 #if WITH_EDITORONLY_DATA
 		bool bShowPipelineStacksConfigurationDialog = !bIsUnattended
 			&& FInterchangeProjectSettingsUtils::ShouldShowPipelineStacksConfigurationDialog(bImportScene, bIsReimport, *SourceData)
@@ -1994,10 +2071,11 @@ UInterchangeManager::ImportInternal(const FString& ContentPath, const UInterchan
 			}
 			else
 			{
-				//When we do not show the UI we use the original stack
-				FInterchangeStackInfo* StackInfoPtr = PipelineStacks.FindByPredicate([ReimportPipelineName](const FInterchangeStackInfo& StackInfo)
+				FName ClosurePipelineStackName = ImportAssetParameters.OverridePipelines.IsEmpty() ? ReimportPipelineName : OverridePipelineName;
+				//When we do not show the UI we use the original import stack or the provided override
+				FInterchangeStackInfo* StackInfoPtr = PipelineStacks.FindByPredicate([ClosurePipelineStackName](const FInterchangeStackInfo& StackInfo)
 					{
-						return StackInfo.StackName == ReimportPipelineName;
+						return StackInfo.StackName == ClosurePipelineStackName;
 					});
 
 				check(StackInfoPtr);
diff --git a/Engine/Source/Runtime/Interchange/Engine/Public/InterchangeManager.h b/Engine/Source/Runtime/Interchange/Engine/Public/InterchangeManager.h
--- a/Engine/Source/Runtime/Interchange/Engine/Public/InterchangeManager.h
+++ b/Engine/Source/Runtime/Interchange/Engine/Public/InterchangeManager.h
@@ -582,6 +582,17 @@ public:
 	 */
 	INTERCHANGEENGINE_API bool ImportAsset(const FString& ContentPath, const UInterchangeSourceData* SourceData, const FImportAssetParameters& ImportAssetParameters);
 
+	/**
+	 * Call this to start a synchronous asset import process.
+	 * This process can import many different assets into the game content.
+	 *
+	 * @Param ContentPath - The path where the imported assets will be created.
+	 * @Param SourceData - The source data input to translate.
+	 * @param ImportAssetParameters - All parameters that need to be passed to the import asset function.
+	 * @return return an import result which can be use to know when the asynchronous import is terminate.
+	 */
+	INTERCHANGEENGINE_API UE::Interchange::FAssetImportResultRef ImportAssetWithResult(const FString& ContentPath, const UInterchangeSourceData* SourceData, const FImportAssetParameters& ImportAssetParameters);
+
 	/**
 	 * Call this to start an asynchronous asset import process.
 	 * This process can import many different assets into the game content.
