Interchange: Fix translator settings serialized in the import data

#rnx

diff --git a/Engine/Source/Runtime/Interchange/Engine/Private/InterchangeAssetImportData.cpp b/Engine/Source/Runtime/Interchange/Engine/Private/InterchangeAssetImportData.cpp
--- a/Engine/Source/Runtime/Interchange/Engine/Private/InterchangeAssetImportData.cpp
+++ b/Engine/Source/Runtime/Interchange/Engine/Private/InterchangeAssetImportData.cpp
@@ -177,7 +177,7 @@ void UInterchangeAssetImportData::SetTranslatorSettings(UInterchangeTranslatorSe
 	CachedTranslatorSettings = {};
 	if (TranslatorSettings)
 	{
-		FString TranslatorSettingsJSON = SerializePipeline(TranslatorSettings);
+		FString TranslatorSettingsJSON = SerializeTranslatorSettings(TranslatorSettings);
 
 		FString TranslatorSettingsClassFullName = TranslatorSettings->GetClass()->GetFullName();
 		CachedTranslatorSettings = TPair<FString, FString>(TranslatorSettingsClassFullName, TranslatorSettingsJSON);
