[Interchange] 'Apply to All' checkbox doesn't work when overriding multiple files during import

#rnx

diff --git a/Engine/Source/Runtime/Interchange/Engine/Private/InterchangeManager.cpp b/Engine/Source/Runtime/Interchange/Engine/Private/InterchangeManager.cpp
--- a/Engine/Source/Runtime/Interchange/Engine/Private/InterchangeManager.cpp
+++ b/Engine/Source/Runtime/Interchange/Engine/Private/InterchangeManager.cpp
@@ -61,7 +61,21 @@ bool UInterchangeManager::bIsCreatingSingleton = false;
 
 namespace UE::Interchange::Private
 {
-	static TAtomic<uint64> AsyncHelperCounter = 0;
+	namespace StaticHelpers
+	{
+		static TAtomic<uint64> AsyncHelperCounter = 0;
+	
+		/**
+			 * Interchange import task will show a dialog in case user try to override an existing asset and bReplaceExisting is false,
+			 * if this optional is set, it will override or not all existing assets this Import stack tries to override.
+			 * This setting resets whenever Importing finishes.
+			 * -1:= Reset (not set)
+			 *  0:= Set and False
+			 *  1:= Set and True;
+			 */
+		static TAtomic<int8> bReplaceExistingAllDialogAnswer = -1;
+	}
+	
 
 	const FLogCategoryBase* GetLogInterchangePtr()
 	{
@@ -220,21 +234,22 @@ UE::Interchange::FImportAsyncHelper::FImportAsyncHelper()
 {
 	bCancel = false;
 
-	if (UE::Interchange::Private::AsyncHelperCounter == 0)
+	if (UE::Interchange::Private::StaticHelpers::AsyncHelperCounter == 0)
 	{
 		UInterchangeManager::GetInterchangeManager().OnImportStarted.Broadcast();
 	}
-	UE::Interchange::Private::AsyncHelperCounter++;
+	UE::Interchange::Private::StaticHelpers::AsyncHelperCounter++;
 }
 
 UE::Interchange::FImportAsyncHelper::~FImportAsyncHelper()
 {
 	CleanUp();
 
-	UE::Interchange::Private::AsyncHelperCounter--;
-	if (UE::Interchange::Private::AsyncHelperCounter == 0)
+	UE::Interchange::Private::StaticHelpers::AsyncHelperCounter--;
+	if (UE::Interchange::Private::StaticHelpers::AsyncHelperCounter == 0)
 	{
 		UInterchangeManager::GetInterchangeManager().OnImportFinished.Broadcast();
+		UInterchangeManager::GetInterchangeManager().ResetReplaceExistingAlldialogAnswerSet();
 	}
 }
 
@@ -2298,7 +2313,7 @@ void UInterchangeManager::SetInterchangeImportEnabled(bool bEnabled)
 
 bool UInterchangeManager::IsImporting()
 {
-	return UE::Interchange::Private::AsyncHelperCounter > 0;
+	return UE::Interchange::Private::StaticHelpers::AsyncHelperCounter > 0;
 }
 
 bool UInterchangeManager::ExportAsset(const UObject* Asset, bool bIsAutomated)
@@ -2719,3 +2734,23 @@ void UInterchangeManager::SetActiveMode(bool IsActive)
 		NotificationTickHandle.Reset();
 	}
 }
+
+void UInterchangeManager::SetReplaceExistingAlldialogAnswer(bool bReplaceExistingAllDialogAnswer)
+{
+	UE::Interchange::Private::StaticHelpers::bReplaceExistingAllDialogAnswer = bReplaceExistingAllDialogAnswer ? 1 : 0;
+}
+
+void UInterchangeManager::ResetReplaceExistingAlldialogAnswerSet()
+{
+	UE::Interchange::Private::StaticHelpers::bReplaceExistingAllDialogAnswer = -1;
+}
+
+TOptional<bool> UInterchangeManager::GetReplaceExistingAlldialogAnswer()
+{
+	TOptional<bool> Result;
+	if (UE::Interchange::Private::StaticHelpers::bReplaceExistingAllDialogAnswer != -1)
+	{
+		Result = (UE::Interchange::Private::StaticHelpers::bReplaceExistingAllDialogAnswer == 1);
+	}
+	return Result;
+}
\ No newline at end of file
diff --git a/Engine/Source/Runtime/Interchange/Engine/Private/Tasks/InterchangeTaskImportObject.cpp b/Engine/Source/Runtime/Interchange/Engine/Private/Tasks/InterchangeTaskImportObject.cpp
--- a/Engine/Source/Runtime/Interchange/Engine/Private/Tasks/InterchangeTaskImportObject.cpp
+++ b/Engine/Source/Runtime/Interchange/Engine/Private/Tasks/InterchangeTaskImportObject.cpp
@@ -361,6 +361,9 @@ void UE::Interchange::FTaskImportObject_GameThread::DoTask(ENamedThreads::Type C
 		//The node is disabled return now
 		return;
 	}
+
+	UInterchangeManager& InterchangeManager = UInterchangeManager::GetInterchangeManager();
+
 	//If we do a reimport no need to create a package
 	if (ObjectToReimport)
 	{
@@ -441,9 +444,9 @@ void UE::Interchange::FTaskImportObject_GameThread::DoTask(ENamedThreads::Type C
 		{
 			const FString AssetFullName = ExistingAsset->GetFullName();
 			//If the bReplaceExistingAllDialogAnswer was set do not show again the message dialog, simply reuse the previous answer.
-			if (AsyncHelper->TaskData.bReplaceExistingAllDialogAnswer.IsSet())
+			if (InterchangeManager.GetReplaceExistingAlldialogAnswer().IsSet())
 			{
-				bSkipObjectNoReplace = !AsyncHelper->TaskData.bReplaceExistingAllDialogAnswer.GetValue();
+				bSkipObjectNoReplace = !InterchangeManager.GetReplaceExistingAlldialogAnswer().GetValue();
 			}
 			else
 			{
@@ -456,12 +459,12 @@ void UE::Interchange::FTaskImportObject_GameThread::DoTask(ENamedThreads::Type C
 					switch (DialogResult)
 					{
 						case EAppReturnType::YesAll:
-							AsyncHelper->TaskData.bReplaceExistingAllDialogAnswer = true;
+							InterchangeManager.SetReplaceExistingAlldialogAnswer(true);
 						case EAppReturnType::Yes:
 							bSkipObjectNoReplace = false;
 							break;
 						case EAppReturnType::NoAll:
-							AsyncHelper->TaskData.bReplaceExistingAllDialogAnswer = false;
+							InterchangeManager.SetReplaceExistingAlldialogAnswer(false);
 						case EAppReturnType::No:
 							bSkipObjectNoReplace = true;
 					}
@@ -526,7 +529,7 @@ void UE::Interchange::FTaskImportObject_GameThread::DoTask(ENamedThreads::Type C
 			return;
 		}
 
-		if (UInterchangeManager::GetInterchangeManager().IsObjectBeingImported(ExistingAsset))
+		if (InterchangeManager.IsObjectBeingImported(ExistingAsset))
 		{
 			//Skip this node, it is currently being imported by another import task
 			UInterchangeResultError_Generic* Message = Factory->AddMessage<UInterchangeResultError_Generic>();
diff --git a/Engine/Source/Runtime/Interchange/Engine/Public/InterchangeManager.h b/Engine/Source/Runtime/Interchange/Engine/Public/InterchangeManager.h
--- a/Engine/Source/Runtime/Interchange/Engine/Public/InterchangeManager.h
+++ b/Engine/Source/Runtime/Interchange/Engine/Public/InterchangeManager.h
@@ -111,12 +111,6 @@ namespace UE
 
 			/** Whether or not to overwrite existing assets. */
 			bool bReplaceExisting = true;
-
-			/**
-			 * Interchange import task will show a dialog in case user try to override an existing asset and bReplaceExisting is false,
-			 * if this optional is set, it will override or not all existing assets this task try to override.
-			 */
-			TOptional<bool> bReplaceExistingAllDialogAnswer;
 		};
 
 		class FImportResult : protected FGCObject
@@ -748,6 +742,24 @@ public:
 	 * Example: if a skeletal mesh re-import cannot apply the existing alternate skinning data, it will enqueue a post import task to re-import those alternate skinning files.
 	 */
 	INTERCHANGEENGINE_API bool EnqueuePostImportTask(TSharedPtr<FInterchangePostImportTask> PostImportTask);
+
+
+	/**
+	* Sets the bReplaceExistingAllDialogAnswer which is responsible for:
+	* if the Import process will override or not all existing assets this Import stack tries to override.
+	*/
+	INTERCHANGEENGINE_API static void SetReplaceExistingAlldialogAnswer(bool bReplaceExistingAllDialogAnswer);
+
+	/**
+	* Checks if the bReplaceExistingAllDialogAnswer.
+	*/
+	INTERCHANGEENGINE_API static void ResetReplaceExistingAlldialogAnswerSet();
+
+	/**
+	* Gets the bReplaceExistingAllDialogAnswer.
+	*/
+	INTERCHANGEENGINE_API static TOptional<bool> GetReplaceExistingAlldialogAnswer();
+
 protected:
 
 	/** Return true if Interchange can show UI. */
