Interchange: allow pipeline asset to mark any property with ResetPreDialog. This will reset the property to there default value when the dialog show a pipeline stack.

This change remove the AlwaysResetToDefault meta value for AssetName property to use the new ResetPreDialog, so the value is not reset when the UI get refresh because of a property change.
ForceMeshType is now also reset when showing the dialog. ForceMeshType should normaly be a single time use like the asset name.

#rnx

#changelist validated
#virtualized

diff --git a/Engine/Plugins/Interchange/Editor/Source/Pipelines/Private/InterchangeEditorPipelineDetails.cpp b/Engine/Plugins/Interchange/Editor/Source/Pipelines/Private/InterchangeEditorPipelineDetails.cpp
--- a/Engine/Plugins/Interchange/Editor/Source/Pipelines/Private/InterchangeEditorPipelineDetails.cpp
+++ b/Engine/Plugins/Interchange/Editor/Source/Pipelines/Private/InterchangeEditorPipelineDetails.cpp
@@ -620,6 +620,47 @@ void FInterchangePipelineBaseDetailsCustomization::CustomizeDetails(IDetailLayou
 					.AutoWidth()
 					.Padding(6.0f, 1.0f, 3.0f, 1.0f)
 					.VAlign(VAlign_Center)
+					[
+						SNew(STextBlock)
+						.Visibility(PipelineInternalEditionData ? EVisibility::Collapsed : EVisibility::Visible)
+						.Font(IDetailLayoutBuilder::GetDetailFont())
+						.Text(LOCTEXT("ResetPreDialogName", "Reset PreDialog"))
+					]
+					+ SHorizontalBox::Slot()
+					.AutoWidth()
+					.Padding(3.0f, 1.0f)
+					[
+						SNew(SCheckBox)
+						.Visibility(PipelineInternalEditionData ? EVisibility::Collapsed : EVisibility::Visible)
+						.ToolTipText(LOCTEXT("ResetPreDialogTooltip", "If true this property will be reset when displaying the interchange import dialog."))
+						.OnCheckStateChanged_Lambda([InterchangePipelinePtr = InterchangePipeline, PropertyPath](ECheckBoxState CheckType)
+						{
+							if (!ensure(InterchangePipelinePtr.IsValid()))
+							{
+								return;
+							}
+							FScopedTransaction ScopedTransaction(LOCTEXT("TransactionResetPreDialogToggle", "Toggle property reset pre-dialog."), !GIsTransacting);
+							InterchangePipelinePtr->Modify();
+							InterchangePipelinePtr->FindOrAddPropertyStates(PropertyPath).SetPropertyPreDialogReset((CheckType == ECheckBoxState::Checked));
+							InterchangePipelinePtr->PostEditChange();
+						})
+						.IsChecked_Lambda([InterchangePipelinePtr = InterchangePipeline, PropertyPath]()
+						{
+							if (!InterchangePipelinePtr.IsValid())
+							{
+								return ECheckBoxState::Unchecked;
+							}
+							if (const FInterchangePipelinePropertyStates* PropertyStates = InterchangePipelinePtr->GetPropertyStates(PropertyPath))
+							{
+								return PropertyStates->IsPropertyPreDialogReset() ? ECheckBoxState::Checked : ECheckBoxState::Unchecked;
+							}
+							return ECheckBoxState::Unchecked;
+						})
+					]
+					+ SHorizontalBox::Slot()
+					.AutoWidth()
+					.Padding(6.0f, 1.0f, 3.0f, 1.0f)
+					.VAlign(VAlign_Center)
 					[
 						SNew(STextBlock)
 						.Visibility(PipelineInternalEditionData ? EVisibility::Collapsed : EVisibility::Visible)
diff --git a/Engine/Plugins/Interchange/Editor/Source/Pipelines/Private/SInterchangePipelineConfigurationDialog.cpp b/Engine/Plugins/Interchange/Editor/Source/Pipelines/Private/SInterchangePipelineConfigurationDialog.cpp
--- a/Engine/Plugins/Interchange/Editor/Source/Pipelines/Private/SInterchangePipelineConfigurationDialog.cpp
+++ b/Engine/Plugins/Interchange/Editor/Source/Pipelines/Private/SInterchangePipelineConfigurationDialog.cpp
@@ -313,8 +313,9 @@ TSharedRef<SBox> SInterchangePipelineConfigurationDialog::SpawnPipelineConfigura
 					}
 					else
 					{
+						constexpr bool bResetPreDialogTrue = true;
 						//Load the settings for this pipeline
-						GeneratedPipeline->LoadSettings(Stack.StackName);
+						GeneratedPipeline->LoadSettings(Stack.StackName, bResetPreDialogTrue);
 						GeneratedPipeline->PreDialogCleanup(Stack.StackName);
 					}
 					GeneratedPipeline->SetBasicLayoutMode(bBasicLayout);
@@ -986,7 +987,7 @@ void SInterchangePipelineConfigurationDialog::RefreshStack(bool bStackSelectionC
 				if (!GeneratedPipeline->IsFromReimportOrOverride() || !bStackSelectionChange)
 				{
 					//Load the settings for this pipeline
-					GeneratedPipeline->LoadSettings(Stack.StackName);
+					GeneratedPipeline->LoadSettings(Stack.StackName, bStackSelectionChange);
 					if (bStackSelectionChange)
 					{
 						//Do not reset pipeline value if we are just refreshing the filtering
diff --git a/Engine/Plugins/Interchange/Runtime/Source/Pipelines/Public/InterchangeGenericAssetsPipeline.h b/Engine/Plugins/Interchange/Runtime/Source/Pipelines/Public/InterchangeGenericAssetsPipeline.h
--- a/Engine/Plugins/Interchange/Runtime/Source/Pipelines/Public/InterchangeGenericAssetsPipeline.h
+++ b/Engine/Plugins/Interchange/Runtime/Source/Pipelines/Public/InterchangeGenericAssetsPipeline.h
@@ -58,7 +58,7 @@ public:
 	bool bUseSourceNameForAsset = true;
 
 	/** If set, and there is only one asset and one source, the imported asset is given this name. */
-	UPROPERTY(EditAnywhere, BlueprintReadWrite, Category = "Common", meta = (StandAlonePipelineProperty = "True", AlwaysResetToDefault = "True"))
+	UPROPERTY(EditAnywhere, BlueprintReadWrite, Category = "Common", meta = (StandAlonePipelineProperty = "True"))
 	FString AssetName;
 
 	/** Translation offset applied to meshes and animations. */
diff --git a/Engine/Plugins/Interchange/Runtime/Source/Pipelines/Public/InterchangeGenericMaterialPipeline.h b/Engine/Plugins/Interchange/Runtime/Source/Pipelines/Public/InterchangeGenericMaterialPipeline.h
--- a/Engine/Plugins/Interchange/Runtime/Source/Pipelines/Public/InterchangeGenericMaterialPipeline.h
+++ b/Engine/Plugins/Interchange/Runtime/Source/Pipelines/Public/InterchangeGenericMaterialPipeline.h
@@ -76,7 +76,7 @@ public:
 	EInterchangeMaterialSearchLocation SearchLocation = EInterchangeMaterialSearchLocation::Local;
 
 	/** If set, and there is only one asset and one source, the imported asset will be given this name. */
-	UPROPERTY(EditAnywhere, BlueprintReadWrite, Category = "Materials", meta = (StandAlonePipelineProperty = "True", AlwaysResetToDefault = "True"))
+	UPROPERTY(EditAnywhere, BlueprintReadWrite, Category = "Materials", meta = (StandAlonePipelineProperty = "True"))
 	FString AssetName;
 
 	/** Determines what kind of material assets should be created for the imported materials. */
diff --git a/Engine/Plugins/Interchange/Runtime/Source/Pipelines/Public/InterchangeGenericTexturePipeline.h b/Engine/Plugins/Interchange/Runtime/Source/Pipelines/Public/InterchangeGenericTexturePipeline.h
--- a/Engine/Plugins/Interchange/Runtime/Source/Pipelines/Public/InterchangeGenericTexturePipeline.h
+++ b/Engine/Plugins/Interchange/Runtime/Source/Pipelines/Public/InterchangeGenericTexturePipeline.h
@@ -30,7 +30,7 @@ public:
 	bool bImportTextures = true;
 
 	/** If set, and there is only one asset and one source, the imported asset will be given this name. */
-	UPROPERTY(EditAnywhere, BlueprintReadWrite, Category = "Textures", meta=(StandAlonePipelineProperty = "True", AlwaysResetToDefault = "True"))
+	UPROPERTY(EditAnywhere, BlueprintReadWrite, Category = "Textures", meta=(StandAlonePipelineProperty = "True"))
 	FString AssetName;
 
 #if WITH_EDITORONLY_DATA
diff --git a/Engine/Source/Runtime/Interchange/Core/Private/InterchangePipelineBase.cpp b/Engine/Source/Runtime/Interchange/Core/Private/InterchangePipelineBase.cpp
--- a/Engine/Source/Runtime/Interchange/Core/Private/InterchangePipelineBase.cpp
+++ b/Engine/Source/Runtime/Interchange/Core/Private/InterchangePipelineBase.cpp
@@ -25,9 +25,9 @@ namespace UE
 	}
 }
 
-void UInterchangePipelineBase::LoadSettings(const FName PipelineStackName)
+void UInterchangePipelineBase::LoadSettings(const FName PipelineStackName, bool bResetPreDialog /*= false*/)
 {
-	LoadSettingsInternal(PipelineStackName, GEditorPerProjectIni, PropertiesStates);
+	LoadSettingsInternal(PipelineStackName, GEditorPerProjectIni, PropertiesStates, bResetPreDialog);
 }
 
 void UInterchangePipelineBase::SaveSettings(const FName PipelineStackName)
@@ -104,7 +104,7 @@ FName UInterchangePipelineBase::GetResultsPropertyName()
 	return GET_MEMBER_NAME_CHECKED(UInterchangePipelineBase, Results);
 }
 
-void UInterchangePipelineBase::LoadSettingsInternal(const FName PipelineStackName, const FString& ConfigFilename, TMap<FName, FInterchangePipelinePropertyStates>& ParentPropertiesStates)
+void UInterchangePipelineBase::LoadSettingsInternal(const FName PipelineStackName, const FString& ConfigFilename, TMap<FName, FInterchangePipelinePropertyStates>& ParentPropertiesStates, bool bResetPreDialog)
 {
 	int32 PortFlags = 0;
 	UClass* Class = this->GetClass();
@@ -127,9 +127,11 @@ void UInterchangePipelineBase::LoadSettingsInternal(const FName PipelineStackNam
 #if WITH_EDITOR
 		if (Property->GetBoolMetaData(FName("AlwaysResetToDefault")))
 		{
-			//Stand alone pipeline property
+			//Not loading the property value will reset it
 			continue;
 		}
+		
+		
 #endif //WITH_EDITOR
 
 		if (const FInterchangePipelinePropertyStates* PropertyStates = ParentPropertiesStates.Find(PropertyPath))
@@ -139,6 +141,12 @@ void UInterchangePipelineBase::LoadSettingsInternal(const FName PipelineStackNam
 				//Skip this locked property
 				continue;
 			}
+
+			//Some property need to be reset only when loading the import dialog
+			if (bResetPreDialog && PropertyStates->IsPropertyPreDialogReset())
+			{
+				continue;
+			}
 		}
 
 		const bool bIsPropertyInherited = Property->GetOwnerClass() != Class;
@@ -201,7 +209,7 @@ void UInterchangePipelineBase::LoadSettingsInternal(const FName PipelineStackNam
 			// Load the settings if the referenced pipeline is a subobject of ours
 			if (SubPipeline->IsInOuter(this))
 			{
-				SubPipeline->LoadSettingsInternal(PipelineStackName, ConfigFilename, ParentPropertiesStates);
+				SubPipeline->LoadSettingsInternal(PipelineStackName, ConfigFilename, ParentPropertiesStates, bResetPreDialog);
 			}
 		}
 		else
diff --git a/Engine/Source/Runtime/Interchange/Core/Public/InterchangePipelineBase.h b/Engine/Source/Runtime/Interchange/Core/Public/InterchangePipelineBase.h
--- a/Engine/Source/Runtime/Interchange/Core/Public/InterchangePipelineBase.h
+++ b/Engine/Source/Runtime/Interchange/Core/Public/InterchangePipelineBase.h
@@ -78,7 +78,19 @@ struct FInterchangePipelinePropertyStates
 	{
 		bLocked = bLockValue;
 	}
-	
+
+	/** Return true if the property is locked. */
+	bool IsPropertyPreDialogReset() const
+	{
+		return bPreDialogReset;
+	}
+
+	void SetPropertyPreDialogReset(const bool bPreDialogResetValue)
+	{
+		bPreDialogReset = bPreDialogResetValue;
+	}
+
+
 	/** Return true if the property is visible for the specified context. */
 	bool IsPropertyVisibleInBasicLayout() const
 	{
@@ -115,6 +127,10 @@ struct FInterchangePipelinePropertyStates
 	UPROPERTY(EditAnywhere, BlueprintReadWrite, Category = "Property States")
 	bool bLocked = false;
 
+	/** If true, the property will be reset to default when loading the import dialog. */
+	UPROPERTY(EditAnywhere, BlueprintReadWrite, Category = "Property States")
+	bool bPreDialogReset = false;
+
 	/** The property states for the import context. */
 	UPROPERTY(EditAnywhere, BlueprintReadWrite, Category = "Context Properties States")
 	FInterchangePipelinePropertyStatePerContext BasicLayoutStates;
@@ -287,7 +303,7 @@ public:
 		return true;
 	}
 
-	INTERCHANGECORE_API void LoadSettings(const FName PipelineStackName);
+	INTERCHANGECORE_API void LoadSettings(const FName PipelineStackName, bool bResetPreDialog = false);
 
 	INTERCHANGECORE_API void SaveSettings(const FName PipelineStackName);
 
@@ -519,7 +535,7 @@ protected:
 	{
 	}
 
-	INTERCHANGECORE_API void LoadSettingsInternal(const FName PipelineStackName, const FString& ConfigFilename, TMap<FName, FInterchangePipelinePropertyStates>& ParentPropertiesStates);
+	INTERCHANGECORE_API void LoadSettingsInternal(const FName PipelineStackName, const FString& ConfigFilename, TMap<FName, FInterchangePipelinePropertyStates>& ParentPropertiesStates, bool bResetPreDialog);
 
 	INTERCHANGECORE_API void SaveSettingsInternal(const FName PipelineStackName, const FString& ConfigFilename);
 
