AssetTools: Add a IAssetTools::SanitizeName function. Interchange call the new sanitize code API when sanitizing a name, not when sanitizing a path.

#rnx

diff --git a/Engine/Plugins/Interchange/Editor/Source/InterchangeEditor/InterchangeEditor.Build.cs b/Engine/Plugins/Interchange/Editor/Source/InterchangeEditor/InterchangeEditor.Build.cs
--- a/Engine/Plugins/Interchange/Editor/Source/InterchangeEditor/InterchangeEditor.Build.cs
+++ b/Engine/Plugins/Interchange/Editor/Source/InterchangeEditor/InterchangeEditor.Build.cs
@@ -22,6 +22,7 @@ namespace UnrealBuildTool.Rules
 				new string[]
 				{
 					"AssetDefinition",
+					"AssetTools",
 					"ContentBrowser",
 					"InterchangeCommonParser",
 					"InterchangeCore",
diff --git a/Engine/Plugins/Interchange/Editor/Source/InterchangeEditor/Private/InterchangeEditorModule.cpp b/Engine/Plugins/Interchange/Editor/Source/InterchangeEditor/Private/InterchangeEditorModule.cpp
--- a/Engine/Plugins/Interchange/Editor/Source/InterchangeEditor/Private/InterchangeEditorModule.cpp
+++ b/Engine/Plugins/Interchange/Editor/Source/InterchangeEditor/Private/InterchangeEditorModule.cpp
@@ -1,13 +1,12 @@
 // Copyright Epic Games, Inc. All Rights Reserved.
 #include "InterchangeEditorModule.h"
 
-#include "InterchangeEditorLog.h"
-
-#include "InterchangeManager.h"
-#include "InterchangeFbxAssetImportDataConverter.h"
-
+#include "AssetToolsModule.h"
 #include "Engine/Engine.h"
 #include "IMessageLogListing.h"
+#include "InterchangeEditorLog.h"
+#include "InterchangeFbxAssetImportDataConverter.h"
+#include "InterchangeManager.h"
 #include "MessageLogModule.h"
 #include "Misc/App.h"
 #include "Misc/CoreDelegates.h"
@@ -83,7 +82,7 @@ namespace UE::Interchange::InterchangeEditorModule
 		UEditorLoadingSavingSettings* LoadingSavingSettings = GetMutableDefault<UEditorLoadingSavingSettings>();
 		LoadingSavingSettings->bAutoSaveEnable = bOldAutoSaveState;
 	}
-}
+} //ns UE::Interchange::InterchangeEditorModule
 
 FInterchangeEditorModule& FInterchangeEditorModule::Get()
 {
@@ -104,6 +103,7 @@ void FInterchangeEditorModule::StartupModule()
 		FDelegateHandle InterchangeEditorModuleDelegate;
 		FDelegateHandle InterchangeEditorModuleDelegateOnImportStarted;
 		FDelegateHandle InterchangeEditorModuleDelegateOnImportFinished;
+		FDelegateHandle InterchangeEditorModuleDelegateOnSanitizeName;
 
 		UInterchangeManager& InterchangeManager = UInterchangeManager::GetInterchangeManager();
 		InterchangeEditorModuleDelegate = InterchangeManager.OnBatchImportComplete.AddStatic(&InterchangeEditorModule::LogErrors);
@@ -111,12 +111,23 @@ void FInterchangeEditorModule::StartupModule()
 		InterchangeEditorModuleDelegateOnImportFinished = InterchangeManager.OnImportFinished.AddStatic(&InterchangeEditorModule::ImportFinished);
 		InterchangeManager.RegisterImportDataConverter(UInterchangeFbxAssetImportDataConverter::StaticClass());
 
-		auto UnregisterItems = [InterchangeEditorModuleDelegate, InterchangeEditorModuleDelegateOnImportStarted, InterchangeEditorModuleDelegateOnImportFinished]()
+		InterchangeEditorModuleDelegateOnSanitizeName = InterchangeManager.OnSanitizeName.AddLambda([](FString& SanitizeName, const ESanitizeNameTypeFlags NameType)
+			{
+				//Call the asset tools sanitize
+				IAssetTools& AssetTools = FModuleManager::LoadModuleChecked<FAssetToolsModule>("AssetTools").Get();
+				AssetTools.SanitizeName(SanitizeName);
+			});
+
+		auto UnregisterItems = [InterchangeEditorModuleDelegate
+			, InterchangeEditorModuleDelegateOnImportStarted
+			, InterchangeEditorModuleDelegateOnImportFinished
+			, InterchangeEditorModuleDelegateOnSanitizeName]()
 		{
 			UInterchangeManager& InterchangeManager = UInterchangeManager::GetInterchangeManager();
 			InterchangeManager.OnBatchImportComplete.Remove(InterchangeEditorModuleDelegate);
 			InterchangeManager.OnImportStarted.Remove(InterchangeEditorModuleDelegateOnImportStarted);
 			InterchangeManager.OnImportFinished.Remove(InterchangeEditorModuleDelegateOnImportFinished);
+			InterchangeManager.OnSanitizeName.Remove(InterchangeEditorModuleDelegateOnSanitizeName);
 		};
 
 		InterchangeManager.OnPreDestroyInterchangeManager.AddLambda(UnregisterItems);
diff --git a/Engine/Source/Developer/AssetTools/Private/AssetTools.cpp b/Engine/Source/Developer/AssetTools/Private/AssetTools.cpp
--- a/Engine/Source/Developer/AssetTools/Private/AssetTools.cpp
+++ b/Engine/Source/Developer/AssetTools/Private/AssetTools.cpp
@@ -6515,6 +6515,27 @@ void UAssetToolsImpl::UnregisterIsNameAllowedDelegate(const FName OwnerName)
 	IsNameAllowedDelegates.Remove(OwnerName);
 }
 
+bool UAssetToolsImpl::SanitizeName(FString& NameToSanitize)
+{
+	//Call the delegates
+	for (const TPair<FName, FSanitizeName>& DelegatePair : SanitizeNameDelegates)
+	{
+		DelegatePair.Value.Execute(NameToSanitize);
+	}
+
+	return ensure(IsNameAllowed(NameToSanitize, nullptr));
+}
+
+void UAssetToolsImpl::RegisterSanitizeNameDelegate(const FName OwnerName, FSanitizeName Delegate)
+{
+	SanitizeNameDelegates.Add(OwnerName, Delegate);
+}
+
+void UAssetToolsImpl::UnregisterSanitizeNameDelegate(const FName OwnerName)
+{
+	SanitizeNameDelegates.Remove(OwnerName);
+}
+
 void UAssetToolsImpl::RegisterCanMigrateAsset(const FName OwnerName, UE::AssetTools::FCanMigrateAsset Delegate)
 {
 	CanMigrateAssetDelegates.Add(OwnerName, MoveTemp(Delegate));
diff --git a/Engine/Source/Developer/AssetTools/Private/AssetTools.h b/Engine/Source/Developer/AssetTools/Private/AssetTools.h
--- a/Engine/Source/Developer/AssetTools/Private/AssetTools.h
+++ b/Engine/Source/Developer/AssetTools/Private/AssetTools.h
@@ -150,6 +150,9 @@ public:
 	virtual bool IsNameAllowed(const FString& Name, FText* OutErrorMessage = nullptr) const override;
 	virtual void RegisterIsNameAllowedDelegate(const FName OwnerName, FIsNameAllowed Delegate) override;
 	virtual void UnregisterIsNameAllowedDelegate(const FName OwnerName) override;
+	virtual bool SanitizeName(FString& NameToSanitize) override;
+	virtual void RegisterSanitizeNameDelegate(const FName OwnerName, FSanitizeName Delegate) override;
+	virtual void UnregisterSanitizeNameDelegate(const FName OwnerName) override;
 	virtual void RegisterCanMigrateAsset(const FName OwnerName, UE::AssetTools::FCanMigrateAsset Delegate) override;
 	virtual void UnregisterCanMigrateAsset(const FName OwnerName) override;
 	virtual bool CanAssetBePublic(FStringView AssetPath) const override;
@@ -272,6 +275,7 @@ private:
 	bool CreateAssetsAsExternallyReferenceable;
 
 	TMap<FName, FIsNameAllowed> IsNameAllowedDelegates;
+	TMap<FName, FSanitizeName> SanitizeNameDelegates;
 
 	UE::AssetTools::FOnPackageMigration OnPackageMigration;
 
diff --git a/Engine/Source/Developer/AssetTools/Public/IAssetTools.h b/Engine/Source/Developer/AssetTools/Public/IAssetTools.h
--- a/Engine/Source/Developer/AssetTools/Public/IAssetTools.h
+++ b/Engine/Source/Developer/AssetTools/Public/IAssetTools.h
@@ -132,7 +132,7 @@ DECLARE_DYNAMIC_DELEGATE_TwoParams(FAdvancedCopyCompletedEvent, bool, bSuccess,
 
 DECLARE_MULTICAST_DELEGATE_OneParam(FAssetPostRenameEvent, const TArray<FAssetRenameData>&);
 DECLARE_DELEGATE_RetVal_TwoParams(bool, FIsNameAllowed, const FString& /*Name*/, FText* /*OutErrorMessage*/);
-
+DECLARE_DELEGATE_OneParam(FSanitizeName, FString& /*NameToSanitize*/);
 
 struct FAdvancedAssetCategory
 {
@@ -637,6 +637,13 @@ public:
 	/** Remove a previously-set global name filter */
 	virtual void UnregisterIsNameAllowedDelegate(const FName OwnerName) = 0;
 
+	/** Sanitize the name by calling all registered delegates. */
+	virtual bool SanitizeName(FString& NameToSanitize) = 0;
+	/** Allows setting of a global name filter that is applied to assets and folders */
+	virtual void RegisterSanitizeNameDelegate(const FName OwnerName, FSanitizeName Delegate) = 0;
+	/** Remove a previously-set global name filter */
+	virtual void UnregisterSanitizeNameDelegate(const FName OwnerName) = 0;
+
 	/** Show notification that writable folder filter blocked an action */
 	virtual void NotifyBlockedByWritableFolderFilter() const = 0;
 
diff --git a/Engine/Source/Runtime/Interchange/Engine/Private/InterchangeManager.cpp b/Engine/Source/Runtime/Interchange/Engine/Private/InterchangeManager.cpp
--- a/Engine/Source/Runtime/Interchange/Engine/Private/InterchangeManager.cpp
+++ b/Engine/Source/Runtime/Interchange/Engine/Private/InterchangeManager.cpp
@@ -809,27 +809,6 @@ void UE::Interchange::FImportResult::AddReferencedObjects(FReferenceCollector& C
 	Collector.AddReferencedObject(Results);
 }
 
-void UE::Interchange::SanitizeObjectPath(FString& ObjectPath)
-{
-	const TCHAR* InvalidChar = INVALID_OBJECTPATH_CHARACTERS INVALID_LONGPACKAGE_CHARACTERS;
-
-	while (*InvalidChar)
-	{
-		ObjectPath.ReplaceCharInline(*InvalidChar, TCHAR('_'), ESearchCase::CaseSensitive);
-		++InvalidChar;
-	}
-}
-
-void UE::Interchange::SanitizeObjectName(FString& ObjectName)
-{
-	const TCHAR* InvalidChar = INVALID_OBJECTNAME_CHARACTERS;
-	while (*InvalidChar)
-	{
-		ObjectName.ReplaceCharInline(*InvalidChar, TCHAR('_'), ESearchCase::CaseSensitive);
-		++InvalidChar;
-	}
-}
-
 UInterchangePipelineBase* UE::Interchange::GeneratePipelineInstance(const FSoftObjectPath& PipelineInstance)
 {
 	UPackage* PipelineInstancePackage = GetTransientPackage();
@@ -1700,8 +1679,8 @@ UInterchangeManager::ImportInternal(const FString& ContentPath, const UInterchan
 	{
 		FString AssetpackageName = ImportAssetParameters.DestinationName.IsEmpty() ? FPaths::GetCleanFilename(SourceData->GetFilename()) : ImportAssetParameters.DestinationName;
 		FString PackageBasePath = ContentPath;
-		UE::Interchange::SanitizeObjectName(AssetpackageName);
-		UE::Interchange::SanitizeObjectPath(PackageBasePath);
+		SanitizeNameInline(AssetpackageName, ESanitizeNameTypeFlags::ObjectName | ESanitizeNameTypeFlags::ObjectPath | ESanitizeNameTypeFlags::LongPackage);
+		SanitizeNameInline(PackageBasePath, ESanitizeNameTypeFlags::ObjectPath | ESanitizeNameTypeFlags::LongPackage);
 		FString FullPackagePath = FPaths::Combine(*PackageBasePath, *AssetpackageName);
 		if (!UE::Interchange::FPackageUtils::IsMapPackageAsset(FullPackagePath))
 		{
@@ -1776,7 +1755,7 @@ UInterchangeManager::ImportInternal(const FString& ContentPath, const UInterchan
 	FString ContentBasePath = ContentPath;
 	if (!ImportAssetParameters.ReimportAsset)
 	{
-		UE::Interchange::SanitizeObjectPath(ContentBasePath);
+		SanitizeNameInline(ContentBasePath, ESanitizeNameTypeFlags::ObjectPath | ESanitizeNameTypeFlags::LongPackage);
 	}
 	else
 	{
@@ -2316,6 +2295,51 @@ bool UInterchangeManager::IsImporting()
 	return UE::Interchange::Private::StaticHelpers::AsyncHelperCounter > 0;
 }
 
+void UInterchangeManager::SanitizeNameInline(FString& NameToSanitize, const ESanitizeNameTypeFlags NameType)
+{
+	//Default behavior don't use the delegates
+	auto SanitizeCharacters = [&NameToSanitize](const TCHAR* InvalidCharacters)
+		{
+			while (*InvalidCharacters)
+			{
+				NameToSanitize.ReplaceCharInline(*InvalidCharacters, TCHAR('_'), ESearchCase::CaseSensitive);
+				++InvalidCharacters;
+			}
+		};
+
+	//Sanitize all Name type users ask for
+
+	bool bBroadcastSanitizeName = false;
+	if ((NameType & ESanitizeNameTypeFlags::Name) != ESanitizeNameTypeFlags::None)
+	{
+		SanitizeCharacters(INVALID_NAME_CHARACTERS);
+		bBroadcastSanitizeName = true;
+	}
+
+	if ((NameType & ESanitizeNameTypeFlags::ObjectName) != ESanitizeNameTypeFlags::None)
+	{
+		SanitizeCharacters(INVALID_OBJECTNAME_CHARACTERS);
+		bBroadcastSanitizeName = true;
+	}
+
+	if ((NameType & ESanitizeNameTypeFlags::ObjectPath) != ESanitizeNameTypeFlags::None)
+	{
+		SanitizeCharacters(INVALID_OBJECTPATH_CHARACTERS);
+	}
+
+	if ((NameType & ESanitizeNameTypeFlags::LongPackage) != ESanitizeNameTypeFlags::None)
+	{
+		SanitizeCharacters(INVALID_LONGPACKAGE_CHARACTERS);
+	}
+
+	if (bBroadcastSanitizeName)
+	{
+		//Call the delegates which will add extra sanitize rules, this is a thread safe delegate
+		OnSanitizeName.Broadcast(NameToSanitize, NameType);
+	}
+}
+
+
 bool UInterchangeManager::ExportAsset(const UObject* Asset, bool bIsAutomated)
 {
 	return false;
diff --git a/Engine/Source/Runtime/Interchange/Engine/Private/Tasks/InterchangeTaskCreateSceneObjects.cpp b/Engine/Source/Runtime/Interchange/Engine/Private/Tasks/InterchangeTaskCreateSceneObjects.cpp
--- a/Engine/Source/Runtime/Interchange/Engine/Private/Tasks/InterchangeTaskCreateSceneObjects.cpp
+++ b/Engine/Source/Runtime/Interchange/Engine/Private/Tasks/InterchangeTaskCreateSceneObjects.cpp
@@ -87,7 +87,7 @@ void UE::Interchange::FTaskCreateSceneObjects::DoTask(ENamedThreads::Type Curren
 		AsyncHelper->AddCreatedFactory(FactoryNode->GetUniqueID(), Factory);
 
 		FString SceneNodeName = FactoryNode->GetAssetName();
-		SanitizeObjectName(SceneNodeName);
+		UInterchangeManager::GetInterchangeManager().SanitizeNameInline(SceneNodeName, ESanitizeNameTypeFlags::ObjectName | ESanitizeNameTypeFlags::ObjectPath | ESanitizeNameTypeFlags::LongPackage);
 
 		UInterchangeFactoryBase::FImportSceneObjectsParams CreateSceneObjectsParams;
 		CreateSceneObjectsParams.ObjectName = SceneNodeName;
diff --git a/Engine/Source/Runtime/Interchange/Engine/Private/Tasks/InterchangeTaskImportObject.cpp b/Engine/Source/Runtime/Interchange/Engine/Private/Tasks/InterchangeTaskImportObject.cpp
--- a/Engine/Source/Runtime/Interchange/Engine/Private/Tasks/InterchangeTaskImportObject.cpp
+++ b/Engine/Source/Runtime/Interchange/Engine/Private/Tasks/InterchangeTaskImportObject.cpp
@@ -33,22 +33,16 @@ namespace UE::Interchange::Private
 		const UInterchangeSourceData* SourceData = AsyncHelper.SourceDatas[SourceIndex];
 		check(SourceData);
 		FString NodeDisplayName = FactoryNode->GetAssetName();
-
+		
 		// Set the asset name and the package name
 		OutAssetName = NodeDisplayName;
-		SanitizeObjectName(OutAssetName);
-
-		FString SanitizedPackageBasePath = PackageBasePath;
-		SanitizeObjectPath(SanitizedPackageBasePath);
-
+		UInterchangeManager::GetInterchangeManager().SanitizeNameInline(OutAssetName, ESanitizeNameTypeFlags::ObjectName | ESanitizeNameTypeFlags::ObjectPath | ESanitizeNameTypeFlags::LongPackage);
 		FString SubPath;
-		if (FactoryNode->GetCustomSubPath(SubPath))
-		{
-			SanitizeObjectPath(SubPath);
-		}
-
-		OutPackageName = FPaths::Combine(*SanitizedPackageBasePath, *SubPath, *OutAssetName);
+		FactoryNode->GetCustomSubPath(SubPath);
+		OutPackageName = FPaths::Combine(*PackageBasePath, *SubPath, *OutAssetName);
+		UInterchangeManager::GetInterchangeManager().SanitizeNameInline(OutPackageName, ESanitizeNameTypeFlags::ObjectPath | ESanitizeNameTypeFlags::LongPackage);
 	}
+
 	bool ShouldReimportFactoryNode(UInterchangeFactoryBaseNode* FactoryNode, const UInterchangeBaseNodeContainer* NodeContainer, UObject* ReimportObject)
 	{
 		TRACE_CPUPROFILER_EVENT_SCOPE(UE::Interchange::Private::ShouldReimportFactoryNode)
diff --git a/Engine/Source/Runtime/Interchange/Engine/Public/InterchangeManager.h b/Engine/Source/Runtime/Interchange/Engine/Public/InterchangeManager.h
--- a/Engine/Source/Runtime/Interchange/Engine/Public/InterchangeManager.h
+++ b/Engine/Source/Runtime/Interchange/Engine/Public/InterchangeManager.h
@@ -48,6 +48,19 @@ DECLARE_MULTICAST_DELEGATE(FOnImportStarted);
 /** Delegate Type that is fired when Interchange finished importing. Won't fire when an import process finishes while one is still in progress. Primarily Used for AutoSave setting. */
 DECLARE_MULTICAST_DELEGATE(FOnImportFinished);
 
+enum class ESanitizeNameTypeFlags : uint8
+{
+	None = 0x00,
+	Name = 0x01,
+	ObjectName = 0x02,
+	ObjectPath = 0x04,
+	LongPackage = 0x08
+};
+ENUM_CLASS_FLAGS(ESanitizeNameTypeFlags)
+
+//Thread safe delegate since this can be broadcast in any thread
+DECLARE_TS_MULTICAST_DELEGATE_TwoParams(FOnSanitizeName, FString& /*NameToSanitize*/, const ESanitizeNameTypeFlags /*TypeName*/);
+
 namespace UE
 {
 	namespace Interchange
@@ -312,13 +325,9 @@ namespace UE
 			TMap<int32, TArray<FImportedObjectInfo>> ImportedSceneObjectsPerSourceIndex;
 		};
 
-		void SanitizeObjectPath(FString& ObjectPath);
-
-		void SanitizeObjectName(FString& ObjectName);
-
 		/* This function takes an asset that represents a pipeline and generates a UInterchangePipelineBase asset. */
 		INTERCHANGEENGINE_API UInterchangePipelineBase* GeneratePipelineInstance(const FSoftObjectPath& PipelineInstance);
-	} //ns interchange
+	} //ns Interchange
 } //ns UE
 
 /**
@@ -457,9 +466,14 @@ public:
 	FOnImportStarted OnImportStarted;
 	//Fires when the last import process finishes.
 	FOnImportFinished OnImportFinished;
+
+	//Fire when we need to sanitize a name, make sure your delegate code is thread safe, since it will be broadcast in any thread.
+	FOnSanitizeName OnSanitizeName;
+
 	// Called when before the application is exiting.
 	FSimpleMulticastDelegate OnPreDestroyInterchangeManager;
 
+	INTERCHANGEENGINE_API void SanitizeNameInline(FString& NameToSanitize, const ESanitizeNameTypeFlags NameType);
 	/**
 	 * All translators must be registered with the manager.
 	 * @Param Translator - The UClass of the translator you want to register.
