Custom morph target import

#rnx

diff --git a/Engine/Plugins/Interchange/Runtime/Source/Pipelines/Private/InterchangeGenericMeshPipeline.cpp b/Engine/Plugins/Interchange/Runtime/Source/Pipelines/Private/InterchangeGenericMeshPipeline.cpp
--- a/Engine/Plugins/Interchange/Runtime/Source/Pipelines/Private/InterchangeGenericMeshPipeline.cpp
+++ b/Engine/Plugins/Interchange/Runtime/Source/Pipelines/Private/InterchangeGenericMeshPipeline.cpp
@@ -61,7 +61,9 @@ void UInterchangeGenericMeshPipeline::AdjustSettingsForContext(EInterchangePipel
 	if (ImportType == EInterchangePipelineContext::AssetCustomLODImport
 		|| ImportType == EInterchangePipelineContext::AssetCustomLODReimport
 		|| ImportType == EInterchangePipelineContext::AssetAlternateSkinningImport
-		|| ImportType == EInterchangePipelineContext::AssetAlternateSkinningReimport)
+		|| ImportType == EInterchangePipelineContext::AssetAlternateSkinningReimport
+		|| ImportType == EInterchangePipelineContext::AssetCustomMorphTargetImport
+		|| ImportType == EInterchangePipelineContext::AssetCustomMorphTargetReImport)
 	{
 		bCreatePhysicsAsset = false;
 		PhysicsAsset = nullptr;
@@ -87,6 +89,28 @@ void UInterchangeGenericMeshPipeline::AdjustSettingsForContext(EInterchangePipel
 			CommonSkeletalMeshesAndAnimationsProperties->Skeleton = nullptr;
 			CommonSkeletalMeshesAndAnimationsProperties->bImportOnlyAnimations = false;
 		}
+		else if (ImportType == EInterchangePipelineContext::AssetCustomMorphTargetImport
+			|| ImportType == EInterchangePipelineContext::AssetCustomMorphTargetReImport)
+		{
+			//Custom morph target are imported has a combined static mesh
+			CommonMeshesProperties->ForceAllMeshAsType = EInterchangeForceMeshType::IFMT_StaticMesh;
+			CommonMeshesProperties->bAutoDetectMeshType = false;
+			CommonMeshesProperties->bBakeMeshes = true;
+			CommonMeshesProperties->bBakePivotMeshes = false;
+			CommonMeshesProperties->bImportLods = true;
+			CommonMeshesProperties->bKeepSectionsSeparate = false;
+			CommonMeshesProperties->VertexColorImportOption = EInterchangeVertexColorImportOption::IVCIO_Ignore;
+			bImportSkeletalMeshes = false;
+			bImportStaticMeshes = true;
+			bCombineStaticMeshes = true;
+			bBuildNanite = false;
+			LodGroup = NAME_None;
+			bImportCollision = false;
+			bImportCollisionAccordingToMeshName = false;
+			bGenerateLightmapUVs = false;
+			bGenerateDistanceFieldAsIfTwoSided = false;
+			bSupportFaceRemap = false;
+		}
 	}
 	const FString CommonMeshesCategory = UInterchangeGenericCommonMeshesProperties::GetPipelineCategory(nullptr);
 	const FString StaticMeshesCategory = UInterchangeGenericMeshPipeline::GetPipelineCategory(UStaticMesh::StaticClass());
diff --git a/Engine/Source/Developer/SkeletalMeshUtilitiesCommon/Private/LODUtilities.cpp b/Engine/Source/Developer/SkeletalMeshUtilitiesCommon/Private/LODUtilities.cpp
--- a/Engine/Source/Developer/SkeletalMeshUtilitiesCommon/Private/LODUtilities.cpp
+++ b/Engine/Source/Developer/SkeletalMeshUtilitiesCommon/Private/LODUtilities.cpp
@@ -3670,6 +3670,7 @@ static void BuildMorphTargetsInternal(
 
 	// Prepare base data
 	FSkeletalMeshLODModel& BaseLODModel = BaseSkelMesh->GetImportedModel()->LODModels[LODIndex];
+	const FSkeletalMeshLODInfo* BaseLodInfo = BaseSkelMesh->GetLODInfo(LODIndex);
 
 	FMeshDataBundle MeshDataBundle;
 	ConvertImportDataToMeshData(BaseImportData, MeshDataBundle);
@@ -3874,6 +3875,21 @@ static void BuildMorphTargetsInternal(
 		delete Results[Index];
 		Results[Index] = nullptr;
 
+		//Reset the morph target imported source filename
+		MorphTarget->SetCustomImportedSourceFilename(LODIndex, FString());
+		//Now set the imported source filename from the imported data
+		if (BaseLodInfo)
+		{
+			if (const FMorphTargetImportedSourceFileInfo* MorphTargetImportedSourceFileInfo = BaseLodInfo->ImportedMorphTargetSourceFilename.Find(MorphTarget->GetName()))
+			{
+				const FString& SourceFilename = MorphTargetImportedSourceFileInfo->GetSourceFilename();
+				if (!SourceFilename.IsEmpty())
+				{
+					MorphTarget->SetCustomImportedSourceFilename(LODIndex, *SourceFilename);
+				}
+			}
+		}
+
 		// We might have created new MorphTarget in an async thread, so we need to remove the async flag so they can get
 		// garbage collected in the future now that their references are properly setup and reachable by the GC.
 		MorphTarget->ClearInternalFlags(EInternalObjectFlags::Async);
diff --git a/Engine/Source/Editor/Persona/Persona.Build.cs b/Engine/Source/Editor/Persona/Persona.Build.cs
--- a/Engine/Source/Editor/Persona/Persona.Build.cs
+++ b/Engine/Source/Editor/Persona/Persona.Build.cs
@@ -60,6 +60,7 @@ public class Persona : ModuleRules
                 "EditorWidgets",
                 "BlueprintGraph",
                 "RHI",
+				"InterchangeEngine",
                 "Json",
                 "JsonUtilities",
                 "ClothingSystemEditorInterface",
diff --git a/Engine/Source/Editor/Persona/Private/SMorphTargetViewer.cpp b/Engine/Source/Editor/Persona/Private/SMorphTargetViewer.cpp
--- a/Engine/Source/Editor/Persona/Private/SMorphTargetViewer.cpp
+++ b/Engine/Source/Editor/Persona/Private/SMorphTargetViewer.cpp
@@ -1,24 +1,28 @@
 // Copyright Epic Games, Inc. All Rights Reserved.
 
-
 #include "SMorphTargetViewer.h"
-#include "Widgets/Input/SCheckBox.h"
-#include "Framework/Commands/UIAction.h"
-#include "Textures/SlateIcon.h"
-#include "Framework/MultiBox/MultiBoxBuilder.h"
-#include "Widgets/Input/SSpinBox.h"
+
+#include "Animation/AnimInstance.h"
 #include "Animation/DebugSkelMeshComponent.h"
-#include "ScopedTransaction.h"
-#include "Widgets/Input/SSearchBox.h"
 #include "Animation/MorphTarget.h"
-#include "Animation/AnimInstance.h"
-#include "HAL/PlatformApplicationMisc.h"
-#include "GPUSkinCache.h"
 #include "Engine/RendererSettings.h"
+#include "Framework/Commands/UIAction.h"
+#include "Framework/MultiBox/MultiBoxBuilder.h"
+#include "GPUSkinCache.h"
+#include "HAL/PlatformApplicationMisc.h"
 #include "IPersonaPreviewScene.h"
-#include "SkeletalMeshAttributes.h"
+#include "InterchangeMeshUtilities.h"
+#include "InterchangeManager.h"
 #include "Rendering/SkeletalMeshLODImporterData.h"
+#include "ScopedTransaction.h"
+#include "SkeletalMeshAttributes.h"
 #include "SkeletalRenderPublic.h"
+#include "SkinnedAssetCompiler.h"
+#include "Textures/SlateIcon.h"
+#include "Widgets/Input/SButton.h"
+#include "Widgets/Input/SCheckBox.h"
+#include "Widgets/Input/SSearchBox.h"
+#include "Widgets/Input/SSpinBox.h"
 
 #define LOCTEXT_NAMESPACE "SMorphTargetViewer"
 
@@ -125,6 +129,28 @@ TSharedRef< SWidget > SMorphTargetListRow::GenerateWidgetForColumn( const FName&
 {
 	if ( ColumnName == ColumnId_MorphTargetNameLabel )
 	{
+		FText SourceFilenamesTooltip;
+		FText MorphNameText = FText::FromName(Item->Name);
+
+		bool bFirst = true;
+		for (int32 LodIndex = 0; LodIndex < Item->SourceFilenames.Num(); ++LodIndex)
+		{
+			if (!Item->SourceFilenames[LodIndex].IsEmpty())
+			{
+				FText CarriageReturn;
+				if (bFirst)
+				{
+					MorphNameText = FText::Format(LOCTEXT("MorphRowName", "{0} (Imported)"), MorphNameText);
+					bFirst = false;
+				}
+				else
+				{
+					CarriageReturn = bFirst ? FText() : FText::FromString(TEXT("\n"));
+				}
+				SourceFilenamesTooltip = FText::Format(LOCTEXT("{0}{1}SourceFilenameTooltipEntry", "LOD {2} Source Filename: {3}"), CarriageReturn, SourceFilenamesTooltip, FText::AsNumber(LodIndex), FText::FromString(Item->SourceFilenames[LodIndex]));
+			}
+		}
+
 		return
 			SNew( SVerticalBox )
 
@@ -134,7 +160,8 @@ TSharedRef< SWidget > SMorphTargetListRow::GenerateWidgetForColumn( const FName&
 			.VAlign( VAlign_Center )
 			[
 				SNew( STextBlock )
-				.Text( FText::FromName(Item->Name) )
+				.Text( MorphNameText )
+				.ToolTipText(SourceFilenamesTooltip)
 				.HighlightText( MorphTargetViewer->GetFilterText() )
 			];
 	}
@@ -352,7 +379,24 @@ void SMorphTargetViewer::Construct(const FArguments& InArgs, const TSharedRef<IP
 			SNew( STextBlock )
 			.Text( SkeletalMeshName )
 		]
-
+		+ SVerticalBox::Slot()
+		.AutoHeight()
+		.Padding(0,2)
+		[
+			SNew(SHorizontalBox)
+			// Import morph target
+			+SHorizontalBox::Slot()
+			.AutoWidth()
+			[
+				SNew( SButton )
+				.OnClicked(this, &SMorphTargetViewer::OnImportMorphTargetButton)
+				[
+					SNew(STextBlock)
+					.ToolTipText(LOCTEXT("ImportCustomMorphTargetButtonTooltip", "Import a new morph target from a file."))
+					.Text(LOCTEXT("ImportCustomMorphTargetButtonText", "Import Morph Target"))
+				]
+			]
+		]
 		+ SVerticalBox::Slot()
 		.AutoHeight()
 		.Padding(0,2)
@@ -440,24 +484,137 @@ TSharedPtr<SWidget> SMorphTargetViewer::OnGetContextMenuContent() const
 
 	MenuBuilder.BeginSection("MorphTargetAction", LOCTEXT( "MorphsAction", "Selected Item Actions" ) );
 	{
-		FUIAction Action;
+		TArray< TSharedPtr< FDisplayedMorphTargetInfo > > SelectedRows = MorphTargetListView->GetSelectedItems();
+		const int32 SelectionCount = SelectedRows.Num();
 
+		if (SelectionCount > 0)
 		{
-			Action.ExecuteAction = FExecuteAction::CreateSP(const_cast<SMorphTargetViewer*>(this), &SMorphTargetViewer::OnDeleteMorphTargets);
-			Action.CanExecuteAction = FCanExecuteAction::CreateSP(this, &SMorphTargetViewer::CanPerformDelete);
-			const FText Label = LOCTEXT("DeleteMorphTargetButtonLabel", "Delete");
-			const FText ToolTipText = LOCTEXT("DeleteMorphTargetButtonTooltip", "Deletes the selected morph targets.");
-			MenuBuilder.AddMenuEntry(Label, ToolTipText, FSlateIcon(), Action);
-		}
+			const int32 LodCount = SkeletalMesh->GetLODNum();
+			bool bShowImportMenu = false;
+			struct FLodMorphTargetImportMenuInfo
+			{
+				bool bShowImportMenu = true;
+				bool bShowReimportMenu = true;
+				bool bShowReimportWithNewFileMenu = true;
+				bool IsMenuShow() const { return bShowImportMenu || bShowReimportMenu || bShowReimportWithNewFileMenu; }
+				void HideMenus()
+				{
+					bShowImportMenu = false;
+					bShowReimportMenu = false;
+					bShowReimportWithNewFileMenu = false;
+				}
+			};
+			TMap<int32, FLodMorphTargetImportMenuInfo> MenuInforPerLods;
+			for (int32 LodIndex = 0; LodIndex < LodCount; ++LodIndex)
+			{
+				FLodMorphTargetImportMenuInfo& MenuInfo = MenuInforPerLods.FindOrAdd(LodIndex);
+				if (!SkeletalMesh->HasMeshDescription(LodIndex))
+				{
+					MenuInfo.HideMenus();
+					continue;
+				}
 
-		{
-			Action.ExecuteAction = FExecuteAction::CreateSP(const_cast<SMorphTargetViewer*>(this), &SMorphTargetViewer::OnCopyMorphTargetNames);
-			Action.CanExecuteAction = nullptr;
-			const FText Label = LOCTEXT("CopyMorphTargetNamesButtonLabel", "Copy Names");
-			const FText ToolTipText = LOCTEXT("CopyMorphTargetNamesButtonTooltip", "Copy the names of selected morph targets to clipboard");
-			MenuBuilder.AddMenuEntry(Label, ToolTipText, FSlateIcon(), Action);
-		}
+				for (int32 RowIndex = 0; RowIndex < SelectionCount; ++RowIndex)
+				{
+					UMorphTarget* MorphTarget = SkeletalMesh->FindMorphTarget(SelectedRows[RowIndex]->Name);
+					if (MorphTarget)
+					{
+						//Look at the lod 0 to see if the morph target is an imported morph target
+						constexpr int32 LodIndex0 = 0;
+						if (!MorphTarget->IsCustomImported(LodIndex0))
+						{
+							MenuInfo.bShowImportMenu = false;
+							MenuInfo.bShowReimportMenu = false;
+							MenuInfo.bShowReimportWithNewFileMenu = false;
+						}
+						else
+						{
+							const bool bIsCustomImportedLod = MorphTarget->IsCustomImported(LodIndex);
+							MenuInfo.bShowImportMenu &= (SelectionCount == 1 && !bIsCustomImportedLod);
+							MenuInfo.bShowReimportMenu &= bIsCustomImportedLod;
+							MenuInfo.bShowReimportWithNewFileMenu &= (SelectionCount == 1 && bIsCustomImportedLod);
+						}
+					}
+				}
+				bShowImportMenu |= MenuInfo.IsMenuShow();
+			}
+
+			if (bShowImportMenu)
+			{
+				//Create the import menu for every lods
+				for (int32 LodIndex = 0; LodIndex < LodCount; ++LodIndex)
+				{
+					const FLodMorphTargetImportMenuInfo& MenuInfo = MenuInforPerLods.FindOrAdd(LodIndex);
 
+					//We can import a morph only if the lod is custom imported
+					if (MenuInfo.IsMenuShow())
+					{
+						FText SubMenuLabel = FText::Format(LOCTEXT("LodSubMenu", "LOD {0}"), FText::AsNumber(LodIndex));
+						MenuBuilder.AddSubMenu(SubMenuLabel,
+							FText(),
+							FNewMenuDelegate::CreateLambda([this, LodIndex, MenuInfo](FMenuBuilder& InSubMenuBuilder)
+								{
+									FUIAction Action;
+									
+									//Import Morph target
+									if (MenuInfo.bShowImportMenu)
+									{
+										Action.ExecuteAction = FExecuteAction::CreateSP(const_cast<SMorphTargetViewer*>(this), &SMorphTargetViewer::OnReimportMorphTargets, LodIndex);
+										Action.CanExecuteAction = nullptr;
+										const FText Label = LOCTEXT("ImportMorphTargetLabel", "Import");
+										const FText ToolTipText = LOCTEXT("ImportMorphTargetTooltip", "Import all selected custom imported morph target");
+										InSubMenuBuilder.AddMenuEntry(Label, ToolTipText, FSlateIcon(), Action);
+									}
+									//Reimport Morph target
+									if (MenuInfo.bShowReimportMenu)
+									{
+										Action.ExecuteAction = FExecuteAction::CreateSP(const_cast<SMorphTargetViewer*>(this), &SMorphTargetViewer::OnReimportMorphTargets, LodIndex);
+										Action.CanExecuteAction = nullptr;
+										const FText Label = LOCTEXT("ReimportMorphTargetLabel", "Reimport");
+										const FText ToolTipText = LOCTEXT("ReimportMorphTargetTooltip", "Reimport all selected custom imported morph target");
+										InSubMenuBuilder.AddMenuEntry(Label, ToolTipText, FSlateIcon(), Action);
+									}
+
+									//Reimport Morph target with new file
+									if (MenuInfo.bShowReimportWithNewFileMenu)
+									{
+										Action.ExecuteAction = FExecuteAction::CreateSP(const_cast<SMorphTargetViewer*>(this), &SMorphTargetViewer::OnReimportMorphTargetsWithNewFile, LodIndex);
+										Action.CanExecuteAction = nullptr;
+										const FText Label = LOCTEXT("ReimportWithNewFileMorphTargetLabel", "Reimport With New File");
+										const FText ToolTipText = LOCTEXT("ReimportMorphTargetTooltip", "Ask a file and re-import every selected morph target.");
+										InSubMenuBuilder.AddMenuEntry(Label, ToolTipText, FSlateIcon(), Action);
+									}
+								})
+						);
+					}
+				}
+
+				MenuBuilder.AddMenuSeparator();
+			}
+
+			//Basic morph target context menu
+			{
+				FUIAction Action;
+
+				//Delete morph target
+				{
+					Action.ExecuteAction = FExecuteAction::CreateSP(const_cast<SMorphTargetViewer*>(this), &SMorphTargetViewer::OnDeleteMorphTargets);
+					Action.CanExecuteAction = FCanExecuteAction::CreateSP(this, &SMorphTargetViewer::CanPerformDelete);
+					const FText Label = LOCTEXT("DeleteMorphTargetButtonLabel", "Delete");
+					const FText ToolTipText = LOCTEXT("DeleteMorphTargetButtonTooltip", "Deletes the selected morph targets.");
+					MenuBuilder.AddMenuEntry(Label, ToolTipText, FSlateIcon(), Action);
+				}
+
+				//Copy morph target name
+				{
+					Action.ExecuteAction = FExecuteAction::CreateSP(const_cast<SMorphTargetViewer*>(this), &SMorphTargetViewer::OnCopyMorphTargetNames);
+					Action.CanExecuteAction = nullptr;
+					const FText Label = LOCTEXT("CopyMorphTargetNamesButtonLabel", "Copy Names");
+					const FText ToolTipText = LOCTEXT("CopyMorphTargetNamesButtonTooltip", "Copy the names of selected morph targets to clipboard");
+					MenuBuilder.AddMenuEntry(Label, ToolTipText, FSlateIcon(), Action);
+				}
+			}
+		}
 	}
 	MenuBuilder.EndSection();
 
@@ -484,7 +641,13 @@ void SMorphTargetViewer::CreateMorphTargetList( const FString& SearchText )
 
 			int32 NumberOfVerts = (MorphTargets[I]->GetMorphLODModels().Num() > 0)? MorphTargets[I]->GetMorphLODModels()[0].Vertices.Num() : 0;
 
-			TSharedRef<FDisplayedMorphTargetInfo> Info = FDisplayedMorphTargetInfo::Make( MorphTargets[I]->GetFName(), NumberOfVerts);
+			TArray<FString> SourceFilenames;
+			SourceFilenames.AddDefaulted(SkeletalMesh->GetLODNum());
+			for (int32 LodIndex = 0; LodIndex < SkeletalMesh->GetLODNum(); ++LodIndex)
+			{
+				SourceFilenames[LodIndex] = MorphTargets[I]->GetCustomImportedSourceFilename(LodIndex);
+			}
+			TSharedRef<FDisplayedMorphTargetInfo> Info = FDisplayedMorphTargetInfo::Make( MorphTargets[I]->GetFName(), NumberOfVerts, SourceFilenames);
 			if(MeshComponent)
 			{
 				const float *CurveValPtr = MeshComponent->GetMorphTargetCurves().Find( MorphTargets[I]->GetFName() );
@@ -534,8 +697,14 @@ void SMorphTargetViewer::OnDeleteMorphTargets()
 		}
 	}
 
-	// Remove from mesh
-	SkeletalMesh->RemoveMorphTargets(MorphTargetNames);
+	//Scope a skeletal mesh build
+	{
+		FScopedSkeletalMeshPostEditChange ScopePostEditChange(SkeletalMesh);
+		// Remove from mesh
+		SkeletalMesh->RemoveMorphTargets(MorphTargetNames);
+	}
+	//Wait until the skeletal mesh compilation is done
+	FSkinnedAssetCompilingManager::Get().FinishCompilation({ SkeletalMesh });
 
 	CreateMorphTargetList( NameFilterBox->GetText().ToString() );
 }
@@ -560,6 +729,80 @@ void SMorphTargetViewer::OnCopyMorphTargetNames()
 	}
 }
 
+FReply SMorphTargetViewer::OnImportMorphTargetButton()
+{
+	constexpr int32 LodIndex0 = 0;
+	constexpr bool bWithNewFileTrue = true;
+	constexpr bool bRecreateMorphTargetListTrue = true;
+	InternalImportMorphTarget(LodIndex0, bWithNewFileTrue, nullptr, bRecreateMorphTargetListTrue);
+	
+	return FReply::Handled();
+}
+
+void SMorphTargetViewer::OnReimportMorphTargets(int32 LodIndex)
+{
+	constexpr bool bWithNewFileFalse = false;
+	TArray< TSharedPtr< FDisplayedMorphTargetInfo > > SelectedRows = MorphTargetListView->GetSelectedItems();
+	for (int RowIndex = 0; RowIndex < SelectedRows.Num(); ++RowIndex)
+	{
+		UMorphTarget* MorphTarget = SkeletalMesh->FindMorphTarget(SelectedRows[RowIndex]->Name);
+		if (MorphTarget)
+		{
+			constexpr bool bRecreateMorphTargetListFalse = false;
+			InternalImportMorphTarget(LodIndex, bWithNewFileFalse, MorphTarget, bRecreateMorphTargetListFalse);
+		}
+	}
+	CreateMorphTargetList();
+}
+
+void SMorphTargetViewer::OnReimportMorphTargetsWithNewFile(int32 LodIndex)
+{
+	constexpr bool bWithNewFileTrue = true;
+	TArray< TSharedPtr< FDisplayedMorphTargetInfo > > SelectedRows = MorphTargetListView->GetSelectedItems();
+	for (int RowIndex = 0; RowIndex < SelectedRows.Num(); ++RowIndex)
+	{
+		UMorphTarget* MorphTarget = SkeletalMesh->FindMorphTarget(SelectedRows[RowIndex]->Name);
+		if (MorphTarget)
+		{
+			constexpr bool bRecreateMorphTargetListFalse = false;
+			InternalImportMorphTarget(LodIndex, bWithNewFileTrue, MorphTarget, bRecreateMorphTargetListFalse);
+		}
+	}
+	CreateMorphTargetList();
+}
+
+void SMorphTargetViewer::InternalImportMorphTarget(int32 LodIndex, bool bWithNewFile, UMorphTarget* ReimportMorphTarget, bool bRecreateMorphTargetList)
+{
+	FString Filename;
+
+	bool bInternalWithNewFile = bWithNewFile || !ReimportMorphTarget || !ReimportMorphTarget->IsCustomImported(LodIndex);
+	if (bInternalWithNewFile)
+	{
+		FText PickerTitle = FText::Format(NSLOCTEXT("SMorphTargetViewer", "OnImportNewMorphTarget_PickerTitle", "Choose a file to import a morph target for LOD{0}"), FText::AsNumber(LodIndex));
+
+		if (!UInterchangeMeshUtilities::ShowMeshFilePicker(Filename, PickerTitle))
+		{
+			return;
+		}
+	}
+	else if(ensure(ReimportMorphTarget && ReimportMorphTarget->IsCustomImported(LodIndex)))
+	{
+		Filename = ReimportMorphTarget->GetCustomImportedSourceFilename(LodIndex);
+	}
+
+	constexpr bool bAsyncFalse = false;
+	UInterchangeManager& InterchangeManager = UInterchangeManager::GetInterchangeManager();
+	const UInterchangeSourceData* SourceData = InterchangeManager.CreateSourceData(Filename);
+	//Import a new morph target
+	TFuture<bool> FutureResult = UInterchangeMeshUtilities::ImportMorphTarget(SkeletalMesh, LodIndex, SourceData, bAsyncFalse, ReimportMorphTarget ? ReimportMorphTarget->GetName() : FString());
+	ensure(FutureResult.IsReady());
+
+	if (bRecreateMorphTargetList)
+	{
+		CreateMorphTargetList();
+	}
+}
+
 SMorphTargetViewer::~SMorphTargetViewer()
 {
 	if (PreviewScenePtr.IsValid())
diff --git a/Engine/Source/Editor/Persona/Private/SMorphTargetViewer.h b/Engine/Source/Editor/Persona/Private/SMorphTargetViewer.h
--- a/Engine/Source/Editor/Persona/Private/SMorphTargetViewer.h
+++ b/Engine/Source/Editor/Persona/Private/SMorphTargetViewer.h
@@ -12,6 +12,7 @@
 #include "Widgets/Views/SListView.h"
 
 class USkeletalMesh;
+class UMorphTarget;
 
 //////////////////////////////////////////////////////////////////////////
 // FDisplayedMorphTargetInfo
@@ -23,20 +24,22 @@ public:
 	float Weight;
 	bool bAutoFillData;
 	int32 NumberOfVerts;
+	TArray<FString> SourceFilenames;
 
 	/** Static function for creating a new item, but ensures that you can only have a TSharedRef to one */
-	static TSharedRef<FDisplayedMorphTargetInfo> Make(const FName& Source, int32 InNumberOfVerts)
+	static TSharedRef<FDisplayedMorphTargetInfo> Make(const FName& Source, int32 InNumberOfVerts, const TArray<FString>& InSourceFilenames)
 	{
-		return MakeShareable(new FDisplayedMorphTargetInfo(Source, InNumberOfVerts));
+		return MakeShareable(new FDisplayedMorphTargetInfo(Source, InNumberOfVerts, InSourceFilenames));
 	}
 
 protected:
 	/** Hidden constructor, always use Make above */
-	FDisplayedMorphTargetInfo(const FName& InSource, int32 InNumberOfVerts)
+	FDisplayedMorphTargetInfo(const FName& InSource, int32 InNumberOfVerts, const TArray<FString>& InSourceFilenames)
 		: Name( InSource )
 		, Weight( 0 )
 		, bAutoFillData(true)
 		, NumberOfVerts(InNumberOfVerts)
+		, SourceFilenames(InSourceFilenames)
 	{}
 
 	/** Hidden constructor, always use Make above */
@@ -140,6 +143,15 @@ public:
 	/** Handler for copying morph taret names */
 	void OnCopyMorphTargetNames();
 
+	/** Handler for importing a new morph target */
+	FReply OnImportMorphTargetButton();
+	void OnReimportMorphTargets(int32 LodIndex);
+	void OnReimportMorphTargetsWithNewFile(int32 LodIndex);
+
+	void InternalImportMorphTarget(int32 LodIndex, bool bWithNewFile, UMorphTarget* ReimportMorphTarget, bool bRecreateMorphTargetList);
+
+
+
 	/**
 	* Accessor so our rows can grab the filtertext for highlighting
 	*
diff --git a/Engine/Source/Runtime/Core/Public/UObject/FortniteMainBranchObjectVersion.h b/Engine/Source/Runtime/Core/Public/UObject/FortniteMainBranchObjectVersion.h
--- a/Engine/Source/Runtime/Core/Public/UObject/FortniteMainBranchObjectVersion.h
+++ b/Engine/Source/Runtime/Core/Public/UObject/FortniteMainBranchObjectVersion.h
@@ -486,6 +486,9 @@ struct FFortniteMainBranchObjectVersion
 		// Fixup for the data that has been damaged by LandscapeTargetLayersInLandscapeActor (loss of landscape layer info object assignments)
 		FixupLandscapeTargetLayersInLandscapeActor,
 
+		//Allow custom import of morph target
+		MorphTargetCustomImport,
+		
 		// -----<new versions can be added above this line>-------------------------------------------------
 		VersionPlusOne,
 		LatestVersion = VersionPlusOne - 1
diff --git a/Engine/Source/Runtime/Engine/Classes/Animation/MorphTarget.h b/Engine/Source/Runtime/Engine/Classes/Animation/MorphTarget.h
--- a/Engine/Source/Runtime/Engine/Classes/Animation/MorphTarget.h
+++ b/Engine/Source/Runtime/Engine/Classes/Animation/MorphTarget.h
@@ -67,6 +67,9 @@ struct FMorphTargetLODModel
 	/** Is this LOD generated by reduction setting */
 	bool bGeneratedByEngine;
 
+	/* The source filename use to import this morph target. If source is empty this morph target was import with the LOD geometry. */
+	FString SourceFilename;
+
 	FMorphTargetLODModel()
 		: NumVertices(0)
 		, NumBaseMeshVerts(0)
@@ -85,6 +88,8 @@ struct FMorphTargetLODModel
 		// since engine cleared it, we mark as engine generated
 		// this makes it clear to clear up later
 		bGeneratedByEngine = true;
+
+		SourceFilename.Empty();
 	}
 
 	void DiscardVertexData()
@@ -154,6 +159,11 @@ public:
 	ENGINE_API virtual bool HasValidData() const;
 	ENGINE_API virtual void EmptyMorphLODModels();
 
+	/** return true if this morphtarget was custom imported from a file */
+	ENGINE_API virtual bool IsCustomImported(int32 LODIndex) const;
+	ENGINE_API virtual FString GetCustomImportedSourceFilename(int32 LODIndex) const;
+	ENGINE_API virtual void SetCustomImportedSourceFilename(int32 LODIndex, const FString& InSourceFilename);
+
 	/** Discard CPU Buffers after render resources have been created. */
 	UE_DEPRECATED(5.0, "No longer in use, will be deleted. Whether to discard vertex data is now determined during cooking instead of loading.")
 	ENGINE_API virtual void DiscardVertexData();
diff --git a/Engine/Source/Runtime/Engine/Classes/Engine/SkinnedAssetCommon.h b/Engine/Source/Runtime/Engine/Classes/Engine/SkinnedAssetCommon.h
--- a/Engine/Source/Runtime/Engine/Classes/Engine/SkinnedAssetCommon.h
+++ b/Engine/Source/Runtime/Engine/Classes/Engine/SkinnedAssetCommon.h
@@ -96,6 +96,25 @@ struct FSectionReference
 	}
 };
 
+USTRUCT()
+struct FMorphTargetImportedSourceFileInfo
+{
+	GENERATED_USTRUCT_BODY()
+
+	ENGINE_API const FString& GetSourceFilename() const;
+
+	ENGINE_API void SetSourceFilename(const FString& Filename);
+
+private:
+	UPROPERTY(EditAnywhere, Category = MorphTargetInfo)
+	FString SourceFilename;
+	
+	UPROPERTY()
+	FGuid DerivedDataHash;
+
+	friend FArchive& operator<<(FArchive& Ar, FMorphTargetImportedSourceFileInfo& MorphTargetImportedSourceFileInfo);
+};
+
 /** Struct containing information for a particular LOD level, such as materials and info for when to use it. */
 USTRUCT()
 struct FSkeletalMeshLODInfo
@@ -177,6 +196,14 @@ struct FSkeletalMeshLODInfo
 	UPROPERTY(EditAnywhere, Category = SkeletalMeshLODInfo, meta = (UIMin = "0.01", ClampMin = "0.01", UIMax = "10000.0", ClampMax = "10000.0"))
 	float MorphTargetPositionErrorTolerance = 20.0f;
 
+#if WITH_EDITORONLY_DATA
+
+	/** Store the custom import morph target source file. The key of the map is the morph target name and the value is the source file path. */
+	UPROPERTY(VisibleAnywhere, Category = SkeletalMeshLODInfo)
+	TMap<FString, FMorphTargetImportedSourceFileInfo> ImportedMorphTargetSourceFilename;
+
+#endif //WITH_EDITORONLY_DATA
+
 	/** Whether to disable morph targets for this LOD. */
 	UPROPERTY()
 	uint8 bHasBeenSimplified:1;
diff --git a/Engine/Source/Runtime/Engine/Private/MorphTools.cpp b/Engine/Source/Runtime/Engine/Private/MorphTools.cpp
--- a/Engine/Source/Runtime/Engine/Private/MorphTools.cpp
+++ b/Engine/Source/Runtime/Engine/Private/MorphTools.cpp
@@ -57,6 +57,24 @@ FArchive& operator<<(FArchive& Ar, FMorphTargetLODModel& M)
 
 			Ar << M.NumBaseMeshVerts << M.SectionIndices << M.bGeneratedByEngine;
 		}
+
+		if (Ar.IsLoading() && Ar.CustomVer(FFortniteMainBranchObjectVersion::GUID) < FFortniteMainBranchObjectVersion::MorphTargetCustomImport)
+		{
+			M.SourceFilename.Empty();
+		}
+		else
+		{
+			//Do not cook the source filename, we don't need it at runtime
+			if (Ar.IsCooking() && Ar.IsSaving())
+			{
+				FString EmptySourceFilename;
+				Ar << EmptySourceFilename;
+			}
+			else
+			{
+				Ar << M.SourceFilename;
+			}
+		}
 	}
 
 	return Ar;
@@ -126,6 +144,38 @@ void UMorphTarget::EmptyMorphLODModels()
 	MorphLODModels.Empty();
 }
 
+bool UMorphTarget::IsCustomImported(int32 LODIndex) const
+{
+	if (LODIndex < MorphLODModels.Num())
+	{
+		// Calling GetMorphLODModels to potentially get from subclasses
+		const FMorphTargetLODModel& MorphModel = GetMorphLODModels()[LODIndex];
+		return !MorphModel.SourceFilename.IsEmpty();
+	}
+	return false;
+}
+
+FString UMorphTarget::GetCustomImportedSourceFilename(int32 LODIndex) const
+{
+	if (LODIndex < MorphLODModels.Num())
+	{
+		// Calling GetMorphLODModels to potentially get from subclasses
+		const FMorphTargetLODModel& MorphModel = GetMorphLODModels()[LODIndex];
+		return MorphModel.SourceFilename;
+	}
+	return FString();
+}
+
+void UMorphTarget::SetCustomImportedSourceFilename(int32 LODIndex, const FString& InSourceFilename)
+{
+	if (LODIndex < MorphLODModels.Num())
+	{
+		// Calling GetMorphLODModels to potentially get from subclasses
+		FMorphTargetLODModel& MorphModel = GetMorphLODModels()[LODIndex];
+		MorphModel.SourceFilename = InSourceFilename;
+	}
+}
+
 void UMorphTarget::DiscardVertexData()
 {
 	for (FMorphTargetLODModel& Model : MorphLODModels)
diff --git a/Engine/Source/Runtime/Engine/Private/SkeletalMesh.cpp b/Engine/Source/Runtime/Engine/Private/SkeletalMesh.cpp
--- a/Engine/Source/Runtime/Engine/Private/SkeletalMesh.cpp
+++ b/Engine/Source/Runtime/Engine/Private/SkeletalMesh.cpp
@@ -4094,6 +4094,13 @@ bool USkeletalMesh::RemoveMorphTargets(TConstArrayView<FName> InMorphTargetNames
 
 			bRemoved = true;
 		}
+
+		//Clean up the LodInfo Imported morph target source filename
+		for (int32 LODIndex = 0; LODIndex < GetLODNum(); ++LODIndex)
+		{
+			FSkeletalMeshLODInfo& LODInfoEntry = *GetLODInfo(LODIndex);
+			LODInfoEntry.ImportedMorphTargetSourceFilename.Remove(MorphTargetName.ToString());
+		}
 	}
 
 	return bRemoved;
@@ -4163,6 +4170,18 @@ bool USkeletalMesh::RenameMorphTarget(FName InOldName, FName InNewName)
 	// Rename the morph target itself
 	MorphTarget->Rename(*InNewName.ToString(), nullptr, REN_ForceNoResetLoaders | REN_DontCreateRedirectors);
 
+	//Clean up the LodInfo Imported morph target source filename we must also rename the entry
+	for (int32 InternalLodIndex = 0; InternalLodIndex < GetLODNum(); ++InternalLodIndex)
+	{
+		FSkeletalMeshLODInfo& LODInfoEntry = *GetLODInfo(InternalLodIndex);
+		if (const FMorphTargetImportedSourceFileInfo* MorphTargetImportedSourceFileInfo = LODInfoEntry.ImportedMorphTargetSourceFilename.Find(InOldName.ToString()))
+		{
+			const FString OldFilename = MorphTargetImportedSourceFileInfo->GetSourceFilename();
+			LODInfoEntry.ImportedMorphTargetSourceFilename.FindOrAdd(InNewName.ToString()).SetSourceFilename(OldFilename);
+			LODInfoEntry.ImportedMorphTargetSourceFilename.Remove(InOldName.ToString());
+		}
+	}
+
 	// Re-register the morph target
 	RegisterMorphTarget(MorphTarget);
 
diff --git a/Engine/Source/Runtime/Engine/Private/SkinnedAssetCommon.cpp b/Engine/Source/Runtime/Engine/Private/SkinnedAssetCommon.cpp
--- a/Engine/Source/Runtime/Engine/Private/SkinnedAssetCommon.cpp
+++ b/Engine/Source/Runtime/Engine/Private/SkinnedAssetCommon.cpp
@@ -1,24 +1,62 @@
 // Copyright Epic Games, Inc. All Rights Reserved.
 
 #include "Engine/SkinnedAssetCommon.h"
+
 #include "Animation/AnimSequence.h"
 #include "Engine/SkeletalMeshLODSettings.h"
 #include "Engine/TextureStreamingTypes.h"
-#include "Materials/MaterialInterface.h"
-#include "UObject/CoreObjectVersion.h"
+#include "HAL/FileManager.h"
 #include "Interfaces/ITargetPlatform.h"
+#include "Materials/MaterialInterface.h"
+#include "Misc/Paths.h"
+#include "Misc/SecureHash.h"
 #include "Rendering/SkeletalMeshLODModel.h"
+#include "SkeletalMeshLegacyCustomVersions.h"
+#include "UObject/CoreObjectVersion.h"
 #include "UObject/EditorObjectVersion.h"
 #include "UObject/RenderingObjectVersion.h"
-#include "SkeletalMeshLegacyCustomVersions.h"
 
 #include UE_INLINE_GENERATED_CPP_BY_NAME(SkinnedAssetCommon)
 
+
 /*-----------------------------------------------------------------------------
 	FSkeletalMeshLODInfo
 -----------------------------------------------------------------------------*/
 
 #if WITH_EDITORONLY_DATA
+
+const FString& FMorphTargetImportedSourceFileInfo::GetSourceFilename() const
+{
+	return SourceFilename;
+}
+
+void FMorphTargetImportedSourceFileInfo::SetSourceFilename(const FString& Filename)
+{
+	FString DerivedDataKey = Filename;
+	if (FPaths::FileExists(Filename))
+	{
+		DerivedDataKey += IFileManager::Get().GetTimeStamp(*Filename).ToString();
+	}
+
+	const auto CharArray = DerivedDataKey.GetCharArray();
+	FSHA1 Sha;
+	Sha.Update(reinterpret_cast<const uint8*>(CharArray.GetData()), CharArray.Num() * CharArray.GetTypeSize());
+	Sha.Final();
+	// Retrieve the hash and use it to construct a pseudo-GUID.
+	uint32 Hash[5];
+	Sha.GetHash((uint8*)Hash);
+	DerivedDataHash = FGuid(Hash[0] ^ Hash[4], Hash[1], Hash[2], Hash[3]);
+
+	SourceFilename = Filename;
+}
+
+FArchive& operator<<(FArchive& Ar, FMorphTargetImportedSourceFileInfo& MorphTargetImportedSourceFileInfo)
+{
+	Ar << MorphTargetImportedSourceFileInfo.SourceFilename;
+	Ar << MorphTargetImportedSourceFileInfo.DerivedDataHash;
+	return Ar;
+}
+
 static void SerializeReductionSettingsForDDC(FArchive& Ar, FSkeletalMeshOptimizationSettings& ReductionSettings)
 {
 	check(Ar.IsSaving());
@@ -92,6 +130,7 @@ FGuid FSkeletalMeshLODInfo::ComputeDeriveDataCacheKey(const FSkeletalMeshLODGrou
 	Ar << SectionsToPrioritize;
 	Ar << WeightOfPrioritization;
 	Ar << MorphTargetPositionErrorTolerance;
+	Ar << ImportedMorphTargetSourceFilename;
 
 	//TODO: Ask the derivedata key of the UObject reference by FSoftObjectPath. So if someone change the UObject, this LODs will get dirty
 	//and will be rebuild.
diff --git a/Engine/Source/Runtime/Interchange/Core/Public/InterchangePipelineBase.h b/Engine/Source/Runtime/Interchange/Core/Public/InterchangePipelineBase.h
--- a/Engine/Source/Runtime/Interchange/Core/Public/InterchangePipelineBase.h
+++ b/Engine/Source/Runtime/Interchange/Core/Public/InterchangePipelineBase.h
@@ -47,10 +47,12 @@ enum class EInterchangePipelineContext : uint8
 	AssetReimport,
 	SceneImport,
 	SceneReimport,
-	AssetCustomLODImport, //The import for custom LOD is there because we use a copy of the asset import data pipeline stack.
+	AssetCustomLODImport,
 	AssetCustomLODReimport,
-	AssetAlternateSkinningImport, //The import for custom LOD is there because we use a copy of the asset import data pipeline stack.
+	AssetAlternateSkinningImport,
 	AssetAlternateSkinningReimport,
+	AssetCustomMorphTargetImport, //Import the content has a combine static mesh so we can add a custom morph target to a skeletal mesh
+	AssetCustomMorphTargetReImport,
 };
 
 USTRUCT(BlueprintType)
diff --git a/Engine/Source/Runtime/Interchange/Engine/InterchangeEngine.Build.cs b/Engine/Source/Runtime/Interchange/Engine/InterchangeEngine.Build.cs
--- a/Engine/Source/Runtime/Interchange/Engine/InterchangeEngine.Build.cs
+++ b/Engine/Source/Runtime/Interchange/Engine/InterchangeEngine.Build.cs
@@ -24,6 +24,8 @@ public class InterchangeEngine : ModuleRules
 				"DeveloperSettings",
 				"Json",
 				"JsonUtilities",
+				"MeshDescription",
+				"SkeletalMeshDescription",
 			}
 		);
 
diff --git a/Engine/Source/Runtime/Interchange/Engine/Private/InterchangeMeshUtilities.cpp b/Engine/Source/Runtime/Interchange/Engine/Private/InterchangeMeshUtilities.cpp
--- a/Engine/Source/Runtime/Interchange/Engine/Private/InterchangeMeshUtilities.cpp
+++ b/Engine/Source/Runtime/Interchange/Engine/Private/InterchangeMeshUtilities.cpp
@@ -2,24 +2,32 @@
 
 #include "InterchangeMeshUtilities.h"
 
+#include "AssetCompilingManager.h"
 #include "AssetRegistry/AssetData.h"
 #include "AssetRegistry/AssetRegistryModule.h"
 #include "Async/Future.h"
 #include "Engine/StaticMesh.h"
 #include "Engine/SkeletalMesh.h"
 #include "Engine/SkinnedAssetAsyncCompileUtils.h"
+#include "Framework/Notifications/NotificationManager.h"
 #include "InterchangeAssetImportData.h"
 #include "InterchangeFactoryBase.h"
+#include "InterchangeFilePickerBase.h"
 #include "InterchangeEngineLogPrivate.h"
 #include "InterchangeManager.h"
 #include "InterchangeProjectSettings.h"
 #include "InterchangePythonPipelineBase.h"
 #include "InterchangeSourceData.h"
 #include "Logging/LogMacros.h"
+#include "MeshDescription.h"
 #include "Misc/ScopedSlowTask.h"
+#include "Rendering/SkeletalMeshLODImporterData.h"
+#include "SkeletalMeshAttributes.h"
+#include "StaticMeshAttributes.h"
 #include "UObject/Class.h"
 #include "UObject/Object.h"
 #include "UObject/ObjectMacros.h"
+#include "Widgets/Notifications/SNotificationList.h"
 
 #if WITH_EDITOR
 #include "LODUtilities.h"
@@ -27,6 +35,32 @@
 
 #include UE_INLINE_GENERATED_CPP_BY_NAME(InterchangeMeshUtilities)
 
+namespace UE::Interchange::Private
+{
+	void DeletePathAssets(const FString& ImportAssetPath)
+	{
+		FAssetRegistryModule& AssetRegistryModule = FModuleManager::LoadModuleChecked<FAssetRegistryModule>("AssetRegistry");
+
+		TArray<FAssetData> AssetsToDelete;
+		AssetRegistryModule.Get().GetAssetsByPath(FName(*ImportAssetPath), AssetsToDelete, true);
+		for (FAssetData AssetData : AssetsToDelete)
+		{
+			UObject* ObjToDelete = AssetData.GetAsset();
+			if (ObjToDelete)
+			{
+				//Avoid temporary package to be saved
+				UPackage* Package = ObjToDelete->GetOutermost();
+				Package->SetDirtyFlag(false);
+				//Avoid gc, remove keep flags
+				ObjToDelete->ClearFlags(RF_Standalone);
+				ObjToDelete->ClearInternalFlags(EInternalObjectFlags::Async);
+				//Make the object transient to prevent saving
+				ObjToDelete->SetFlags(RF_Transient);
+			}
+		}
+	}
+}
+
 FInterchangeSkeletalMeshAlternateSkinWeightPostImportTask::FInterchangeSkeletalMeshAlternateSkinWeightPostImportTask(USkeletalMesh* InSkeletalMesh)
 	:SkeletalMesh(InSkeletalMesh)
 {
@@ -99,6 +133,42 @@ bool FInterchangeSkeletalMeshAlternateSkinWeightPostImportTask::AddLodToReimport
 	return true;
 }
 
+bool UInterchangeMeshUtilities::ShowMeshFilePicker(FString& OutFilename, const FText& Title)
+{
+	//Pop a file picker that join both interchange and other format
+	//Ask the user for a file path
+	const UInterchangeProjectSettings* InterchangeProjectSettings = GetDefault<UInterchangeProjectSettings>();
+	UInterchangeFilePickerBase* FilePicker = nullptr;
+
+	//In runtime we do not have any pipeline configurator
+#if WITH_EDITORONLY_DATA
+	TSoftClassPtr <UInterchangeFilePickerBase> FilePickerClass = InterchangeProjectSettings->FilePickerClass;
+	if (FilePickerClass.IsValid())
+	{
+		UClass* FilePickerClassLoaded = FilePickerClass.LoadSynchronous();
+		if (FilePickerClassLoaded)
+		{
+			FilePicker = NewObject<UInterchangeFilePickerBase>(GetTransientPackage(), FilePickerClassLoaded, NAME_None, RF_NoFlags);
+		}
+	}
+#endif
+	if (FilePicker)
+	{
+		FInterchangeFilePickerParameters Parameters;
+		Parameters.bAllowMultipleFiles = false;
+		Parameters.Title = Title;
+		Parameters.bShowAllFactoriesExtension = false;
+		TArray<FString> Filenames;
+		bool bFilePickerResult = FilePicker->ScriptedFilePickerForTranslatorAssetType(EInterchangeTranslatorAssetType::Meshes, Parameters, Filenames);
+		if (bFilePickerResult && Filenames.Num() > 0)
+		{
+			OutFilename = Filenames[0];
+			return FPaths::FileExists(OutFilename);
+		}
+	}
+	return false;
+}
+
 TFuture<bool> UInterchangeMeshUtilities::ImportCustomLod(UObject* MeshObject, const int32 LodIndex, const UInterchangeSourceData* SourceData, bool bAsync)
 {
 	TSharedPtr<TPromise<bool>> Promise = MakeShared<TPromise<bool>>();
@@ -214,31 +284,6 @@ TFuture<bool> UInterchangeMeshUtilities::InternalImportCustomLod(TSharedPtr<TPro
 	}
 
  	FString ImportAssetPath = TEXT("/Engine/TempEditor/Interchange/") + FGuid::NewGuid().ToString(EGuidFormats::Base36Encoded);
-	
-
-	//Create a functor to remove asset keep flag and package dirty flags
-	auto DeletePathAssets = [ImportAssetPath]()
-		{
-			FAssetRegistryModule& AssetRegistryModule = FModuleManager::LoadModuleChecked<FAssetRegistryModule>("AssetRegistry");
-
-			TArray<FAssetData> AssetsToDelete;
-			AssetRegistryModule.Get().GetAssetsByPath(FName(*ImportAssetPath), AssetsToDelete, true);
-			for (FAssetData AssetData : AssetsToDelete)
-			{
-				UObject* ObjToDelete = AssetData.GetAsset();
-				if (ObjToDelete)
-				{
-					//Avoid temporary package to be saved
-					UPackage* Package = ObjToDelete->GetOutermost();
-					Package->SetDirtyFlag(false);
-					//Avoid gc, remove keep flags
-					ObjToDelete->ClearFlags(RF_Standalone);
-					ObjToDelete->ClearInternalFlags(EInternalObjectFlags::Async);
-					//Make the object transient to prevent saving
-					ObjToDelete->SetFlags(RF_Transient);
-				}
-			}
-		};
 
 	UE::Interchange::FAssetImportResultRef AssetImportResult = bAsync 
 		? InterchangeManager.ImportAssetAsync(ImportAssetPath, SourceData, ImportAssetParameters)
@@ -247,7 +292,7 @@ TFuture<bool> UInterchangeMeshUtilities::InternalImportCustomLod(TSharedPtr<TPro
 	FString SourceDataFilename = SourceData->GetFilename();
 	if (SkeletalMesh)
 	{
-		AssetImportResult->OnDone([Promise, SkeletalMesh, LodIndex, SourceDataFilename, DeletePathAssets](UE::Interchange::FImportResult& ImportResult)
+		AssetImportResult->OnDone([Promise, SkeletalMesh, LodIndex, SourceDataFilename, ImportAssetPath](UE::Interchange::FImportResult& ImportResult)
 			{
 				USkeletalMesh* SourceSkeletalMesh = Cast< USkeletalMesh >(ImportResult.GetFirstAssetOfClass(USkeletalMesh::StaticClass()));
 
@@ -261,12 +306,12 @@ TFuture<bool> UInterchangeMeshUtilities::InternalImportCustomLod(TSharedPtr<TPro
 				{
 					Promise->SetValue(false);
 				}
-				DeletePathAssets();
+				UE::Interchange::Private::DeletePathAssets(ImportAssetPath);
 			});
 	}
 	else if (StaticMesh)
 	{
-		AssetImportResult->OnDone([Promise, StaticMesh, LodIndex, SourceDataFilename, DeletePathAssets](UE::Interchange::FImportResult& ImportResult)
+		AssetImportResult->OnDone([Promise, StaticMesh, LodIndex, SourceDataFilename, ImportAssetPath](UE::Interchange::FImportResult& ImportResult)
 			{
 				UStaticMesh* SourceStaticMesh = Cast< UStaticMesh >(ImportResult.GetFirstAssetOfClass(UStaticMesh::StaticClass()));
 				if(SourceStaticMesh)
@@ -277,7 +322,7 @@ TFuture<bool> UInterchangeMeshUtilities::InternalImportCustomLod(TSharedPtr<TPro
 				{
 					Promise->SetValue(false);
 				}
-				DeletePathAssets();
+				UE::Interchange::Private::DeletePathAssets(ImportAssetPath);
 			});
 	}
 
@@ -287,3 +332,222 @@ TFuture<bool> UInterchangeMeshUtilities::InternalImportCustomLod(TSharedPtr<TPro
 	return Promise->GetFuture();
 #endif
 }
+
+TFuture<bool> UInterchangeMeshUtilities::ImportMorphTarget(USkeletalMesh* SkeletalMesh, const int32 LodIndex, const UInterchangeSourceData* SourceData, bool bAsync, const FString& MorphTargetName)
+{
+	TSharedPtr<TPromise<bool>> Promise = MakeShared<TPromise<bool>>();
+
+#if WITH_EDITOR
+	auto ExitNoImport = [&Promise]()
+		{
+			Promise->SetValue(false);
+			FNotificationInfo NotificationInfo(FText::GetEmpty());
+			NotificationInfo.ExpireDuration = 5.0f;
+			NotificationInfo.bUseSuccessFailIcons = true;
+			NotificationInfo.Text = NSLOCTEXT("InterchangeMeshUtilities", "ImportMorphTargetFail", "Fail importing morph target!");
+			if (TSharedPtr<SNotificationItem> NotificationPtr = FSlateNotificationManager::Get().AddNotification(NotificationInfo))
+			{
+				NotificationPtr->SetCompletionState(SNotificationItem::CS_Fail);
+			}
+			return Promise->GetFuture();
+		};
+
+	if (!SkeletalMesh)
+	{
+		UE_LOG(LogInterchangeEngine, Warning, TEXT("FInterchangeMeshUtilities::ImportMorphTarget: Cannot import morph targets, invalid skeletal mesh."));
+		return ExitNoImport();
+	}
+
+	UInterchangeManager& InterchangeManager = UInterchangeManager::GetInterchangeManager();
+
+	UInterchangeAssetImportData* InterchangeAssetImportData = nullptr;
+	EInterchangePipelineContext ImportType = MorphTargetName.IsEmpty() ? EInterchangePipelineContext::AssetCustomMorphTargetImport : EInterchangePipelineContext::AssetCustomMorphTargetReImport;
+	UObject* SourceImportData = nullptr;
+	SourceImportData = SkeletalMesh->GetAssetImportData();
+	InterchangeAssetImportData = Cast<UInterchangeAssetImportData>(SourceImportData);
+	if (LodIndex > SkeletalMesh->GetLODNum())
+	{
+		UE_LOG(LogInterchangeEngine, Warning, TEXT("FInterchangeMeshUtilities::ImportMorphTarget: Cannot import morph targets, invalid skeletal mesh LOD index %d."), LodIndex);
+		return ExitNoImport();
+	}
+
+	const bool bInterchangeCanImportSourceData = InterchangeManager.CanTranslateSourceData(SourceData);
+
+	if (!bInterchangeCanImportSourceData)
+	{
+		UE_LOG(LogInterchangeEngine, Warning, TEXT("FInterchangeMeshUtilities::ImportMorphTarget: Cannot import skeletal mesh morph targets, no interchange translator support this source file. [%s]"), *(SourceData->GetFilename()));
+		return ExitNoImport();
+	}
+
+	//Convert the asset import data if needed
+	if (!InterchangeAssetImportData)
+	{
+		//Try to convert the asset import data
+		InterchangeManager.ConvertImportData(SourceImportData, UInterchangeAssetImportData::StaticClass(), reinterpret_cast<UObject**>(&InterchangeAssetImportData));
+	}
+
+	FImportAssetParameters ImportAssetParameters;
+	ImportAssetParameters.bIsAutomated = true;
+	if (InterchangeAssetImportData)
+	{
+		TArray<UObject*> Pipelines = InterchangeAssetImportData->GetPipelines();
+		for (UObject* SelectedPipeline : Pipelines)
+		{
+			UInterchangePipelineBase* GeneratedPipeline = nullptr;
+			if (UInterchangePythonPipelineAsset* PythonPipelineAsset = Cast<UInterchangePythonPipelineAsset>(SelectedPipeline))
+			{
+				GeneratedPipeline = Cast<UInterchangePipelineBase>(StaticDuplicateObject(PythonPipelineAsset->GeneratedPipeline, GetTransientPackage()));
+			}
+			else
+			{
+				GeneratedPipeline = Cast<UInterchangePipelineBase>(StaticDuplicateObject(SelectedPipeline, GetTransientPackage()));
+			}
+			if (ensure(GeneratedPipeline))
+			{
+				GeneratedPipeline->AdjustSettingsForContext(ImportType, nullptr, nullptr);
+				ImportAssetParameters.OverridePipelines.Add(GeneratedPipeline);
+			}
+		}
+	}
+	else
+	{
+		//Create import data
+		InterchangeAssetImportData = NewObject<UInterchangeAssetImportData>();
+		const UInterchangeProjectSettings* InterchangeProjectSettings = GetDefault<UInterchangeProjectSettings>();
+
+		if (const UClass* GenericPipelineClass = InterchangeProjectSettings->GenericPipelineClass.LoadSynchronous())
+		{
+			if (UInterchangePipelineBase* GenericPipeline = NewObject<UInterchangePipelineBase>(GetTransientPackage(), GenericPipelineClass))
+			{
+				GenericPipeline->ClearFlags(EObjectFlags::RF_Standalone | EObjectFlags::RF_Public);
+				GenericPipeline->AdjustSettingsForContext(ImportType, nullptr, nullptr);
+				ImportAssetParameters.OverridePipelines.Add(GenericPipeline);
+			}
+		}
+	}
+
+	FString ImportAssetPath = TEXT("/Engine/TempEditor/Interchange/") + FGuid::NewGuid().ToString(EGuidFormats::Base36Encoded);
+
+	UE::Interchange::FAssetImportResultRef AssetImportResult = bAsync
+		? InterchangeManager.ImportAssetAsync(ImportAssetPath, SourceData, ImportAssetParameters)
+		: InterchangeManager.ImportAssetWithResult(ImportAssetPath, SourceData, ImportAssetParameters);
+
+	FString SourceDataFilename = SourceData->GetFilename();
+	AssetImportResult->OnDone([Promise, SkeletalMesh, LodIndex, SourceDataFilename, ImportAssetPath, bAsync, MorphTargetName](UE::Interchange::FImportResult& ImportResult)
+		{
+			UStaticMesh* SourceStaticMesh = Cast< UStaticMesh >(ImportResult.GetFirstAssetOfClass(UStaticMesh::StaticClass()));
+
+			auto LambdaCleanup = [&Promise, &ImportAssetPath](bool bResult)
+				{
+					UE::Interchange::Private::DeletePathAssets(ImportAssetPath);
+					Promise->SetValue(bResult);
+					FNotificationInfo NotificationInfo(FText::GetEmpty());
+					NotificationInfo.ExpireDuration = 5.0f;
+					NotificationInfo.bUseSuccessFailIcons = true;
+					NotificationInfo.Text = bResult
+						? NSLOCTEXT("InterchangeMeshUtilities", "ImportMorphTargetSuccessful", "Morph target imported successfully!")
+						: NSLOCTEXT("InterchangeMeshUtilities", "ImportMorphTargetFail", "Fail importing morph target!");
+					TSharedPtr<SNotificationItem> NotificationPtr = FSlateNotificationManager::Get().AddNotification(NotificationInfo);
+					if (NotificationPtr)
+					{
+						NotificationPtr->SetCompletionState(bResult ? SNotificationItem::CS_Success : SNotificationItem::CS_Fail);
+					}
+				};
+
+			if (SourceStaticMesh)
+			{
+				const int32 SourceLodCount = SourceStaticMesh->GetNumSourceModels();
+				if (SourceLodCount == 0)
+				{
+					UE_LOG(LogInterchangeEngine, Warning, TEXT("FInterchangeMeshUtilities::ImportMorphTarget: Cannot import skeletal mesh morph targets, There was no geometry in the provided source file. [%s]"), *(SourceDataFilename));
+					LambdaCleanup(false);
+					return;
+				}
+				const int32 SourceLodIndex = SourceStaticMesh->IsMeshDescriptionValid(LodIndex) ? LodIndex : 0;
+				FMeshDescription* SourceMeshDescription = SourceStaticMesh->GetMeshDescription(SourceLodIndex);
+				if (!SourceMeshDescription)
+				{
+					UE_LOG(LogInterchangeEngine, Warning, TEXT("FInterchangeMeshUtilities::ImportMorphTarget: Cannot import skeletal mesh morph targets, There was no geometry for LOD index %d in the provided source file. [%s]"), LodIndex, *(SourceDataFilename));
+					LambdaCleanup(false);
+					return;
+				}
+				FMeshDescription* TargetMeshDescription = SkeletalMesh->GetMeshDescription(LodIndex);
+				if (!TargetMeshDescription)
+				{
+					UE_LOG(LogInterchangeEngine, Warning, TEXT("FInterchangeMeshUtilities::ImportMorphTarget: Cannot import morph targets, missing skeletal mesh geometry for LOD index %d."), LodIndex);
+					LambdaCleanup(false);
+					return;
+				}
+				if (SourceMeshDescription->Vertices().Num() != TargetMeshDescription->Vertices().Num())
+				{
+					UE_LOG(LogInterchangeEngine, Warning, TEXT("FInterchangeMeshUtilities::ImportMorphTarget: Cannot import morph targets, the imported morph target geometry don't match target skeletal mesh geometry topology. (LOD index %d)"), LodIndex);
+					LambdaCleanup(false);
+					return;
+				}
+
+#if WITH_EDITORONLY_DATA
+				FSkeletalMeshLODInfo* LodInfo = SkeletalMesh->GetLODInfo(LodIndex);
+				if (!ensure(LodInfo))
+				{
+					UE_LOG(LogInterchangeEngine, Warning, TEXT("FInterchangeMeshUtilities::ImportMorphTarget: Cannot import morph targets, the skeletal mesh LOD info do not exist. (LOD index %d)"), LodIndex);
+					LambdaCleanup(false);
+					return;
+				}
+#endif //WITH_EDITORONLY_DATA
+
+				//Use the provided MorphTargetName if not empty. Existing name will force a re-import
+				const FString ImportedMorphTargetString = MorphTargetName.IsEmpty() ? SourceStaticMesh->GetName() : MorphTargetName;
+				const FName ImportedMorphTargetName = FName(*ImportedMorphTargetString);
+
+				FSkeletalMeshAttributes TargetAttributes(*TargetMeshDescription);
+				if (!TargetAttributes.GetMorphTargetNames().Contains(ImportedMorphTargetName))
+				{
+					TargetAttributes.RegisterMorphTargetAttribute(ImportedMorphTargetName, false);
+				}
+				TVertexAttributesRef<FVector3f> TargetMorphTargetPosDeltaAttribute = TargetAttributes.GetVertexMorphPositionDelta(ImportedMorphTargetName);
+				TVertexAttributesConstRef<FVector3f> TargetVertexPositions = TargetAttributes.GetVertexPositions();
+
+				FStaticMeshConstAttributes SourceAttributes(*SourceMeshDescription);
+				TVertexAttributesConstRef<FVector3f> SourceVertexPositions = SourceAttributes.GetVertexPositions();
+
+				//Populate the deltas in the target mesh description
+				for (FVertexID VertexID : SourceMeshDescription->Vertices().GetElementIDs())
+				{
+					TargetMorphTargetPosDeltaAttribute[VertexID] = SourceVertexPositions[VertexID] - TargetVertexPositions[VertexID];
+				}
+
+				//Commit the modified LodImportData
+				{
+					FScopedSkeletalMeshPostEditChange ScopePostEditChange(SkeletalMesh);
+					SkeletalMesh->PreEditChange(nullptr);
+
+#if WITH_EDITORONLY_DATA
+					//Add the mapping to the lod info so we can re-import the morph target
+					LodInfo->ImportedMorphTargetSourceFilename.FindOrAdd(ImportedMorphTargetString).SetSourceFilename(SourceDataFilename);
+#endif //WITH_EDITORONLY_DATA
+
+					SkeletalMesh->CommitMeshDescription(LodIndex);
+				}
+				
+				//Wait until the skeletal mesh compilation is done
+				if (!bAsync)
+				{
+					FAssetCompilingManager::Get().FinishCompilationForObjects({ SkeletalMesh });
+				}
+
+				//Set the promise to success
+				LambdaCleanup(true);
+			}
+			else
+			{
+
+				LambdaCleanup(false);
+			}
+		});
+
+	return Promise->GetFuture();
+#else
+	Promise->SetValue(false);
+	return Promise->GetFuture();
+#endif
+}
\ No newline at end of file
diff --git a/Engine/Source/Runtime/Interchange/Engine/Public/InterchangeMeshUtilities.h b/Engine/Source/Runtime/Interchange/Engine/Public/InterchangeMeshUtilities.h
--- a/Engine/Source/Runtime/Interchange/Engine/Public/InterchangeMeshUtilities.h
+++ b/Engine/Source/Runtime/Interchange/Engine/Public/InterchangeMeshUtilities.h
@@ -35,25 +35,64 @@ private:
 	TArray<int32> ReImportAlternateSkinWeightsLods;
 };
 
-UCLASS(Experimental, MinimalAPI)
+UCLASS(BlueprintType, MinimalAPI)
 class UInterchangeMeshUtilities : public UObject
 {
 	GENERATED_BODY()
 public:
 
+	/**
+	 * This function pop a file picker with all meshes supported format.
+	 *
+	 * @param OutFilename - The filename the user pick.
+	 * #param Title - The picker dialog title.
+	 * @Return - return true if the users successfully pick a file, false if something went wrong or the user cancel.
+	 *
+	 */
+	static INTERCHANGEENGINE_API bool ShowMeshFilePicker(FString& OutFilename, const FText& Title);
+
 	/**
 	 * This function import a mesh from the source data and add/replace the MeshObject LOD (at LodIndex) with the imported mesh LOD data.
 	 *
 	 * @Param MeshObject - The Mesh we want to add the lod
 	 * @Param LodIndex - The index of the LOD we want to replace or add
 	 * @Param SourceData - The source to import the custom LOD
+	 * @param bAsync - If true the future will not be set when the function return and the import will be asynchronous.
 	 * @Return - return future boolean which will be true if it successfully add or replace the MeshObject LOD at LodIndex with the imported data.
 	 *
-	 * @Note - Since this is the async version, the SourceData parameter is required.
-	 * @Note - This function will search for an available interchange asset factory that can do the job for the MeshObject class
 	 */
 	static INTERCHANGEENGINE_API TFuture<bool> ImportCustomLod(UObject* MeshObject, const int32 LodIndex, const UInterchangeSourceData* SourceData, bool bAsync);
 
+	/**
+	 * This function import a morph target from the source data and add/replace the skeletal mesh morph target.
+	 *
+	 * @Param SkeletalMesh - The target skeletal mesh we want to add the morph targets
+	 * @Param LodIndex - The index of the LOD we want to replace or add the morph targets
+	 * @Param SourceData - The source to import the morph targets
+	 * @Param bAsync - If true the future will not be set when the function return and the import will be asynchronous.
+	 * @Param MorphTargetName - If not empty we will use this name to create the morph target, if there is already an existing morph target it will be re-import
+	 * @Return - return future boolean which will be true if it successfully add or replace the skeletal mesh morph target at LodIndex with the imported data.
+	 *
+	 */
+	static INTERCHANGEENGINE_API TFuture<bool> ImportMorphTarget(USkeletalMesh* SkeletalMesh, const int32 LodIndex, const UInterchangeSourceData* SourceData, bool bAsync, const FString& MorphTargetName);
+
+	/**
+	 * This function import a morph target from the source data and add/replace the skeletal mesh morph target.
+	 *
+	 * @Param SkeletalMesh - The target skeletal mesh we want to add the morph targets
+	 * @Param LodIndex - The index of the LOD we want to replace or add the morph targets
+	 * @Param SourceData - The source to import the morph targets
+	 * @Param MorphTargetName - If not empty we will use this name to create the morph target, if there is already an existing morph target it will be re-import
+	 * @Return - return true if it successfully add or replace the skeletal mesh morph target at LodIndex, flase otherwise.
+	 *
+	 */
+	UFUNCTION(BlueprintCallable, Category = "SkeletalMesh | MorphTarget")
+	INTERCHANGEENGINE_API bool ScriptedImportMorphTarget(USkeletalMesh* SkeletalMesh, const int32 LodIndex, const UInterchangeSourceData* SourceData, const FString& MorphTargetName) const
+	{
+		constexpr bool bAsyncFalse = false;
+		return ImportMorphTarget(SkeletalMesh, LodIndex, SourceData, bAsyncFalse, MorphTargetName).Get();
+	}
+
 
 private:
 
