Interchange: refactor the task to not be tick anywhere on the game thread. We tick the interchange gamethread task in the engine tick to be sure we can do anything we want on the game thread.

#rnx

diff --git a/Engine/Source/Runtime/Engine/Private/SkeletalMesh.cpp b/Engine/Source/Runtime/Engine/Private/SkeletalMesh.cpp
--- a/Engine/Source/Runtime/Engine/Private/SkeletalMesh.cpp
+++ b/Engine/Source/Runtime/Engine/Private/SkeletalMesh.cpp
@@ -1400,9 +1400,15 @@ void USkeletalMesh::PreEditChange(FProperty* PropertyAboutToChange)
 {
 	TRACE_CPUPROFILER_EVENT_SCOPE(USkeletalMesh::PreEditChange);
 
-	// Tell the compiler to finish compiling us if we have a pending
-	// compilation ongoing plus any dependency (i.e. UGroomBindings).
-	FAssetCompilingManager::Get().FinishCompilationForObjects({ this });
+	//Don't call finish compile if this skeletal mesh is compiling and we are in a FSkinnedAssetAsyncBuildScope for this skeletal mesh.
+	//If on the game thread we call LockPropertyUntil and we call PreEditChange after, in such a case a deadlock will happen if
+	//we call finish compile on this skeletal mesh.
+	if (FSkinnedAssetAsyncBuildScope::ShouldWaitOnLockedProperties(this))
+	{
+		// Tell the compiler to finish compiling us if we have a pending
+		// compilation ongoing plus any dependency (i.e. UGroomBindings).
+		FAssetCompilingManager::Get().FinishCompilationForObjects({ this });
+	}
 
 	Super::PreEditChange(PropertyAboutToChange);
 }
diff --git a/Engine/Source/Runtime/Engine/Private/SkeletalMeshComponentPhysics.cpp b/Engine/Source/Runtime/Engine/Private/SkeletalMeshComponentPhysics.cpp
--- a/Engine/Source/Runtime/Engine/Private/SkeletalMeshComponentPhysics.cpp
+++ b/Engine/Source/Runtime/Engine/Private/SkeletalMeshComponentPhysics.cpp
@@ -2169,8 +2169,12 @@ void USkeletalMeshComponent::UpdateHasValidBodies()
 		// For each body in physics asset..
 		for( int32 BodyIndex = 0; BodyIndex < PhysicsAsset->SkeletalBodySetups.Num(); BodyIndex++ )
 		{
+			int32 BoneIndex = INDEX_NONE;
 			// .. find the matching graphics bone index
-			int32 BoneIndex = GetBoneIndex( PhysicsAsset->SkeletalBodySetups[ BodyIndex ]->BoneName );
+			if (TObjectPtr<USkeletalBodySetup> SkeletalBodySetup = PhysicsAsset->SkeletalBodySetups[BodyIndex])
+			{
+				BoneIndex = GetBoneIndex(SkeletalBodySetup->BoneName);
+			}
 
 			// If we found a valid graphics bone, set the 'valid' flag
 			if(BoneIndex != INDEX_NONE)
diff --git a/Engine/Source/Runtime/Interchange/Engine/Private/InterchangeManager.cpp b/Engine/Source/Runtime/Interchange/Engine/Private/InterchangeManager.cpp
--- a/Engine/Source/Runtime/Interchange/Engine/Private/InterchangeManager.cpp
+++ b/Engine/Source/Runtime/Interchange/Engine/Private/InterchangeManager.cpp
@@ -23,6 +23,7 @@
 #include "InterchangeProjectSettings.h"
 #include "InterchangePythonPipelineBase.h"
 #include "InterchangeSourceData.h"
+#include "InterchangeTaskSystem.h"
 #include "InterchangeTranslatorBase.h"
 #include "InterchangeWriterBase.h"
 #include "Internationalization/Internationalization.h"
@@ -537,47 +538,56 @@ void UE::Interchange::FImportAsyncHelper::SendAnalyticImportEndData()
 
 void UE::Interchange::FImportAsyncHelper::ReleaseTranslatorsSource()
 {
-	for (UInterchangeTranslatorBase* BaseTranslator : Translators)
+	//Make sure Executing translator task are done before releasing the resource
+	const int32 TranslatorCount = Translators.Num();
+	for (int32 TranslatorIndex = 0; TranslatorIndex < TranslatorCount; ++TranslatorIndex)
 	{
+		const uint64 TaskId = TranslatorTasks.IsValidIndex(TranslatorIndex) ? TranslatorTasks[TranslatorIndex] : INTERCHANGE_INVALID_TASK_ID;
+		UInterchangeTranslatorBase* BaseTranslator = Translators[TranslatorIndex];
+
 		if (BaseTranslator)
 		{
+			//Wait until the translator is done before releasing the resource. This should not happen since resource are free during completion
+			const EInterchangeTaskStatus TaskStatus = UE::Interchange::FInterchangeTaskSystem::Get().GetTaskStatus(TaskId);
+			if (TaskStatus == EInterchangeTaskStatus::Executing)
+			{
+				ensure(TaskStatus != EInterchangeTaskStatus::Executing);
+				UE::Interchange::FInterchangeTaskSystem::Get().WaitUntilTasksComplete({ TaskId });
+			}
 			BaseTranslator->ReleaseSource();
 		}
 	}
 }
 
-FGraphEventArray UE::Interchange::FImportAsyncHelper::GetCompletionTaskGraphEvent()
+TArray<uint64> UE::Interchange::FImportAsyncHelper::GetCompletionTaskGraphEvent()
 {
-	FGraphEventArray TasksToComplete;
+	TArray<uint64> TasksToComplete;
 
 	TasksToComplete.Append(TranslatorTasks);
 	TasksToComplete.Append(PipelineTasks);
 	
-	if (ParsingTask.GetReference())
+	if (ParsingTask != INTERCHANGE_INVALID_TASK_ID)
 	{
 		TasksToComplete.Add(ParsingTask);
 	}
 
 	//Parsing task must be done before the other tasks get added
-	FTaskGraphInterface::Get().WaitUntilTasksComplete(TasksToComplete, ENamedThreads::GameThread);
+	UE::Interchange::FInterchangeTaskSystem::Get().WaitUntilTasksComplete(TasksToComplete);
 	TasksToComplete.Reset();
 
 	TasksToComplete.Append(BeginImportObjectTasks);
 	TasksToComplete.Append(ImportObjectTasks);
 	TasksToComplete.Append(FinalizeImportObjectTasks);
 	TasksToComplete.Append(SceneTasks);
-	if (WaitAssetCompilationTask.GetReference())
-	{
-		TasksToComplete.Add(WaitAssetCompilationTask);
-	}
+	TasksToComplete.Append(WaitAssetCompilationTasks);
 	TasksToComplete.Append(PostImportTasks);
 
-	if (PreCompletionTask.GetReference())
+	if (PreCompletionTask != INTERCHANGE_INVALID_TASK_ID)
 	{
 		TasksToComplete.Add(PreCompletionTask);
 	}
 	
-	if (CompletionTask.GetReference())
+	if (CompletionTask != INTERCHANGE_INVALID_TASK_ID)
 	{
 		//Completion task will make sure any created asset before canceling will be mark for delete
 		TasksToComplete.Add(CompletionTask);
@@ -589,7 +599,6 @@ FGraphEventArray UE::Interchange::FImportAsyncHelper::GetCompletionTaskGraphEven
 void UE::Interchange::FImportAsyncHelper::InitCancel()
 {
 	bCancel = true;
-	ReleaseTranslatorsSource();
 }
 
 void UE::Interchange::FImportAsyncHelper::CleanUp()
@@ -667,10 +676,7 @@ bool UE::Interchange::FImportResult::IsValid() const
 void UE::Interchange::FImportResult::SetInProgress()
 {
 	EStatus ExpectedStatus = EStatus::Invalid;
-	if (ImportStatus.compare_exchange_strong(ExpectedStatus, EStatus::InProgress))
-	{
-		GraphEvent = FGraphEvent::CreateGraphEvent();
-	}
+	ImportStatus.compare_exchange_strong(ExpectedStatus, EStatus::InProgress);
 }
 
 void UE::Interchange::FImportResult::SetDone()
@@ -701,13 +707,12 @@ void UE::Interchange::FImportResult::SetDone()
 				WeakObjects.Emplace(Object);
 			}
 
-			// call the callbacks on the game thread
-			Async(EAsyncExecution::TaskGraphMainThread, [InWeakObjects = MoveTemp(WeakObjects), ImportDoneNative = OnImportDoneNative, ImportDone = OnImportDone]()
+			if(ensure(IsInGameThread()))
 			{
 				TArray<UObject*> ValidObjects;
-				ValidObjects.Reserve(InWeakObjects.Num());
+				ValidObjects.Reserve(WeakObjects.Num());
 
-				for (const TWeakObjectPtr<UObject>& WeakObject : InWeakObjects)
+				for (const TWeakObjectPtr<UObject>& WeakObject : WeakObjects)
 				{
 					if (UObject* ValidObject = WeakObject.Get())
 					{
@@ -715,13 +720,10 @@ void UE::Interchange::FImportResult::SetDone()
 					}
 				}
 
-				ImportDoneNative.ExecuteIfBound(ValidObjects);
-				ImportDone.ExecuteIfBound(ValidObjects);
-			});
+				OnImportDoneNative.ExecuteIfBound(ValidObjects);
+				OnImportDone.ExecuteIfBound(ValidObjects);
+			}
 		}
-
-
-		GraphEvent->DispatchSubsequents();
 	}
 }
 
@@ -729,14 +731,17 @@ void UE::Interchange::FImportResult::WaitUntilDone(bool bSynchronous /*= false*/
 {
 	if (ImportStatus == EStatus::InProgress)
 	{
-		if (bSynchronous)
+		//Pin the weak ptr, do not hold the shared ptr until the end of the import, simply get the completion task id
+		TArray<uint64> TasksIds;
 		{
-			FTaskGraphInterface::Get().ProcessThreadUntilIdle(ENamedThreads::GameThread_Local);
-		}
-		else
-		{
-			FTaskGraphInterface::Get().WaitUntilTaskCompletes(GraphEvent);
+			TSharedPtr<FImportAsyncHelper, ESPMode::ThreadSafe> PinAsyncHelper = AsyncHelper.Pin();
+			if (PinAsyncHelper.IsValid())
+			{
+				TasksIds = PinAsyncHelper->GetCompletionTaskGraphEvent();
+			}
 		}
+		//Wait until the whole import is done
+		UE::Interchange::FInterchangeTaskSystem::Get().WaitUntilTasksComplete(TasksIds);
 	}
 }
 
@@ -809,6 +814,11 @@ void UE::Interchange::FImportResult::AddReferencedObjects(FReferenceCollector& C
 	Collector.AddReferencedObject(Results);
 }
 
+void UE::Interchange::FImportResult::SetAsyncHelper(TWeakPtr<FImportAsyncHelper> InAsyncHelper)
+{
+	AsyncHelper = InAsyncHelper;
+}
+
 UInterchangePipelineBase* UE::Interchange::GeneratePipelineInstance(const FSoftObjectPath& PipelineInstance)
 {
 	UPackage* PipelineInstancePackage = GetTransientPackage();
@@ -929,22 +939,6 @@ UInterchangeManager& UInterchangeManager::GetInterchangeManager()
 
 		bIsCreatingSingleton = false;
 
-		InterchangeManager->GCEndDelegate = FCoreUObjectDelegates::GetPostGarbageCollect().AddLambda([]()
-			{
-				if (IsInterchangeImportEnabled() && InterchangeManager.IsValid())
-				{
-					InterchangeManager->StartQueuedTasks(InterchangeManager->bGCEndDelegateCancellAllTask);
-				}
-			});
-		InterchangeManager->GCPreDelegate = FCoreUObjectDelegates::GetPreGarbageCollectDelegate().AddLambda([]()
-			{
-				if (IsInterchangeImportEnabled() && InterchangeManager.IsValid())
-				{
-					const bool bCancel = !GIsEditor;
-					InterchangeManager->WaitUntilAllTasksDone(bCancel);
-				}
-			});
-
 		//We cancel any running task when we pre exit the engine
 		FCoreDelegates::OnEnginePreExit.AddLambda([]()
 		{
@@ -964,17 +958,6 @@ UInterchangeManager& UInterchangeManager::GetInterchangeManager()
 				InterchangeManager->WaitUntilAllTasksDone(bCancel);
 			}
 
-			//Remove any delegate
-			if (InterchangeManager->GCEndDelegate.IsValid())
-			{
-				FCoreUObjectDelegates::GetPostGarbageCollect().Remove(InterchangeManager->GCEndDelegate);
-				InterchangeManager->GCEndDelegate.Reset();
-			}
-			if(InterchangeManager->GCPreDelegate.IsValid())
-			{
-				FCoreUObjectDelegates::GetPreGarbageCollectDelegate().Remove(InterchangeManager->GCPreDelegate);
-				InterchangeManager->GCPreDelegate.Reset();
-			}
 			//Task should have been cancel in the Engine pre exit callback
 			ensure(InterchangeManager->ImportTasks.Num() == 0);
 			InterchangeManager->OnPreDestroyInterchangeManager.Broadcast();
@@ -1470,7 +1453,7 @@ void UInterchangeManager::StartQueuedTasks(bool bCancelAllTasks /*= false*/)
 			}
 
 			//Create/Start import tasks
-			FGraphEventArray PipelinePrerequistes;
+			TArray<uint64> PipelinePrerequistes;
 			if(QueuedTaskData.AsyncHelper->TranslatorTasks.Num() == 0)
 			{
 				check(QueuedTaskData.AsyncHelper->Translators.Num() == QueuedTaskData.AsyncHelper->SourceDatas.Num());
@@ -1478,18 +1461,20 @@ void UInterchangeManager::StartQueuedTasks(bool bCancelAllTasks /*= false*/)
 				{
 					//Log the source we begin importing
 					UE_LOG(LogInterchangeEngine, Display, TEXT("Interchange start importing source [%s]"), *QueuedTaskData.AsyncHelper->SourceDatas[SourceDataIndex]->GetFilename());
-					int32 TranslatorTaskIndex = QueuedTaskData.AsyncHelper->TranslatorTasks.Add(TGraphTask<UE::Interchange::FTaskTranslator>::CreateTask().ConstructAndDispatchWhenReady(SourceDataIndex, WeakAsyncHelper));
+					TSharedPtr<UE::Interchange::FTaskTranslator, ESPMode::ThreadSafe> TaskTranslator = MakeShared<UE::Interchange::FTaskTranslator, ESPMode::ThreadSafe>(SourceDataIndex, WeakAsyncHelper);
+					int32 TranslatorTaskIndex = QueuedTaskData.AsyncHelper->TranslatorTasks.Add(UE::Interchange::FInterchangeTaskSystem::Get().AddTask(TaskTranslator));
 					PipelinePrerequistes.Add(QueuedTaskData.AsyncHelper->TranslatorTasks[TranslatorTaskIndex]);
 				}
 			}
 
-			FGraphEventArray GraphParsingPrerequistes;
+			TArray<uint64> GraphParsingPrerequistes;
 			for (int32 GraphPipelineIndex = 0; GraphPipelineIndex < QueuedTaskData.AsyncHelper->Pipelines.Num(); ++GraphPipelineIndex)
 			{
 				UInterchangePipelineBase* GraphPipeline = QueuedTaskData.AsyncHelper->Pipelines[GraphPipelineIndex];
 				TWeakObjectPtr<UInterchangePipelineBase> WeakPipelinePtr = GraphPipeline;
-				int32 GraphPipelineTaskIndex = INDEX_NONE;
-				GraphPipelineTaskIndex = QueuedTaskData.AsyncHelper->PipelineTasks.Add(TGraphTask<UE::Interchange::FTaskPipeline>::CreateTask(&PipelinePrerequistes).ConstructAndDispatchWhenReady(WeakPipelinePtr, WeakAsyncHelper));
+
+				TSharedPtr<UE::Interchange::FTaskPipeline, ESPMode::ThreadSafe> TaskPipeline = MakeShared<UE::Interchange::FTaskPipeline, ESPMode::ThreadSafe>(WeakPipelinePtr, WeakAsyncHelper);
+				int32 GraphPipelineTaskIndex = QueuedTaskData.AsyncHelper->PipelineTasks.Add(UE::Interchange::FInterchangeTaskSystem::Get().AddTask(TaskPipeline, PipelinePrerequistes));
 				//Ensure we run the pipeline in the same order we create the task, since pipeline modify the node container, its important that its not process in parallel, Adding the one we start to the prerequisites
 				//is the way to go here
 				PipelinePrerequistes.Add(QueuedTaskData.AsyncHelper->PipelineTasks[GraphPipelineTaskIndex]);
@@ -1498,16 +1483,15 @@ void UInterchangeManager::StartQueuedTasks(bool bCancelAllTasks /*= false*/)
 				GraphParsingPrerequistes.Add(QueuedTaskData.AsyncHelper->PipelineTasks[GraphPipelineTaskIndex]);
 			}
 
-			if (GraphParsingPrerequistes.Num() > 0)
-			{
-				QueuedTaskData.AsyncHelper->ParsingTask = TGraphTask<UE::Interchange::FTaskParsing>::CreateTask(&GraphParsingPrerequistes).ConstructAndDispatchWhenReady(this, WeakAsyncHelper);
-			}
-			else
+			if (GraphParsingPrerequistes.Num() == 0)
 			{
 				//Fallback on the translator pipeline prerequisites (translator must be done if there is no pipeline)
-				QueuedTaskData.AsyncHelper->ParsingTask = TGraphTask<UE::Interchange::FTaskParsing>::CreateTask(&PipelinePrerequistes).ConstructAndDispatchWhenReady(this, WeakAsyncHelper);
+				GraphParsingPrerequistes.Append(PipelinePrerequistes);
 			}
 
+			TSharedPtr<UE::Interchange::FTaskParsing, ESPMode::ThreadSafe> TaskParsing = MakeShared<UE::Interchange::FTaskParsing, ESPMode::ThreadSafe>(this, WeakAsyncHelper);
+			QueuedTaskData.AsyncHelper->ParsingTask = UE::Interchange::FInterchangeTaskSystem::Get().AddTask(TaskParsing, GraphParsingPrerequistes);
+
 			//The graph parsing task will create the FCreateAssetTask that will run after them, the FAssetImportTask will call the appropriate Post asset import pipeline when the asset is completed
 		}
 	}
@@ -1917,13 +1901,14 @@ UInterchangeManager::ImportInternal(const FString& ContentPath, const UInterchan
 				Progress.MakeDialog();
 				Progress.EnterProgressFrame(1.f);
 				//Translate the source
-				FGraphEventArray PipelinePrerequistes;
 				for (int32 SourceDataIndex = 0; SourceDataIndex < AsyncHelper->SourceDatas.Num(); ++SourceDataIndex)
 				{
+					constexpr uint64 PipelinePrerequistes = INTERCHANGE_INVALID_TASK_ID;
 					//Log the source we begin importing
 					UE_LOG(LogInterchangeEngine, Display, TEXT("Interchange start importing source [%s]"), *AsyncHelper->SourceDatas[SourceDataIndex]->GetFilename());
-					int32 TranslatorTaskIndex = AsyncHelper->TranslatorTasks.Add(TGraphTask<UE::Interchange::FTaskTranslator>::CreateTask(nullptr, ENamedThreads::GameThread_Local).ConstructAndDispatchWhenReady(SourceDataIndex, AsyncHelper));
-					AsyncHelper->TranslatorTasks[TranslatorTaskIndex]->Wait();
+					UE::Interchange::FTaskTranslator TaskTranslator(SourceDataIndex, AsyncHelper);
+					TaskTranslator.Execute();
+					AsyncHelper->TranslatorTasks.Add(PipelinePrerequistes);
 				}
 				Progress.EnterProgressFrame(1.f);
 			}
@@ -2390,11 +2375,13 @@ TSharedRef<UE::Interchange::FImportAsyncHelper, ESPMode::ThreadSafe> UInterchang
 	AsyncHelper->AssetImportResult->OnObjectDoneNative = ImportAssetParameters.OnAssetDoneNative;
 	AsyncHelper->AssetImportResult->OnImportDone = ImportAssetParameters.OnAssetsImportDone;
 	AsyncHelper->AssetImportResult->OnImportDoneNative = ImportAssetParameters.OnAssetsImportDoneNative;
+	AsyncHelper->AssetImportResult->SetAsyncHelper(AsyncHelper);
 
 	AsyncHelper->SceneImportResult->OnObjectDone = ImportAssetParameters.OnSceneObjectDone;
 	AsyncHelper->SceneImportResult->OnObjectDoneNative = ImportAssetParameters.OnSceneObjectDoneNative;
 	AsyncHelper->SceneImportResult->OnImportDone = ImportAssetParameters.OnSceneImportDone;
 	AsyncHelper->SceneImportResult->OnImportDoneNative = ImportAssetParameters.OnSceneImportDoneNative;
+	AsyncHelper->SceneImportResult->SetAsyncHelper(AsyncHelper);
 
 	AsyncHelper->AssetImportResult->SetInProgress();
 
@@ -2464,7 +2451,7 @@ void UInterchangeManager::ReleaseAsyncHelper(TWeakPtr<UE::Interchange::FImportAs
 		check(!AsyncHelperSharedPtr.IsValid() || AsyncHelperSharedPtr->bCancel);
 	}
 
-	int32 ImportTaskNumber = ImportTasks.Num();
+	int32 ImportTaskNumber = ImportTasks.Num() + QueueTaskCount;
 	FString ImportTaskNumberStr = TEXT(" (") + FString::FromInt(ImportTaskNumber) + TEXT(")");
 	if (ImportTaskNumber == 0)
 	{
@@ -2722,10 +2709,10 @@ void UInterchangeManager::WaitUntilAllTasksDone(bool bCancel)
 		if (AsyncHelper.IsValid())
 		{
 			TWeakPtr<UE::Interchange::FImportAsyncHelper, ESPMode::ThreadSafe> WeakAsyncHelper = AsyncHelper;
-			FGraphEventArray TasksToComplete = AsyncHelper->GetCompletionTaskGraphEvent();
+			TArray<uint64> TasksToComplete = AsyncHelper->GetCompletionTaskGraphEvent();
 			//Release the shared pointer before waiting to be sure the async helper can be destroy in the completion task
 			AsyncHelper = nullptr;
-			FTaskGraphInterface::Get().WaitUntilTasksComplete(TasksToComplete, ENamedThreads::GameThread);
+			UE::Interchange::FInterchangeTaskSystem::Get().WaitUntilTasksComplete(TasksToComplete);
 			//We verify that the weak pointer is invalid after the task completed
 			ensure(!WeakAsyncHelper.IsValid());
 		}
diff --git a/Engine/Source/Runtime/Interchange/Engine/Private/InterchangeTaskSystem.cpp b/Engine/Source/Runtime/Interchange/Engine/Private/InterchangeTaskSystem.cpp
new file mode 100644
--- /dev/null
+++ b/Engine/Source/Runtime/Interchange/Engine/Private/InterchangeTaskSystem.cpp
@@ -0,0 +1,365 @@
+// Copyright Epic Games, Inc. All Rights Reserved.
+
+#include "InterchangeTaskSystem.h"
+
+#include "AssetCompilingManager.h"
+#include "Async/Async.h"
+#include "Async/TaskGraphInterfaces.h"
+#include "Containers/Array.h"
+#include "Containers/Map.h"
+#include "Containers/Ticker.h"
+#include "CoreGlobals.h"
+#include "HAL/CriticalSection.h"
+#include "HAL/Thread.h"
+#include "HAL/ThreadManager.h"
+#include "InterchangeEngineLogPrivate.h"
+#include "Misc/CoreDelegates.h"
+#include "Templates/SharedPointer.h"
+#include "UObject/GarbageCollection.h"
+
+bool UE::Interchange::FInterchangeTaskSystem::bIsCreatingSingleton = false;
+
+bool UE::Interchange::FInterchangeTaskBase::SetTaskStatus(EInterchangeTaskStatus NewTaskStatus)
+{
+	FScopeLock TaskStatusLock(&TaskStatusCriticalSection);
+	if (TaskStatus == NewTaskStatus)
+	{
+		return true;
+	}
+	//If status is done, we cannot change it
+	if (!ensure(TaskStatus != EInterchangeTaskStatus::Done))
+	{
+		return false;
+	}
+	TaskStatus = NewTaskStatus;
+	return true;
+}
+
+void UE::Interchange::FInterchangeTaskBase::SetPrerequisites(const TArray<uint64>& InPrerequisiteTasks)
+{
+	check(GetTaskStatus() == EInterchangeTaskStatus::Waiting);
+	PrerequisiteTasks = InPrerequisiteTasks;
+}
+
+void UE::Interchange::FInterchangeTaskBase::Wait() const
+{
+	FInterchangeTaskSystem::Get().WaitUntilTasksComplete({ TaskId });
+}
+
+UE::Interchange::FInterchangeTaskSystem::FInterchangeTaskSystem()
+{
+	if (!bIsCreatingSingleton)
+	{
+		ensureMsgf(bIsCreatingSingleton, TEXT("Interchange task system is a singleton you must call UE::Interchange::FInterchangeTaskSystem::Get() to access it."));
+	}
+}
+
+UE::Interchange::FInterchangeTaskSystem& UE::Interchange::FInterchangeTaskSystem::Get()
+{
+	static TSharedPtr<FInterchangeTaskSystem, ESPMode::ThreadSafe> TaskSystemPtr;
+	static bool bEndOfSingletonLife = false;
+
+	if (!TaskSystemPtr.IsValid())
+	{
+		//Cannot create the singleton outside of the game thread
+		check(IsInGameThread());
+
+		//Ensure before crashing in case this function is call after the engine exit
+		if (bEndOfSingletonLife)
+		{
+			ensure(!bEndOfSingletonLife);
+		}
+
+		bIsCreatingSingleton = true;
+		TaskSystemPtr = MakeShared<FInterchangeTaskSystem, ESPMode::ThreadSafe>().ToSharedPtr();
+		bIsCreatingSingleton = false;
+
+
+
+		TaskSystemPtr->TickTickerHandle = FTSTicker::GetCoreTicker().AddTicker(TEXT("InterchangeTaskSystemTickHandle"), 0.0f, [](float DeltaTime)
+			{
+				TaskSystemPtr->Tick();
+				return true; //Return true will keep the tick active
+			});
+			
+		FCoreDelegates::OnEnginePreExit.AddLambda([]()
+			{
+				if (!ensure(TaskSystemPtr.IsValid()))
+				{
+					return;
+				}
+
+				//We cancel any running task when we pre exit the engine
+				for (TPair<uint64, TSharedPtr<FInterchangeTaskBase, ESPMode::ThreadSafe>> TaskInfo : TaskSystemPtr->TaskPerIdMap)
+				{
+					//Cancel any task
+					TaskSystemPtr->CancelTask(TaskInfo.Key, false);
+				}
+
+				if (TaskSystemPtr->TickTickerHandle.IsValid())
+				{
+					FTSTicker::GetCoreTicker().RemoveTicker(TaskSystemPtr->TickTickerHandle);
+					TaskSystemPtr->TickTickerHandle.Reset();
+				}
+
+				//Free the singleton
+				bEndOfSingletonLife = true;
+				TaskSystemPtr = nullptr;
+			});
+
+	}
+	//Return the singleton
+	check(TaskSystemPtr.IsValid());
+	return *TaskSystemPtr.Get();
+}
+
+uint64 UE::Interchange::FInterchangeTaskSystem::AddTask(TSharedPtr<FInterchangeTaskBase, ESPMode::ThreadSafe> Task, TArray<uint64>& TaskPrerequisites)
+{
+	FScopeLock TaskPerIdMapLock(&TaskPerIdMapCriticalSection);
+	FScopeLock PriorityValueLock(&PriorityValueCriticalSection);
+
+	Task->SetTaskStatus(EInterchangeTaskStatus::Waiting);
+	Task->TaskId = PriorityValue++;
+	Task->SetPrerequisites(TaskPrerequisites);
+	TaskPerIdMap.Add(Task->TaskId, Task);
+
+	return Task->TaskId;
+}
+
+uint64 UE::Interchange::FInterchangeTaskSystem::AddTask(TSharedPtr<FInterchangeTaskBase, ESPMode::ThreadSafe> Task)
+{
+	TArray<uint64> EmptyArray;
+	return AddTask(Task, EmptyArray);
+}
+
+UE::Interchange::EInterchangeTaskStatus UE::Interchange::FInterchangeTaskSystem::GetTaskStatus(const uint64 TaskId) const
+{
+	FScopeLock TaskPerIdMapLock(&TaskPerIdMapCriticalSection);
+	if (TaskPerIdMap.Contains(TaskId))
+	{
+		TSharedPtr<FInterchangeTaskBase, ESPMode::ThreadSafe> Task = TaskPerIdMap.FindChecked(TaskId);
+		return Task->GetTaskStatus();
+	}
+
+	//If task doesn't exist we assume its done
+	return EInterchangeTaskStatus::Done;
+}
+
+void UE::Interchange::FInterchangeTaskSystem::InternalAddCancelRequestNoLock(const uint64 TaskId, bool bCancelPrerequisites)
+{
+	if (!TaskPerIdMap.Contains(TaskId))
+	{
+		//Nothing to cancel
+		return;
+	}
+
+	TSharedPtr<FInterchangeTaskBase, ESPMode::ThreadSafe> Task = TaskPerIdMap.FindChecked(TaskId);
+	uint64& CanceltaskIdRequest = CancelTaskRequests.AddDefaulted_GetRef();
+	CanceltaskIdRequest = Task->GetTaskId();
+	if (bCancelPrerequisites)
+	{
+		for (uint64 PrerequisiteTaskId : Task->GetPrerequisiteTasks())
+		{
+			InternalAddCancelRequestNoLock(PrerequisiteTaskId, bCancelPrerequisites);
+		}
+	}
+}
+
+void UE::Interchange::FInterchangeTaskSystem::CancelTask(const uint64 TaskId, bool bCancelPrerequisites)
+{
+	FScopeLock TaskPerIdMapLock(&TaskPerIdMapCriticalSection);
+	FScopeLock CancelTaskRequestsLock(&CancelTaskRequestsCriticalSection);
+	InternalAddCancelRequestNoLock(TaskId, bCancelPrerequisites);
+}
+
+void UE::Interchange::FInterchangeTaskSystem::WaitUntilTasksComplete(const TArray<uint64>& TasksToComplete)
+{
+	check(IsInGameThread());
+	const FTimespan MicroSecondTimespan = FTimespan::FromMicroseconds(1);
+	const float DeltaSecond = static_cast<float>(MicroSecondTimespan.GetTotalSeconds());
+	//Loop
+	FPlatformProcess::ConditionalSleep([this, &TasksToComplete, &MicroSecondTimespan]()
+		{
+			LLM_SCOPE_BYNAME(TEXT("Interchange"));
+			bool bTasksCompleted = true;
+			//Look if all tasks are completed
+			{
+				FScopeLock TaskPerIdMapLock(&TaskPerIdMapCriticalSection);
+				for (uint64 TaskId : TasksToComplete)
+				{
+					if (TaskPerIdMap.Contains(TaskId))
+					{
+						TSharedPtr<FInterchangeTaskBase, ESPMode::ThreadSafe> Task = TaskPerIdMap.FindChecked(TaskId);
+						if (Task->GetTaskStatus() != EInterchangeTaskStatus::Done)
+						{
+							bTasksCompleted = false;
+							break;
+						}
+					}
+				}
+			}
+			if (!bTasksCompleted)
+			{
+				//Tick all systems our tasks can depends to avoid stalling the game thread
+				
+				//We cannot tick the "core ticker" since we are already inside a tick and it will assert.
+				//We cannot tick "runnable thread" it create some issue with the garbage collector tasks
+
+				//Tick Interchange task system
+				Tick();
+
+				//Tick task graph game thread tasks
+				FTaskGraphInterface::Get().ProcessThreadUntilIdle(ENamedThreads::GameThread);
+
+				//Tick the all compiling manager
+				FAssetCompilingManager::Get().ProcessAsyncTasks(true);
+			}
+			return bTasksCompleted;
+		}, DeltaSecond);
+}
+
+void UE::Interchange::FInterchangeTaskSystem::Tick()
+{
+	//Tick must be execute on the game thread
+	check(IsInGameThread());
+	
+	//Tick is suppose to be a safe place where we can create and manipulate UObjects
+	ensure(!IsGarbageCollecting());
+#if WITH_EDITOR
+	ensure(!UE::GetIsEditorLoadingPackage());
+#endif
+
+	FScopeLock TaskPerIdMapLock(&TaskPerIdMapCriticalSection);
+
+	//Execute the cancel task
+	{
+		FScopeLock CancelTaskRequestsLock(&CancelTaskRequestsCriticalSection);
+		for (uint64 CancelTaskId : CancelTaskRequests)
+		{
+			if (!TaskPerIdMap.Contains(CancelTaskId))
+			{
+				continue;
+			}
+			TSharedPtr<FInterchangeTaskBase, ESPMode::ThreadSafe> Task = TaskPerIdMap.FindChecked(CancelTaskId);
+			if (Task->GetTaskStatus() == EInterchangeTaskStatus::Executing)
+			{
+				//Wait until the task is done
+				Task->Wait();
+			}
+			Task->SetTaskStatus(EInterchangeTaskStatus::Done);
+		}
+	}
+	if (TaskPerIdMap.Num() == 0)
+	{
+		//We do not have any task waiting or running
+		ReleaseAndDoneTasks.Reset();
+		return;
+	}
+
+	TArray<uint64> WaitingGameThreadTick;
+	TArray<uint64> WaitingTaskGraph;
+	TArray<uint64> DoneTasks;
+	for (TPair<uint64, TSharedPtr<FInterchangeTaskBase, ESPMode::ThreadSafe>>& TaskInfo : TaskPerIdMap)
+	{
+		TSharedPtr<FInterchangeTaskBase, ESPMode::ThreadSafe> Task = TaskInfo.Value;
+		if (Task->GetTaskStatus() == EInterchangeTaskStatus::Waiting)
+		{
+			if (Task->GetTaskThread() == EInterchangeTaskThread::GameThread)
+			{
+				WaitingGameThreadTick.Add(TaskInfo.Key);
+			}
+			else if (Task->GetTaskThread() == EInterchangeTaskThread::AsyncThread)
+			{
+				WaitingTaskGraph.Add(TaskInfo.Key);
+			}
+		}
+		else if (Task->GetTaskStatus() == EInterchangeTaskStatus::Done)
+		{
+			DoneTasks.Add(TaskInfo.Key);
+		}
+	}
+
+	//Sort all waiting task by priority order, prerequisite task are always added before.
+	//The priority is the task id, lower mean must be execute before higher task id
+	auto SortLambda = [](const uint64& A, const uint64& B)
+		{
+			return A < B;
+		};
+
+	WaitingGameThreadTick.Sort(SortLambda);
+	WaitingTaskGraph.Sort(SortLambda);
+
+	auto UpdateWaitingTask = [this, &DoneTasks](const TArray<uint64>& TaskIds)
+		{
+			for (uint64 TaskId : TaskIds)
+			{
+				bool bAllPrerequisitesCompleted = true;
+				TSharedPtr<FInterchangeTaskBase, ESPMode::ThreadSafe> Task = TaskPerIdMap.FindChecked(TaskId);
+				const TArray<uint64>& TaskPrerequisites = Task->GetPrerequisiteTasks();
+				for (uint64 PrerequisiteTaskId : TaskPrerequisites)
+				{
+					if (ReleaseAndDoneTasks.Contains(PrerequisiteTaskId)
+						|| DoneTasks.Contains(PrerequisiteTaskId))
+					{
+						continue;
+					}
+					//if we have one waiting or executing task, prerequisites are not completed
+					if (TaskPerIdMap.Contains(PrerequisiteTaskId))
+					{
+						TSharedPtr<FInterchangeTaskBase, ESPMode::ThreadSafe> PrerequisiteTask = TaskPerIdMap.FindChecked(PrerequisiteTaskId);
+						if (PrerequisiteTask->GetTaskStatus() != EInterchangeTaskStatus::Done)
+						{
+							bAllPrerequisitesCompleted = false;
+							break;
+						}
+					}
+				}
+				//Skip this task if not all prerequisite are completed
+				if (!bAllPrerequisitesCompleted)
+				{
+					continue;
+				}
+
+				auto TaskCompletionHandler = [Task]()
+					{
+						//Only set the task to done if the status is executing
+						if (Task->GetTaskStatus() == EInterchangeTaskStatus::Executing)
+						{
+							Task->SetTaskStatus(EInterchangeTaskStatus::Done);
+						}
+					};
+
+				if (Task->GetTaskThread() == EInterchangeTaskThread::GameThread)
+				{
+					//Prerequisites completed, we can execute the task
+					Task->SetTaskStatus(EInterchangeTaskStatus::Executing);
+					Task->Execute();
+					TaskCompletionHandler();
+				}
+				else if (Task->GetTaskThread() == EInterchangeTaskThread::AsyncThread)
+				{
+					Task->SetTaskStatus(EInterchangeTaskStatus::Executing);
+					Task->Future = Async(EAsyncExecution::TaskGraph
+					, [Task]()->uint64 //Task callback
+					{
+						Task->Execute();
+						return Task->GetTaskId();
+					}
+					, TaskCompletionHandler); //Completion callback
+				}
+			}
+		};
+
+	//Update GameThread
+	UpdateWaitingTask(WaitingGameThreadTick);
+
+	//Update TaskGraph task
+	UpdateWaitingTask(WaitingTaskGraph);
+
+	//Release Done task
+	for (uint64 TaskId : DoneTasks)
+	{
+		TaskPerIdMap.Remove(TaskId);
+		ReleaseAndDoneTasks.Add(TaskId);
+	}
+}
diff --git a/Engine/Source/Runtime/Interchange/Engine/Private/Tasks/InterchangeTaskCompletion.cpp b/Engine/Source/Runtime/Interchange/Engine/Private/Tasks/InterchangeTaskCompletion.cpp
--- a/Engine/Source/Runtime/Interchange/Engine/Private/Tasks/InterchangeTaskCompletion.cpp
+++ b/Engine/Source/Runtime/Interchange/Engine/Private/Tasks/InterchangeTaskCompletion.cpp
@@ -20,14 +20,16 @@
 #include "UObject/ObjectMacros.h"
 #include "UObject/WeakObjectPtrTemplates.h"
 
-void UE::Interchange::FTaskPreCompletion::DoTask(ENamedThreads::Type CurrentThread, const FGraphEventRef& MyCompletionGraphEvent)
+void UE::Interchange::FTaskPreCompletion_GameThread::Execute()
 {
-	TRACE_CPUPROFILER_EVENT_SCOPE(UE::Interchange::FTaskPreCompletion::DoTask)
+	TRACE_CPUPROFILER_EVENT_SCOPE(UE::Interchange::FTaskPreCompletion_GameThread::DoTask)
 #if INTERCHANGE_TRACE_ASYNCHRONOUS_TASK_ENABLED
 	INTERCHANGE_TRACE_ASYNCHRONOUS_TASK(PreCompletion)
 #endif
 	
 	LLM_SCOPE_BYNAME(TEXT("Interchange"));
+	
+	check(IsInGameThread());
 
 	TSharedPtr<FImportAsyncHelper, ESPMode::ThreadSafe> AsyncHelper = WeakAsyncHelper.Pin();
 	check(AsyncHelper.IsValid());
@@ -175,14 +177,15 @@ void UE::Interchange::FTaskPreCompletion::DoTask(ENamedThreads::Type CurrentThre
 }
 
 
-void UE::Interchange::FTaskCompletion::DoTask(ENamedThreads::Type CurrentThread, const FGraphEventRef& MyCompletionGraphEvent)
+void UE::Interchange::FTaskCompletion_GameThread::Execute()
 {
-	TRACE_CPUPROFILER_EVENT_SCOPE(UE::Interchange::FTaskCompletion::DoTask)
+	TRACE_CPUPROFILER_EVENT_SCOPE(UE::Interchange::FTaskCompletion_GameThread::DoTask)
 #if INTERCHANGE_TRACE_ASYNCHRONOUS_TASK_ENABLED
 	INTERCHANGE_TRACE_ASYNCHRONOUS_TASK(Completion)
 #endif
 
 	LLM_SCOPE_BYNAME(TEXT("Interchange"));
+	check(IsInGameThread());
 
 	TSharedPtr<FImportAsyncHelper, ESPMode::ThreadSafe> AsyncHelper = WeakAsyncHelper.Pin();
 	check(AsyncHelper.IsValid());
diff --git a/Engine/Source/Runtime/Interchange/Engine/Private/Tasks/InterchangeTaskCompletion.h b/Engine/Source/Runtime/Interchange/Engine/Private/Tasks/InterchangeTaskCompletion.h
--- a/Engine/Source/Runtime/Interchange/Engine/Private/Tasks/InterchangeTaskCompletion.h
+++ b/Engine/Source/Runtime/Interchange/Engine/Private/Tasks/InterchangeTaskCompletion.h
@@ -2,9 +2,9 @@
 
 #pragma once
 
-#include "Async/TaskGraphInterfaces.h"
 #include "CoreMinimal.h"
 #include "InterchangeManager.h"
+#include "InterchangeTaskSystem.h"
 #include "Stats/Stats.h"
 #include "UObject/WeakObjectPtrTemplates.h"
 
@@ -12,73 +12,44 @@ namespace UE
 {
 	namespace Interchange
 	{
-		class FTaskPreCompletion
+		class FTaskPreCompletion_GameThread : public FInterchangeTaskBase
 		{
 		private:
 			UInterchangeManager* InterchangeManager;
 			TWeakPtr<FImportAsyncHelper, ESPMode::ThreadSafe> WeakAsyncHelper;
 		public:
-			FTaskPreCompletion(UInterchangeManager* InInterchangeManager, TWeakPtr<FImportAsyncHelper, ESPMode::ThreadSafe> InAsyncHelper)
+			FTaskPreCompletion_GameThread(UInterchangeManager* InInterchangeManager, TWeakPtr<FImportAsyncHelper, ESPMode::ThreadSafe> InAsyncHelper)
 				: InterchangeManager(InInterchangeManager)
 				, WeakAsyncHelper(InAsyncHelper)
 			{
 			}
 
-			FORCEINLINE ENamedThreads::Type GetDesiredThread()
+			virtual EInterchangeTaskThread GetTaskThread() const override
 			{
-				TSharedPtr<FImportAsyncHelper, ESPMode::ThreadSafe> AsyncHelper = WeakAsyncHelper.Pin();
-				if (AsyncHelper.IsValid() && AsyncHelper->bRunSynchronous)
-				{
-					return ENamedThreads::GameThread_Local;
-				}
-				return ENamedThreads::GameThread;
-			}
-			static FORCEINLINE ESubsequentsMode::Type GetSubsequentsMode()
-			{
-				return ESubsequentsMode::TrackSubsequents;
-			}
-
-			FORCEINLINE TStatId GetStatId() const
-			{
-				RETURN_QUICK_DECLARE_CYCLE_STAT(FTaskPreCompletion, STATGROUP_TaskGraphTasks);
+				return EInterchangeTaskThread::GameThread;
 			}
 
-			void DoTask(ENamedThreads::Type CurrentThread, const FGraphEventRef& MyCompletionGraphEvent);
+			virtual void Execute() override;
 		};
 
-		class FTaskCompletion
+		class FTaskCompletion_GameThread : public FInterchangeTaskBase
 		{
 		private:
 			UInterchangeManager* InterchangeManager;
 			TWeakPtr<FImportAsyncHelper, ESPMode::ThreadSafe> WeakAsyncHelper;
 		public:
-			FTaskCompletion(UInterchangeManager* InInterchangeManager, TWeakPtr<FImportAsyncHelper, ESPMode::ThreadSafe> InAsyncHelper)
+			FTaskCompletion_GameThread(UInterchangeManager* InInterchangeManager, TWeakPtr<FImportAsyncHelper, ESPMode::ThreadSafe> InAsyncHelper)
 				: InterchangeManager(InInterchangeManager)
 				, WeakAsyncHelper(InAsyncHelper)
 			{
 			}
 
-			FORCEINLINE ENamedThreads::Type GetDesiredThread()
-			{
-				TSharedPtr<FImportAsyncHelper, ESPMode::ThreadSafe> AsyncHelper = WeakAsyncHelper.Pin();
-				if (AsyncHelper.IsValid() && AsyncHelper->bRunSynchronous)
-				{
-					return ENamedThreads::GameThread_Local;
-				}
-				return ENamedThreads::GameThread;
-			}
-			static FORCEINLINE ESubsequentsMode::Type GetSubsequentsMode()
-			{
-				//In case we need to know when the task is done we need track subsequent to get a valid FGraphEventRef when we create the task
-				return ESubsequentsMode::TrackSubsequents;
-			}
-
-			FORCEINLINE TStatId GetStatId() const
+			virtual EInterchangeTaskThread GetTaskThread() const override
 			{
-				RETURN_QUICK_DECLARE_CYCLE_STAT(FTaskCompletion, STATGROUP_TaskGraphTasks);
+				return EInterchangeTaskThread::GameThread;
 			}
 
-			void DoTask(ENamedThreads::Type CurrentThread, const FGraphEventRef& MyCompletionGraphEvent);
+			virtual void Execute() override;
 		};
 
 	} //ns Interchange
diff --git a/Engine/Source/Runtime/Interchange/Engine/Private/Tasks/InterchangeTaskCreateSceneObjects.cpp b/Engine/Source/Runtime/Interchange/Engine/Private/Tasks/InterchangeTaskCreateSceneObjects.cpp
--- a/Engine/Source/Runtime/Interchange/Engine/Private/Tasks/InterchangeTaskCreateSceneObjects.cpp
+++ b/Engine/Source/Runtime/Interchange/Engine/Private/Tasks/InterchangeTaskCreateSceneObjects.cpp
@@ -23,7 +23,7 @@
 #include "UObject/UObjectGlobals.h"
 #include "UObject/WeakObjectPtrTemplates.h"
 
-UE::Interchange::FTaskCreateSceneObjects::FTaskCreateSceneObjects(const FString& InPackageBasePath, const int32 InSourceIndex, TWeakPtr<FImportAsyncHelper> InAsyncHelper, TArrayView<UInterchangeFactoryBaseNode*> InFactoryNodes, const UClass* InFactoryClass)
+UE::Interchange::FTaskCreateSceneObjects_GameThread::FTaskCreateSceneObjects_GameThread(const FString& InPackageBasePath, const int32 InSourceIndex, TWeakPtr<FImportAsyncHelper> InAsyncHelper, TArrayView<UInterchangeFactoryBaseNode*> InFactoryNodes, const UClass* InFactoryClass)
 	: PackageBasePath(InPackageBasePath)
 	, SourceIndex(InSourceIndex)
 	, WeakAsyncHelper(InAsyncHelper)
@@ -33,32 +33,15 @@ UE::Interchange::FTaskCreateSceneObjects::FTaskCreateSceneObjects(const FString&
 	check(FactoryClass);
 }
 
-ENamedThreads::Type UE::Interchange::FTaskCreateSceneObjects::GetDesiredThread()
+void UE::Interchange::FTaskCreateSceneObjects_GameThread::Execute()
 {
-	// We are creating the factories in this task so it must execute on the GameThread.
-	// Also, there are no "CreatePackage Task" equivalent for scene objects right now, so the factories must create those on the game thread.
-	TSharedPtr<FImportAsyncHelper, ESPMode::ThreadSafe> AsyncHelper = WeakAsyncHelper.Pin();
-	if (AsyncHelper.IsValid() && AsyncHelper->bRunSynchronous)
-	{
-		return ENamedThreads::GameThread_Local;
-	}
-	return ENamedThreads::GameThread;
-}
-
-void UE::Interchange::FTaskCreateSceneObjects::DoTask(ENamedThreads::Type CurrentThread, const FGraphEventRef& MyCompletionGraphEvent)
-{
-	TRACE_CPUPROFILER_EVENT_SCOPE(UE::Interchange::FTaskCreateSceneObjects::DoTask)
+	TRACE_CPUPROFILER_EVENT_SCOPE(UE::Interchange::FTaskCreateSceneObjects_GameThread::DoTask)
 #if INTERCHANGE_TRACE_ASYNCHRONOUS_TASK_ENABLED
 	INTERCHANGE_TRACE_ASYNCHRONOUS_TASK(SpawnActor)
 #endif
 	using namespace UE::Interchange;
 
-	TOptional<FGCScopeGuard> GCScopeGuard;
-	if (!IsInGameThread())
-	{
-		GCScopeGuard.Emplace();
-	}
-
+	check(IsInGameThread());
 	TSharedPtr<FImportAsyncHelper> AsyncHelper = WeakAsyncHelper.Pin();
 	check(WeakAsyncHelper.IsValid());
 
diff --git a/Engine/Source/Runtime/Interchange/Engine/Private/Tasks/InterchangeTaskCreateSceneObjects.h b/Engine/Source/Runtime/Interchange/Engine/Private/Tasks/InterchangeTaskCreateSceneObjects.h
--- a/Engine/Source/Runtime/Interchange/Engine/Private/Tasks/InterchangeTaskCreateSceneObjects.h
+++ b/Engine/Source/Runtime/Interchange/Engine/Private/Tasks/InterchangeTaskCreateSceneObjects.h
@@ -2,8 +2,8 @@
 
 #pragma once
 
-#include "Async/TaskGraphInterfaces.h"
 #include "CoreMinimal.h"
+#include "InterchangeTaskSystem.h"
 #include "Stats/Stats.h"
 #include "UObject/WeakObjectPtrTemplates.h"
 #include "Nodes/InterchangeBaseNode.h"
@@ -17,7 +17,7 @@ namespace UE
 	{
 		class FImportAsyncHelper;
 
-		class FTaskCreateSceneObjects
+		class FTaskCreateSceneObjects_GameThread : public FInterchangeTaskBase
 		{
 		private:
 			FString PackageBasePath;
@@ -27,23 +27,14 @@ namespace UE
 			const UClass* FactoryClass;
 
 		public:
-			explicit FTaskCreateSceneObjects(const FString& InPackageBasePath, const int32 InSourceIndex, TWeakPtr<FImportAsyncHelper> InAsyncHelper, TArrayView<UInterchangeFactoryBaseNode*> InNodes, const UClass* InFactoryClass);
+			explicit FTaskCreateSceneObjects_GameThread(const FString& InPackageBasePath, const int32 InSourceIndex, TWeakPtr<FImportAsyncHelper> InAsyncHelper, TArrayView<UInterchangeFactoryBaseNode*> InNodes, const UClass* InFactoryClass);
 
-			ENamedThreads::Type GetDesiredThread();
-
-			static ESubsequentsMode::Type GetSubsequentsMode()
-			{
-				return ESubsequentsMode::TrackSubsequents;
-			}
-
-			TStatId GetStatId() const
+			virtual EInterchangeTaskThread GetTaskThread() const override
 			{
-				RETURN_QUICK_DECLARE_CYCLE_STAT(FTaskCreateSceneObjects, STATGROUP_TaskGraphTasks);
+				return EInterchangeTaskThread::GameThread;
 			}
 
-			void DoTask(ENamedThreads::Type CurrentThread, const FGraphEventRef& MyCompletionGraphEvent);
+			virtual void Execute() override;
 		};
-
-
 	} //ns Interchange
 }//ns UE
diff --git a/Engine/Source/Runtime/Interchange/Engine/Private/Tasks/InterchangeTaskImportObject.cpp b/Engine/Source/Runtime/Interchange/Engine/Private/Tasks/InterchangeTaskImportObject.cpp
--- a/Engine/Source/Runtime/Interchange/Engine/Private/Tasks/InterchangeTaskImportObject.cpp
+++ b/Engine/Source/Runtime/Interchange/Engine/Private/Tasks/InterchangeTaskImportObject.cpp
@@ -239,7 +239,7 @@ namespace UE::Interchange::Private
 	}
 }//ns UE::Interchange::Private
 
-void UE::Interchange::FTaskImportObject_GameThread::DoTask(ENamedThreads::Type CurrentThread, const FGraphEventRef& MyCompletionGraphEvent)
+void UE::Interchange::FTaskImportObject_GameThread::Execute()
 {
 	TRACE_CPUPROFILER_EVENT_SCOPE(UE::Interchange::FTaskImportObject_GameThread::DoTask)
 #if INTERCHANGE_TRACE_ASYNCHRONOUS_TASK_ENABLED
@@ -587,7 +587,7 @@ void UE::Interchange::FTaskImportObject_GameThread::DoTask(ENamedThreads::Type C
 	}
 }
 
-void UE::Interchange::FTaskImportObject_Async::DoTask(ENamedThreads::Type CurrentThread, const FGraphEventRef& MyCompletionGraphEvent)
+void UE::Interchange::FTaskImportObject_Async::Execute()
 {
 #if INTERCHANGE_TRACE_ASYNCHRONOUS_TASK_ENABLED
 	INTERCHANGE_TRACE_ASYNCHRONOUS_TASK(TaskImportObject_Async)
@@ -597,6 +597,12 @@ void UE::Interchange::FTaskImportObject_Async::DoTask(ENamedThreads::Type Curren
 
 	using namespace UE::Interchange;
 
+	TOptional<FGCScopeGuard> GCScopeGuard;
+	if (!IsInGameThread())
+	{
+		GCScopeGuard.Emplace();
+	}
+
 	TSharedPtr<FImportAsyncHelper, ESPMode::ThreadSafe> AsyncHelper = WeakAsyncHelper.Pin();
 	check(AsyncHelper.IsValid());
 
@@ -623,7 +629,7 @@ void UE::Interchange::FTaskImportObject_Async::DoTask(ENamedThreads::Type Curren
 		});
 }
 
-void UE::Interchange::FTaskImportObjectFinalize_GameThread::DoTask(ENamedThreads::Type CurrentThread, const FGraphEventRef& MyCompletionGraphEvent)
+void UE::Interchange::FTaskImportObjectFinalize_GameThread::Execute()
 {
 #if INTERCHANGE_TRACE_ASYNCHRONOUS_TASK_ENABLED
 	INTERCHANGE_TRACE_ASYNCHRONOUS_TASK(TaskImportObjectFinalize_GameThread)
@@ -648,6 +654,8 @@ void UE::Interchange::FTaskImportObjectFinalize_GameThread::DoTask(ENamedThreads
 		return;
 	}
 
+	check(IsInGameThread());
+
 	UInterchangeFactoryBase::FImportAssetResult ImportAssetResult = Private::InternalImportObjectStartup(AsyncHelper
 		, FactoryNode
 		, SourceIndex
diff --git a/Engine/Source/Runtime/Interchange/Engine/Private/Tasks/InterchangeTaskImportObject.h b/Engine/Source/Runtime/Interchange/Engine/Private/Tasks/InterchangeTaskImportObject.h
--- a/Engine/Source/Runtime/Interchange/Engine/Private/Tasks/InterchangeTaskImportObject.h
+++ b/Engine/Source/Runtime/Interchange/Engine/Private/Tasks/InterchangeTaskImportObject.h
@@ -2,10 +2,10 @@
 
 #pragma once
 
-#include "Async/TaskGraphInterfaces.h"
 #include "CoreMinimal.h"
 #include "InterchangeFactoryBase.h"
 #include "InterchangeManager.h"
+#include "InterchangeTaskSystem.h"
 #include "Stats/Stats.h"
 #include "UObject/WeakObjectPtrTemplates.h"
 #include "Nodes/InterchangeFactoryBaseNode.h"
@@ -15,7 +15,7 @@ namespace UE::Interchange
 	/**
 		* This task create UPackage and UObject, Cook::PackageTracker::NotifyUObjectCreated is not thread safe, so we need to create the packages on the main thread
 		*/
-	class FTaskImportObject_GameThread
+	class FTaskImportObject_GameThread : public FInterchangeTaskBase
 	{
 	private:
 		FString PackageBasePath;
@@ -36,30 +36,15 @@ namespace UE::Interchange
 			check(FactoryClass);
 		}
 
-		ENamedThreads::Type GetDesiredThread() const
+		virtual EInterchangeTaskThread GetTaskThread() const override
 		{
-			TSharedPtr<FImportAsyncHelper, ESPMode::ThreadSafe> AsyncHelper = WeakAsyncHelper.Pin();
-			if (AsyncHelper.IsValid() && AsyncHelper->bRunSynchronous)
-			{
-				return ENamedThreads::GameThread_Local;
-			}
-			return ENamedThreads::GameThread;
-		}
-
-		static ESubsequentsMode::Type GetSubsequentsMode()
-		{
-			return ESubsequentsMode::TrackSubsequents;
+			return EInterchangeTaskThread::GameThread;
 		}
 
-		TStatId GetStatId() const
-		{
-			RETURN_QUICK_DECLARE_CYCLE_STAT(FTaskImportObject_GameThread, STATGROUP_TaskGraphTasks);
-		}
-
-		void DoTask(ENamedThreads::Type CurrentThread, const FGraphEventRef& MyCompletionGraphEvent);
+		virtual void Execute() override;
 	};
 
-	class FTaskImportObject_Async
+	class FTaskImportObject_Async : public FInterchangeTaskBase
 	{
 	private:
 		FString PackageBasePath;
@@ -77,33 +62,24 @@ namespace UE::Interchange
 			check(FactoryNode);
 		}
 
-		ENamedThreads::Type GetDesiredThread() const
+		virtual EInterchangeTaskThread GetTaskThread() const override
 		{
 			TSharedPtr<FImportAsyncHelper, ESPMode::ThreadSafe> AsyncHelper = WeakAsyncHelper.Pin();
 			if (AsyncHelper.IsValid() && AsyncHelper->bRunSynchronous)
 			{
-				return ENamedThreads::GameThread_Local;
+				return EInterchangeTaskThread::GameThread;
 			}
-			return ENamedThreads::AnyBackgroundThreadNormalTask;
-		}
 
-		static ESubsequentsMode::Type GetSubsequentsMode()
-		{
-			return ESubsequentsMode::TrackSubsequents;
+			return EInterchangeTaskThread::AsyncThread;
 		}
 
-		TStatId GetStatId() const
-		{
-			RETURN_QUICK_DECLARE_CYCLE_STAT(FTaskImportObject_Async, STATGROUP_TaskGraphTasks);
-		}
-
-		void DoTask(ENamedThreads::Type CurrentThread, const FGraphEventRef& MyCompletionGraphEvent);
+		virtual void Execute() override;
 	};
 
 	/**
 		* This task create UPackage and UObject, Cook::PackageTracker::NotifyUObjectCreated is not thread safe, so we need to create the packages on the main thread
 		*/
-	class FTaskImportObjectFinalize_GameThread
+	class FTaskImportObjectFinalize_GameThread : public FInterchangeTaskBase
 	{
 	private:
 		FString PackageBasePath;
@@ -122,27 +98,12 @@ namespace UE::Interchange
 			check(FactoryNode);
 		}
 
-		ENamedThreads::Type GetDesiredThread() const
-		{
-			TSharedPtr<FImportAsyncHelper, ESPMode::ThreadSafe> AsyncHelper = WeakAsyncHelper.Pin();
-			if (AsyncHelper.IsValid() && AsyncHelper->bRunSynchronous)
-			{
-				return ENamedThreads::GameThread_Local;
-			}
-			return ENamedThreads::GameThread;
-		}
-
-		static ESubsequentsMode::Type GetSubsequentsMode()
-		{
-			return ESubsequentsMode::TrackSubsequents;
-		}
-
-		TStatId GetStatId() const
+		virtual EInterchangeTaskThread GetTaskThread() const override
 		{
-			RETURN_QUICK_DECLARE_CYCLE_STAT(FTaskImportObjectFinalize_GameThread, STATGROUP_TaskGraphTasks);
+			return EInterchangeTaskThread::GameThread;
 		}
 
-		void DoTask(ENamedThreads::Type CurrentThread, const FGraphEventRef& MyCompletionGraphEvent);
+		virtual void Execute() override;
 	};
 
 }//ns UE::Interchange
diff --git a/Engine/Source/Runtime/Interchange/Engine/Private/Tasks/InterchangeTaskParsing.cpp b/Engine/Source/Runtime/Interchange/Engine/Private/Tasks/InterchangeTaskParsing.cpp
--- a/Engine/Source/Runtime/Interchange/Engine/Private/Tasks/InterchangeTaskParsing.cpp
+++ b/Engine/Source/Runtime/Interchange/Engine/Private/Tasks/InterchangeTaskParsing.cpp
@@ -87,7 +87,7 @@ private:
 };
 
 
-void UE::Interchange::FTaskParsing::DoTask(ENamedThreads::Type CurrentThread, const FGraphEventRef& MyCompletionGraphEvent)
+void UE::Interchange::FTaskParsing::Execute()
 {
 	TRACE_CPUPROFILER_EVENT_SCOPE(UE::Interchange::FTaskParsing::DoTask)
 #if INTERCHANGE_TRACE_ASYNCHRONOUS_TASK_ENABLED
@@ -95,7 +95,11 @@ void UE::Interchange::FTaskParsing::DoTask(ENamedThreads::Type CurrentThread, co
 #endif
 	LLM_SCOPE_BYNAME(TEXT("Interchange"));
 
-	FGCScopeGuard GCScopeGuard;
+	TOptional<FGCScopeGuard> GCScopeGuard;
+	if (!IsInGameThread())
+	{
+		GCScopeGuard.Emplace();
+	}
 
 	TSharedPtr<FImportAsyncHelper, ESPMode::ThreadSafe> AsyncHelper = WeakAsyncHelper.Pin();
 	check(AsyncHelper.IsValid());
@@ -107,8 +111,8 @@ void UE::Interchange::FTaskParsing::DoTask(ENamedThreads::Type CurrentThread, co
 		int32 SourceIndex = INDEX_NONE;
 		bool bIsSceneNode = false;
 		TArray<FString> Dependencies;
-		FGraphEventRef GraphEventRef;
-		FGraphEventArray Prerequisites;
+		uint64 GraphEventRef;
+		TArray<uint64> Prerequisites;
 		const UClass* FactoryClass;
 
 		TArray<UInterchangeFactoryBaseNode*, TInlineAllocator<1>> Nodes; // For scenes, we can group multiple nodes into a single task as they are usually very light
@@ -196,9 +200,9 @@ void UE::Interchange::FTaskParsing::DoTask(ENamedThreads::Type CurrentThread, co
 		}
 	}
 
-	auto CreateTasksForEachTaskData = [](TArray<FTaskData>& TaskDatas, TFunction<FGraphEventRef(FTaskData&)> CreateTasksFunc) -> FGraphEventArray
+	auto CreateTasksForEachTaskData = [](TArray<FTaskData>& TaskDatas, TFunction<uint64(FTaskData&)> CreateTasksFunc) -> TArray<uint64>
 	{
-		FGraphEventArray GraphEvents;
+		TArray<uint64> GraphEvents;
 
 		for (int32 TaskIndex = 0; TaskIndex < TaskDatas.Num(); ++TaskIndex)
 		{
@@ -234,7 +238,7 @@ void UE::Interchange::FTaskParsing::DoTask(ENamedThreads::Type CurrentThread, co
 
 	TMap<FString, FTaskParsingRenameInfo> RenameAssets;
 	TSet<FString> CreatedTasksAssetNames; // Tracks for which asset name we have created a task so that we don't have 2 tasks for the same asset name
-	TFunction<FGraphEventRef(FTaskData&)> CreateTasksFromData = [this, &AsyncHelper, &RenameAssets, &CreatedTasksAssetNames](FTaskData& TaskData)
+	TFunction<uint64(FTaskData&)> CreateTasksFromData = [this, &AsyncHelper, &RenameAssets, &CreatedTasksAssetNames](FTaskData& TaskData)
 	{
 		LLM_SCOPE_BYNAME(TEXT("Interchange"));
 		check(TaskData.Nodes.Num() == 1); //We expect 1 node per asset task
@@ -245,9 +249,13 @@ void UE::Interchange::FTaskParsing::DoTask(ENamedThreads::Type CurrentThread, co
 
 		if (TaskData.bIsSceneNode)
 		{
-			return AsyncHelper->SceneTasks.Add_GetRef(
-				TGraphTask<FTaskCreateSceneObjects>::CreateTask(&(TaskData.Prerequisites))
-				.ConstructAndDispatchWhenReady(AsyncHelper->ContentBasePath, SourceIndex, WeakAsyncHelper, TaskData.Nodes, FactoryClass));
+			TSharedPtr<UE::Interchange::FTaskCreateSceneObjects_GameThread, ESPMode::ThreadSafe> TaskCreateSceneObjects = MakeShared<UE::Interchange::FTaskCreateSceneObjects_GameThread, ESPMode::ThreadSafe>(AsyncHelper->ContentBasePath
+				, SourceIndex
+				, WeakAsyncHelper
+				, TaskData.Nodes
+				, FactoryClass);
+
+			return AsyncHelper->SceneTasks.Add_GetRef(UE::Interchange::FInterchangeTaskSystem::Get().AddTask(TaskCreateSceneObjects, TaskData.Prerequisites));
 		}
 		else
 		{
@@ -290,22 +298,31 @@ void UE::Interchange::FTaskParsing::DoTask(ENamedThreads::Type CurrentThread, co
 				TEXT("Found multiple task data with the same asset name (%s). Only one will be executed."), *AssetFullPath))
 			{
 				LLM_SCOPE_BYNAME(TEXT("Interchange"));
-				FGraphEventArray ImportObjectTasksPrerequistes;
-				int32 BeginImportObjectTaskIndex = AsyncHelper->BeginImportObjectTasks.Add(
-					TGraphTask<FTaskImportObject_GameThread>::CreateTask(&(TaskData.Prerequisites)).ConstructAndDispatchWhenReady(AsyncHelper->ContentBasePath, SourceIndex, WeakAsyncHelper, FactoryNode, FactoryClass)
-				);
+				TArray<uint64> ImportObjectTasksPrerequistes;
+				TSharedPtr<FTaskImportObject_GameThread, ESPMode::ThreadSafe> TaskImportObject_GameThread = MakeShared<FTaskImportObject_GameThread, ESPMode::ThreadSafe>(AsyncHelper->ContentBasePath
+					, SourceIndex
+					, WeakAsyncHelper
+					, FactoryNode
+					, FactoryClass);
+				int32 BeginImportObjectTaskIndex = AsyncHelper->BeginImportObjectTasks.Add(FInterchangeTaskSystem::Get().AddTask(TaskImportObject_GameThread, TaskData.Prerequisites));
 				ImportObjectTasksPrerequistes.Add(AsyncHelper->BeginImportObjectTasks[BeginImportObjectTaskIndex]);
 
-				int32 ImportObjectTaskIndex = AsyncHelper->ImportObjectTasks.Add(
-					TGraphTask<FTaskImportObject_Async>::CreateTask(&(ImportObjectTasksPrerequistes)).ConstructAndDispatchWhenReady(AsyncHelper->ContentBasePath, SourceIndex, WeakAsyncHelper, FactoryNode)
-				);
-
-				FGraphEventArray FinalizeImportObjectTasksPrerequistes;
+				TSharedPtr<FTaskImportObject_Async, ESPMode::ThreadSafe> TaskImportObject_Async = MakeShared<FTaskImportObject_Async, ESPMode::ThreadSafe>(AsyncHelper->ContentBasePath
+					, SourceIndex
+					, WeakAsyncHelper
+					, FactoryNode);
+				int32 ImportObjectTaskIndex = AsyncHelper->ImportObjectTasks.Add(FInterchangeTaskSystem::Get().AddTask(TaskImportObject_Async, ImportObjectTasksPrerequistes));
+				
+				TArray<uint64> FinalizeImportObjectTasksPrerequistes;
 				FinalizeImportObjectTasksPrerequistes.Add(AsyncHelper->ImportObjectTasks[ImportObjectTaskIndex]);
 
-				int32 FinalizeCreateTaskIndex = AsyncHelper->FinalizeImportObjectTasks.Add(
-					TGraphTask<FTaskImportObjectFinalize_GameThread>::CreateTask(&(FinalizeImportObjectTasksPrerequistes)).ConstructAndDispatchWhenReady(AsyncHelper->ContentBasePath, SourceIndex, WeakAsyncHelper, FactoryNode)
-				);
+				
+				
+				TSharedPtr<FTaskImportObjectFinalize_GameThread, ESPMode::ThreadSafe> TaskImportObjectFinalize_GameThread = MakeShared<FTaskImportObjectFinalize_GameThread, ESPMode::ThreadSafe>(AsyncHelper->ContentBasePath
+					, SourceIndex
+					, WeakAsyncHelper
+					, FactoryNode);
+				int32 FinalizeCreateTaskIndex = AsyncHelper->FinalizeImportObjectTasks.Add(FInterchangeTaskSystem::Get().AddTask(TaskImportObjectFinalize_GameThread, FinalizeImportObjectTasksPrerequistes));
 
 				//Only add the name if the factory node is enabled
 				if (FactoryNode->IsEnabled())
@@ -317,8 +334,7 @@ void UE::Interchange::FTaskParsing::DoTask(ENamedThreads::Type CurrentThread, co
 			}
 			else
 			{
-				FGraphEventRef EmptyGraphEvent = FGraphEvent::CreateGraphEvent();
-				EmptyGraphEvent->DispatchSubsequents();
+				constexpr uint64 EmptyGraphEvent = 0;
 				return EmptyGraphEvent;
 			}
 		}
@@ -341,10 +357,10 @@ void UE::Interchange::FTaskParsing::DoTask(ENamedThreads::Type CurrentThread, co
 		}
 	}
 
-	FGraphEventArray CompletionPrerequistes;
+	TArray<uint64> CompletionPrerequistes;
 	const int32 PoolWorkerThreadCount = FTaskGraphInterface::Get().GetNumWorkerThreads() / 2;
 	const int32 MaxNumWorker = FMath::Max(PoolWorkerThreadCount, 1);
-	FGraphEventArray GroupPrerequistes;
+	TArray<uint64> GroupPrerequistes;
 	for (int32 TaskIndex = 0; TaskIndex < TaskDatas.Num(); ++TaskIndex)
 	{
 		FTaskData& TaskData = TaskDatas[TaskIndex];
@@ -408,28 +424,32 @@ void UE::Interchange::FTaskParsing::DoTask(ENamedThreads::Type CurrentThread, co
 	}
 
 	//Add an async task for pre completion
-	FGraphEventArray PreCompletionPrerequistes;
-	AsyncHelper->PreCompletionTask = TGraphTask<FTaskPreCompletion>::CreateTask(&CompletionPrerequistes).ConstructAndDispatchWhenReady(InterchangeManager, WeakAsyncHelper);
+	TArray<uint64> PreCompletionPrerequistes;
+	TSharedPtr<FTaskPreCompletion_GameThread, ESPMode::ThreadSafe> TaskPreCompletion = MakeShared<FTaskPreCompletion_GameThread, ESPMode::ThreadSafe>(InterchangeManager, WeakAsyncHelper);
+	AsyncHelper->PreCompletionTask = FInterchangeTaskSystem::Get().AddTask(TaskPreCompletion, CompletionPrerequistes);
 	PreCompletionPrerequistes.Add(AsyncHelper->PreCompletionTask);
 
-	FGraphEventArray AssetCompilationPrerequistes;
+	TArray<uint64> AssetCompilationPrerequistes;
 	//Start the wait for asset compilation task on a async task
 	for (int32 SourceIndex = 0; SourceIndex < AsyncHelper->SourceDatas.Num(); ++SourceIndex)
 	{
-		AsyncHelper->WaitAssetCompilationTask = TGraphTask<FTaskWaitAssetCompilation>::CreateTask(&PreCompletionPrerequistes).ConstructAndDispatchWhenReady(SourceIndex, WeakAsyncHelper);
-		AssetCompilationPrerequistes.Add(AsyncHelper->WaitAssetCompilationTask);
+		TSharedPtr<FTaskWaitAssetCompilation_GameThread, ESPMode::ThreadSafe> TaskWaitAssetCompilation = MakeShared<FTaskWaitAssetCompilation_GameThread, ESPMode::ThreadSafe>(SourceIndex, WeakAsyncHelper);
+		int32 WaitAssetCompilationTaskIndex = AsyncHelper->WaitAssetCompilationTasks.Add(FInterchangeTaskSystem::Get().AddTask(TaskWaitAssetCompilation, PreCompletionPrerequistes));
+		AssetCompilationPrerequistes.Add(AsyncHelper->WaitAssetCompilationTasks[WaitAssetCompilationTaskIndex]);
 	}
 
 	//Start the Post pipeline task
 	for (int32 SourceIndex = 0; SourceIndex < AsyncHelper->SourceDatas.Num(); ++SourceIndex)
 	{
-		int32 GraphTaskIndex = AsyncHelper->PostImportTasks.Add(
-			TGraphTask<FTaskPostImport>::CreateTask(&(AssetCompilationPrerequistes)).ConstructAndDispatchWhenReady(SourceIndex, WeakAsyncHelper)
-		);
+
+		TSharedPtr<FTaskPostImport_GameThread, ESPMode::ThreadSafe> TaskPostImport = MakeShared<FTaskPostImport_GameThread, ESPMode::ThreadSafe>(SourceIndex, WeakAsyncHelper);
+		int32 PostImportTaskIndex = AsyncHelper->PostImportTasks.Add(FInterchangeTaskSystem::Get().AddTask(TaskPostImport, AssetCompilationPrerequistes));
+
 		//Ensure we run the pipeline in the same order we create the task, since the pipeline modifies the node container, its important that its not processed in parallel, Adding the one we start to the prerequisites
 		//is the way to go here
-		AssetCompilationPrerequistes.Add(AsyncHelper->PostImportTasks[GraphTaskIndex]);
+		AssetCompilationPrerequistes.Add(AsyncHelper->PostImportTasks[PostImportTaskIndex]);
 	}
 
-	AsyncHelper->CompletionTask = TGraphTask<FTaskCompletion>::CreateTask(&AssetCompilationPrerequistes).ConstructAndDispatchWhenReady(InterchangeManager, WeakAsyncHelper);
+	TSharedPtr<FTaskCompletion_GameThread, ESPMode::ThreadSafe> TaskCompletion = MakeShared<FTaskCompletion_GameThread, ESPMode::ThreadSafe>(InterchangeManager, WeakAsyncHelper);
+	AsyncHelper->CompletionTask = FInterchangeTaskSystem::Get().AddTask(TaskCompletion, AssetCompilationPrerequistes);
 }
diff --git a/Engine/Source/Runtime/Interchange/Engine/Private/Tasks/InterchangeTaskParsing.h b/Engine/Source/Runtime/Interchange/Engine/Private/Tasks/InterchangeTaskParsing.h
--- a/Engine/Source/Runtime/Interchange/Engine/Private/Tasks/InterchangeTaskParsing.h
+++ b/Engine/Source/Runtime/Interchange/Engine/Private/Tasks/InterchangeTaskParsing.h
@@ -2,9 +2,9 @@
 
 #pragma once
 
-#include "Async/TaskGraphInterfaces.h"
 #include "CoreMinimal.h"
 #include "InterchangeManager.h"
+#include "InterchangeTaskSystem.h"
 #include "Stats/Stats.h"
 #include "UObject/WeakObjectPtrTemplates.h"
 
@@ -13,7 +13,7 @@ namespace UE
 	namespace Interchange
 	{
 
-		class FTaskParsing
+		class FTaskParsing : public FInterchangeTaskBase
 		{
 		private:
 			UInterchangeManager* InterchangeManager;
@@ -26,27 +26,18 @@ namespace UE
 				check(InterchangeManager);
 			}
 
-			FORCEINLINE ENamedThreads::Type GetDesiredThread()
+			virtual EInterchangeTaskThread GetTaskThread() const override
 			{
 				TSharedPtr<FImportAsyncHelper, ESPMode::ThreadSafe> AsyncHelper = WeakAsyncHelper.Pin();
 				if (AsyncHelper.IsValid() && AsyncHelper->bRunSynchronous)
 				{
-					return ENamedThreads::GameThread_Local;
+					return EInterchangeTaskThread::GameThread;
 				}
-				return ENamedThreads::AnyBackgroundThreadNormalTask;
-			}
-			static FORCEINLINE ESubsequentsMode::Type GetSubsequentsMode()
-			{
 
-				return ESubsequentsMode::TrackSubsequents;
-			}
-
-			FORCEINLINE TStatId GetStatId() const
-			{
-				RETURN_QUICK_DECLARE_CYCLE_STAT(FTaskParsing, STATGROUP_TaskGraphTasks);
+				return EInterchangeTaskThread::AsyncThread;
 			}
 
-			void DoTask(ENamedThreads::Type CurrentThread, const FGraphEventRef& MyCompletionGraphEvent);
+			virtual void Execute() override;
 		};
 
 	} //ns Interchange
diff --git a/Engine/Source/Runtime/Interchange/Engine/Private/Tasks/InterchangeTaskPipeline.cpp b/Engine/Source/Runtime/Interchange/Engine/Private/Tasks/InterchangeTaskPipeline.cpp
--- a/Engine/Source/Runtime/Interchange/Engine/Private/Tasks/InterchangeTaskPipeline.cpp
+++ b/Engine/Source/Runtime/Interchange/Engine/Private/Tasks/InterchangeTaskPipeline.cpp
@@ -3,7 +3,6 @@
 
 #include "AssetCompilingManager.h"
 #include "Async/Async.h"
-#include "Async/TaskGraphInterfaces.h"
 #include "CoreMinimal.h"
 #include "GenericPlatform/GenericPlatformProcess.h"
 #include "InterchangeEngineLogPrivate.h"
@@ -20,8 +19,7 @@
 #include "UObject/WeakObjectPtrTemplates.h"
 
 
-
-void UE::Interchange::FTaskPipeline::DoTask(ENamedThreads::Type CurrentThread, const FGraphEventRef& MyCompletionGraphEvent)
+void UE::Interchange::FTaskPipeline::Execute()
 {
 	TRACE_CPUPROFILER_EVENT_SCOPE(UE::Interchange::FTaskPipeline::DoTask)
 #if INTERCHANGE_TRACE_ASYNCHRONOUS_TASK_ENABLED
@@ -59,17 +57,16 @@ void UE::Interchange::FTaskPipeline::DoTask(ENamedThreads::Type CurrentThread, c
 	}
 }
 
-void UE::Interchange::FTaskWaitAssetCompilation::DoTask(ENamedThreads::Type CurrentThread, const FGraphEventRef& MyCompletionGraphEvent)
+void UE::Interchange::FTaskWaitAssetCompilation_GameThread::Execute()
 {
-	TRACE_CPUPROFILER_EVENT_SCOPE(UE::Interchange::FTaskWaitAssetCompilation::DoTask)
+	TRACE_CPUPROFILER_EVENT_SCOPE(UE::Interchange::FTaskWaitAssetCompilation_GameThread::DoTask)
 #if INTERCHANGE_TRACE_ASYNCHRONOUS_TASK_ENABLED
-		INTERCHANGE_TRACE_ASYNCHRONOUS_TASK(WaitAssetCompilation)
+	INTERCHANGE_TRACE_ASYNCHRONOUS_TASK(WaitAssetCompilation)
 #endif
-
 	LLM_SCOPE_BYNAME(TEXT("Interchange"));
 
 #if WITH_EDITOR
-
+	check(IsInGameThread());
 	TSharedPtr<FImportAsyncHelper, ESPMode::ThreadSafe> AsyncHelper = WeakAsyncHelper.Pin();
 	if (!ensure(AsyncHelper.IsValid()) || AsyncHelper->bCancel)
 	{
@@ -95,74 +92,45 @@ void UE::Interchange::FTaskWaitAssetCompilation::DoTask(ENamedThreads::Type Curr
 	AsyncHelper->IterateImportedAssets(SourceIndex, FillImportedObjectsFromSource);
 	AsyncHelper->IterateImportedSceneObjects(SourceIndex, FillImportedObjectsFromSource);
 
-	//Make sure all assets compilation are done before calling the pipeline post import task, let other thread execute if assets are not compile yet and wait 50ms before a new query
-	FPlatformProcess::ConditionalSleep([&ImportedObjects]()
+	bool bCompilationFinish = true;
+	for (int32 ObjectIndex = 0; ObjectIndex < ImportedObjects.Num(); ++ObjectIndex)
+	{
+		UObject* ImportObject = ImportedObjects[ObjectIndex];
+		if (UMaterialInterface* MaterialInterface = Cast<UMaterialInterface>(ImportObject))
 		{
-			LLM_SCOPE_BYNAME(TEXT("Interchange"));
-			//Compilation status cannot be ask in async thread, query the compile status on the main thread with a small fast function
-			//This ensure we don't stall the main thread until all assets are compile.
-			bool bCompilationFinish = false;
-
-			auto VerifyAssetCompiling = [&bCompilationFinish, &ImportedObjects]()
-				{
-					TRACE_CPUPROFILER_EVENT_SCOPE(UE::Interchange::FTaskWaitAssetCompilation::DoTask::IsCompilingLambda_GameThread);
-					LLM_SCOPE_BYNAME(TEXT("Interchange"));
-					//Make sure all asset compiling managers are up to date, In case the game thread is waiting for the import to finish (like automation test or synchronous import)
-					FAssetCompilingManager::Get().ProcessAsyncTasks();
-
-					bCompilationFinish = true;
-					for (int32 ObjectIndex = 0; ObjectIndex < ImportedObjects.Num(); ++ObjectIndex)
-					{
-						UObject* ImportObject = ImportedObjects[ObjectIndex];
-						if (UMaterialInterface* MaterialInterface = Cast<UMaterialInterface>(ImportObject))
-						{
-							if (MaterialInterface->IsCompiling())
-							{
-								bCompilationFinish = false;
-								break;
-							}
-						}
-						if (IInterface_AsyncCompilation* AssetCompilationInterface = Cast<IInterface_AsyncCompilation>(ImportObject))
-						{
-							if (AssetCompilationInterface->IsCompiling())
-							{
-								bCompilationFinish = false;
-								break;
-							}
-						}
-					}
-				};
-
-			if (IsInGameThread())
+			if (MaterialInterface->IsCompiling())
 			{
-				VerifyAssetCompiling();
+				bCompilationFinish = false;
+				break;
 			}
-			else
+		}
+		if (IInterface_AsyncCompilation* AssetCompilationInterface = Cast<IInterface_AsyncCompilation>(ImportObject))
+		{
+			if (AssetCompilationInterface->IsCompiling())
 			{
-				Async(EAsyncExecution::TaskGraphMainThread, [&VerifyAssetCompiling]()
-					{
-						VerifyAssetCompiling();
-					}).Wait();
+				bCompilationFinish = false;
+				break;
 			}
-			return bCompilationFinish;
-		}, 0.05f);
+		}
+	}
+
+	if (!bCompilationFinish)
+	{
+		//re-enqueue this task so it doesn't execute the subsequent tasks
+		SetTaskStatus(EInterchangeTaskStatus::Waiting);
+	}
 #endif //WITH_EDITOR
 }
 
-void UE::Interchange::FTaskPostImport::DoTask(ENamedThreads::Type CurrentThread, const FGraphEventRef& MyCompletionGraphEvent)
+void UE::Interchange::FTaskPostImport_GameThread::Execute()
 {
-	TRACE_CPUPROFILER_EVENT_SCOPE(UE::Interchange::FTaskPostImport::DoTask)
+	TRACE_CPUPROFILER_EVENT_SCOPE(UE::Interchange::FTaskPostImport_GameThread::DoTask)
 #if INTERCHANGE_TRACE_ASYNCHRONOUS_TASK_ENABLED
 	INTERCHANGE_TRACE_ASYNCHRONOUS_TASK(PipelinePostImport)
 #endif
 
 	LLM_SCOPE_BYNAME(TEXT("Interchange"));
-
-	TOptional<FGCScopeGuard> GCScopeGuard;
-	if (!IsInGameThread())
-	{
-		GCScopeGuard.Emplace();
-	}
+	check(IsInGameThread());
 
 	TSharedPtr<FImportAsyncHelper, ESPMode::ThreadSafe> AsyncHelper = WeakAsyncHelper.Pin();
 	if (!ensure(AsyncHelper.IsValid()) || AsyncHelper->bCancel)
diff --git a/Engine/Source/Runtime/Interchange/Engine/Private/Tasks/InterchangeTaskPipeline.h b/Engine/Source/Runtime/Interchange/Engine/Private/Tasks/InterchangeTaskPipeline.h
--- a/Engine/Source/Runtime/Interchange/Engine/Private/Tasks/InterchangeTaskPipeline.h
+++ b/Engine/Source/Runtime/Interchange/Engine/Private/Tasks/InterchangeTaskPipeline.h
@@ -2,10 +2,10 @@
 
 #pragma once
 
-#include "Async/TaskGraphInterfaces.h"
 #include "CoreMinimal.h"
 #include "InterchangeManager.h"
 #include "InterchangePipelineBase.h"
+#include "InterchangeTaskSystem.h"
 #include "Stats/Stats.h"
 #include "UObject/WeakObjectPtrTemplates.h"
 
@@ -14,7 +14,7 @@ namespace UE
 	namespace Interchange
 	{
 
-		class FTaskPipeline
+		class FTaskPipeline : public FInterchangeTaskBase
 		{
 		private:
 			TWeakObjectPtr<UInterchangePipelineBase> PipelineBase;
@@ -26,114 +26,73 @@ namespace UE
 			{
 			}
 
-			ENamedThreads::Type GetDesiredThread()
+			virtual EInterchangeTaskThread GetTaskThread() const override
 			{
 				TSharedPtr<FImportAsyncHelper, ESPMode::ThreadSafe> AsyncHelper = WeakAsyncHelper.Pin();
 				if (AsyncHelper.IsValid() && AsyncHelper->bRunSynchronous)
 				{
-					return ENamedThreads::GameThread_Local;
+					return EInterchangeTaskThread::GameThread;
 				}
 
 				if (!ensure(PipelineBase.IsValid()))
 				{
-					return ENamedThreads::GameThread;
+					return EInterchangeTaskThread::GameThread;
 				}
-				
+
 				//Scripted (python) cannot run outside of the game thread, it will lock forever if we do this
 				if (PipelineBase.Get()->IsScripted())
 				{
-					return ENamedThreads::GameThread;
+					return EInterchangeTaskThread::GameThread;
 				}
 
-				return PipelineBase.Get()->CanExecuteOnAnyThread(EInterchangePipelineTask::PostTranslator) ? ENamedThreads::AnyBackgroundThreadNormalTask : ENamedThreads::GameThread;
-			}
-
-			static ESubsequentsMode::Type GetSubsequentsMode()
-			{
-				return ESubsequentsMode::TrackSubsequents;
-			}
-
-			FORCEINLINE TStatId GetStatId() const
-			{
-				RETURN_QUICK_DECLARE_CYCLE_STAT(FTaskPipeline, STATGROUP_TaskGraphTasks);
+				return PipelineBase.Get()->CanExecuteOnAnyThread(EInterchangePipelineTask::PostTranslator) ? EInterchangeTaskThread::AsyncThread : EInterchangeTaskThread::GameThread;
 			}
 
-			void DoTask(ENamedThreads::Type CurrentThread, const FGraphEventRef& MyCompletionGraphEvent);
+			virtual void Execute() override;
 		};
 
-		//We want to be sure any asset compilation is finish before calling FTaskPostImport, we use a async task to wait until they are done
-		class FTaskWaitAssetCompilation
+		//We want to be sure any asset compilation is finish before calling subsequents tasks
+		class FTaskWaitAssetCompilation_GameThread : public FInterchangeTaskBase
 		{
 		private:
 			int32 SourceIndex;
 			TWeakPtr<FImportAsyncHelper, ESPMode::ThreadSafe> WeakAsyncHelper;
 
 		public:
-			FTaskWaitAssetCompilation(int32 InSourceIndex, TWeakPtr<FImportAsyncHelper, ESPMode::ThreadSafe> InAsyncHelper)
+			FTaskWaitAssetCompilation_GameThread(int32 InSourceIndex, TWeakPtr<FImportAsyncHelper, ESPMode::ThreadSafe> InAsyncHelper)
 				: SourceIndex(InSourceIndex)
 				, WeakAsyncHelper(InAsyncHelper)
 			{
 			}
 
-			ENamedThreads::Type GetDesiredThread()
-			{
-				TSharedPtr<FImportAsyncHelper, ESPMode::ThreadSafe> AsyncHelper = WeakAsyncHelper.Pin();
-				if (AsyncHelper.IsValid() && AsyncHelper->bRunSynchronous)
-				{
-					return ENamedThreads::GameThread_Local;
-				}
-
-				return ENamedThreads::AnyBackgroundThreadNormalTask;
-			}
-
-			static ESubsequentsMode::Type GetSubsequentsMode()
+			virtual EInterchangeTaskThread GetTaskThread() const override
 			{
-				return ESubsequentsMode::TrackSubsequents;
+				//This task is re-enqueue and don't stall the game thread
+				return EInterchangeTaskThread::GameThread;
 			}
 
-			FORCEINLINE TStatId GetStatId() const
-			{
-				RETURN_QUICK_DECLARE_CYCLE_STAT(FTaskWaitAssetCompilation, STATGROUP_TaskGraphTasks);
-			}
-
-			void DoTask(ENamedThreads::Type CurrentThread, const FGraphEventRef& MyCompletionGraphEvent);
+			virtual void Execute() override;
 		};
 
-		class FTaskPostImport
+		class FTaskPostImport_GameThread : public FInterchangeTaskBase
 		{
 		private:
 			int32 SourceIndex;
 			TWeakPtr<FImportAsyncHelper, ESPMode::ThreadSafe> WeakAsyncHelper;
 
 		public:
-			FTaskPostImport(int32 InSourceIndex, TWeakPtr<FImportAsyncHelper, ESPMode::ThreadSafe> InAsyncHelper)
+			FTaskPostImport_GameThread(int32 InSourceIndex, TWeakPtr<FImportAsyncHelper, ESPMode::ThreadSafe> InAsyncHelper)
 				: SourceIndex(InSourceIndex)
 				, WeakAsyncHelper(InAsyncHelper)
 			{
 			}
 
-			ENamedThreads::Type GetDesiredThread()
-			{
-				TSharedPtr<FImportAsyncHelper, ESPMode::ThreadSafe> AsyncHelper = WeakAsyncHelper.Pin();
-				if (AsyncHelper.IsValid() && AsyncHelper->bRunSynchronous)
-				{
-					return ENamedThreads::GameThread_Local;
-				}
-
-				return ENamedThreads::GameThread;
-			}
-
-			static ESubsequentsMode::Type GetSubsequentsMode()
-			{
-				return ESubsequentsMode::TrackSubsequents;
-			}
-
-			FORCEINLINE TStatId GetStatId() const
+			virtual EInterchangeTaskThread GetTaskThread() const override
 			{
-				RETURN_QUICK_DECLARE_CYCLE_STAT(FTaskPostImport, STATGROUP_TaskGraphTasks);
+				return EInterchangeTaskThread::GameThread;
 			}
 
-			void DoTask(ENamedThreads::Type CurrentThread, const FGraphEventRef& MyCompletionGraphEvent);
+			virtual void Execute() override;
 		};
 	} //ns Interchange
 }//ns UE
diff --git a/Engine/Source/Runtime/Interchange/Engine/Private/Tasks/InterchangeTaskTranslator.cpp b/Engine/Source/Runtime/Interchange/Engine/Private/Tasks/InterchangeTaskTranslator.cpp
--- a/Engine/Source/Runtime/Interchange/Engine/Private/Tasks/InterchangeTaskTranslator.cpp
+++ b/Engine/Source/Runtime/Interchange/Engine/Private/Tasks/InterchangeTaskTranslator.cpp
@@ -1,11 +1,11 @@
 // Copyright Epic Games, Inc. All Rights Reserved.
 #include "InterchangeTaskTranslator.h"
 
-#include "Async/TaskGraphInterfaces.h"
 #include "CoreMinimal.h"
 #include "InterchangeEngineLogPrivate.h"
 #include "InterchangeManager.h"
 #include "InterchangeSourceData.h"
+#include "InterchangeTaskSystem.h"
 #include "InterchangeTranslatorBase.h"
 #include "Nodes/InterchangeBaseNodeContainer.h"
 #include "Stats/Stats.h"
@@ -13,13 +13,19 @@
 #include "UObject/GarbageCollection.h"
 #include "UObject/WeakObjectPtrTemplates.h"
 
-void UE::Interchange::FTaskTranslator::DoTask(ENamedThreads::Type CurrentThread, const FGraphEventRef& MyCompletionGraphEvent)
+void UE::Interchange::FTaskTranslator::Execute()
 {
 	TRACE_CPUPROFILER_EVENT_SCOPE(UE::Interchange::FTaskTranslator::DoTask)
 #if INTERCHANGE_TRACE_ASYNCHRONOUS_TASK_ENABLED
-	INTERCHANGE_TRACE_ASYNCHRONOUS_TASK(Translator)
+	INTERCHANGE_TRACE_ASYNCHRONOUS_TASK(TaskTranslator)
 #endif
-	FGCScopeGuard GCScopeGuard;
+
+	
+	TOptional<FGCScopeGuard> GCScopeGuard;
+	if (!IsInGameThread())
+	{
+		GCScopeGuard.Emplace();
+	}
 
 	TSharedPtr<FImportAsyncHelper, ESPMode::ThreadSafe> AsyncHelper = WeakAsyncHelper.Pin();
 	check(AsyncHelper.IsValid());
diff --git a/Engine/Source/Runtime/Interchange/Engine/Private/Tasks/InterchangeTaskTranslator.h b/Engine/Source/Runtime/Interchange/Engine/Private/Tasks/InterchangeTaskTranslator.h
--- a/Engine/Source/Runtime/Interchange/Engine/Private/Tasks/InterchangeTaskTranslator.h
+++ b/Engine/Source/Runtime/Interchange/Engine/Private/Tasks/InterchangeTaskTranslator.h
@@ -2,9 +2,9 @@
 
 #pragma once
 
-#include "Async/TaskGraphInterfaces.h"
 #include "CoreMinimal.h"
 #include "InterchangeManager.h"
+#include "InterchangeTaskSystem.h"
 #include "Stats/Stats.h"
 #include "UObject/WeakObjectPtrTemplates.h"
 
@@ -12,8 +12,7 @@ namespace UE
 {
 	namespace Interchange
 	{
-
-		class FTaskTranslator
+		class FTaskTranslator : public FInterchangeTaskBase
 		{
 		private:
 			int32 SourceIndex = INDEX_NONE;
@@ -26,29 +25,18 @@ namespace UE
 			{
 			}
 
-			FORCEINLINE ENamedThreads::Type GetDesiredThread()
+			virtual EInterchangeTaskThread GetTaskThread() const override
 			{
 				TSharedPtr<FImportAsyncHelper, ESPMode::ThreadSafe> AsyncHelper = WeakAsyncHelper.Pin();
 				if (AsyncHelper.IsValid() && AsyncHelper->bRunSynchronous)
 				{
-					return ENamedThreads::GameThread_Local;
+					return EInterchangeTaskThread::GameThread;
 				}
 
-				return ENamedThreads::AnyBackgroundThreadNormalTask;
-			}
-			static FORCEINLINE ESubsequentsMode::Type GetSubsequentsMode()
-			{
-				return ESubsequentsMode::TrackSubsequents;
-			}
-
-			FORCEINLINE TStatId GetStatId() const
-			{
-				RETURN_QUICK_DECLARE_CYCLE_STAT(FTaskTranslator, STATGROUP_TaskGraphTasks);
+				return EInterchangeTaskThread::AsyncThread;
 			}
 
-			void DoTask(ENamedThreads::Type CurrentThread, const FGraphEventRef& MyCompletionGraphEvent);
+			virtual void Execute();
 		};
-
-
 	} //ns Interchange
 } //ns UE
diff --git a/Engine/Source/Runtime/Interchange/Engine/Public/InterchangeManager.h b/Engine/Source/Runtime/Interchange/Engine/Public/InterchangeManager.h
--- a/Engine/Source/Runtime/Interchange/Engine/Public/InterchangeManager.h
+++ b/Engine/Source/Runtime/Interchange/Engine/Public/InterchangeManager.h
@@ -17,6 +17,7 @@
 #include "InterchangePipelineConfigurationBase.h"
 #include "InterchangeResultsContainer.h"
 #include "InterchangeSourceData.h"
+#include "InterchangeTaskSystem.h"
 #include "InterchangeTranslatorBase.h"
 #include "InterchangeWriterBase.h"
 #include "Nodes/InterchangeBaseNodeContainer.h"
@@ -172,6 +173,9 @@ namespace UE
 			// Callback when the status switches to done.
 			INTERCHANGEENGINE_API void OnDone(TFunction< void(FImportResult&) > Callback);
 
+			//Set the async helper that own this import result
+			INTERCHANGEENGINE_API void SetAsyncHelper(TWeakPtr<class FImportAsyncHelper> InAsyncHelper);
+
 			// Internal delegates. To set these, use the FImportAssetParameters when calling the Interchange import functions.
 			FOnObjectImportDoneDynamic OnObjectDone;
 			FOnObjectImportDoneNative OnObjectDoneNative;
@@ -194,7 +198,7 @@ namespace UE
 			mutable FRWLock ImportedObjectsRWLock;
 			TObjectPtr<UInterchangeResultsContainer> Results;
 
-			FGraphEventRef GraphEvent; // WaitUntilDone waits for this event to be triggered.
+			TWeakPtr<class FImportAsyncHelper> AsyncHelper = nullptr;
 
 			TFunction< void(FImportResult&) > DoneCallback;
 		};
@@ -237,18 +241,18 @@ namespace UE
 			//Python class instanced assets cannot be saved, so we have to serialize in JSON the data to restore it when we do a reimport.
 			TArray<UObject*> OriginalPipelines;
 
-			TArray<FGraphEventRef> TranslatorTasks;
-			TArray<FGraphEventRef> PipelineTasks;
-			FGraphEventRef WaitAssetCompilationTask;
-			TArray<FGraphEventRef> PostImportTasks;
-			FGraphEventRef ParsingTask;
-			TArray<FGraphEventRef> BeginImportObjectTasks;
-			TArray<FGraphEventRef> ImportObjectTasks;
-			TArray<FGraphEventRef> FinalizeImportObjectTasks;
-			TArray<FGraphEventRef> SceneTasks;
+			TArray<uint64> TranslatorTasks;
+			TArray<uint64> PipelineTasks;
+			TArray<uint64> WaitAssetCompilationTasks;
+			TArray<uint64> PostImportTasks;
+			uint64 ParsingTask;
+			TArray<uint64> BeginImportObjectTasks;
+			TArray<uint64> ImportObjectTasks;
+			TArray<uint64> FinalizeImportObjectTasks;
+			TArray<uint64> SceneTasks;
 
-			FGraphEventRef PreCompletionTask;
-			FGraphEventRef CompletionTask;
+			uint64 PreCompletionTask;
+			uint64 CompletionTask;
 
 			//Return true if we can import this class, or false otherwise.
 			bool IsClassImportAllowed(UClass* Class);
@@ -297,7 +301,7 @@ namespace UE
 			/**
 			 * Wait synchronously after the graph parsing task is done, and return the GraphEventArray up to the completion TaskGraphEvent.
 			 */
-			FGraphEventArray GetCompletionTaskGraphEvent();
+			TArray<uint64> GetCompletionTaskGraphEvent();
 
 			void InitCancel();
 
diff --git a/Engine/Source/Runtime/Interchange/Engine/Public/InterchangeTaskSystem.h b/Engine/Source/Runtime/Interchange/Engine/Public/InterchangeTaskSystem.h
new file mode 100644
--- /dev/null
+++ b/Engine/Source/Runtime/Interchange/Engine/Public/InterchangeTaskSystem.h
@@ -0,0 +1,227 @@
+// Copyright Epic Games, Inc. All Rights Reserved.
+
+#pragma once
+
+#include "Async/Future.h"
+#include "Containers/Array.h"
+#include "Containers/Map.h"
+#include "Containers/Ticker.h"
+#include "HAL/CriticalSection.h"
+#include "Misc/ScopeLock.h"
+
+#define INTERCHANGE_INVALID_TASK_ID 0xFFFFFFFFFFFFFFFF
+
+namespace UE::Interchange
+{
+	//Forward declare to add it has friend
+	class FInterchangeTaskSystem;
+
+	enum class EInterchangeTaskStatus : uint8
+	{
+		Waiting, //Task is queue for execution when prerequisite are terminate and resource are available
+		Executing, //Task is being executing
+		Done //Task was execute and is now terminate
+	};
+
+	enum class EInterchangeTaskThread : uint8
+	{
+		GameThread, //Task will be tick in a ticker in the engine tick on the game thread. This is a safe place to manipulate UObject
+		AsyncThread //Task will be call on asynchronous thread, code running inside those task must be thread safe.
+	};
+
+	class FInterchangeTaskBase
+	{
+	public:
+		FInterchangeTaskBase()
+		{}
+
+		virtual ~FInterchangeTaskBase()
+		{
+			if (GetTaskStatus() == EInterchangeTaskStatus::Executing)
+			{
+				if (Future.IsValid())
+				{
+					Future.Get();
+				}
+			}
+		}
+
+		INTERCHANGEENGINE_API uint64 GetTaskId() const
+		{
+			return TaskId;
+		}
+
+		INTERCHANGEENGINE_API virtual EInterchangeTaskThread GetTaskThread() const
+		{
+			return EInterchangeTaskThread::GameThread;
+		}
+
+		INTERCHANGEENGINE_API const TArray<uint64>& GetPrerequisiteTasks() const
+		{
+			return PrerequisiteTasks;
+		}
+
+		INTERCHANGEENGINE_API EInterchangeTaskStatus GetTaskStatus() const
+		{
+			FScopeLock TaskStatusLock(&TaskStatusCriticalSection);
+			return TaskStatus;
+		}
+
+		/**
+		 * This function must return the task id to fill the promise of the execution.
+		 * 
+		 * If the caller want to re-enqueue the task and not run the subsequent tasks after
+		 * this function, It must set the task status to Waiting and the task will be
+		 * re-execute later on another tick. This is the best way to implement a async wait.
+		 */
+		INTERCHANGEENGINE_API virtual void Execute() = 0;
+
+		/**
+		 * Wait can be call only on game thread. Async thread wait can create dead lock by using all the thread resources.
+		 * 
+		 * The caller have some restriction to follow
+		 *    - This must be a safe place for any game thread task that will be executing during the wait.
+		 *        - Wait should be call from a engine tick and no lock (like the global UObject locks) should be taken before calling wait.
+		 *        - Slate ui callback are safe
+		 *        - Engine tick is safe, use FTSTicker::GetCoreTicker() when you can
+		 *    - Wait should not happen from a package/asset load call stack
+		 *    - Wait should not happen from a save package/asset call stack
+		 *    - Wait should not happen from a garbage collect call stack
+		 * 
+		 * If your system need to wait and you already have some locks or you are under a unsafe place you need to redesign your system to
+		 * move the wait before you get to a unsafe call stack place to wait.
+		 */
+		INTERCHANGEENGINE_API void Wait() const;
+	protected:
+
+		//Only friend FInterchangeTaskSystem can set a new status
+		bool SetTaskStatus(EInterchangeTaskStatus NewTaskStatus);
+
+		void SetPrerequisites(const TArray<uint64>& InPrerequisiteTasks);
+
+	private:
+
+		mutable FCriticalSection TaskStatusCriticalSection;
+		uint64 TaskId = INTERCHANGE_INVALID_TASK_ID;
+		EInterchangeTaskStatus TaskStatus = EInterchangeTaskStatus::Waiting;
+		TArray<uint64> PrerequisiteTasks;
+
+		//If a task is asynchronous the future will be set.
+		//The future will contain the task id when it will be ready.
+		TFuture<uint64> Future;
+
+		friend FInterchangeTaskSystem;
+	};
+
+	/**
+	 * This class is use to run a lambda with the interchange task system.
+	 */
+	class FInterchangeTaskLambda : public FInterchangeTaskBase
+	{
+	private:
+		EInterchangeTaskThread TaskThread = EInterchangeTaskThread::GameThread;
+		TFunction<void()> ExecuteLambda;
+		
+
+	public:
+		FInterchangeTaskLambda(EInterchangeTaskThread InTaskThread, TFunction<void()> InExecuteLambda)
+			: TaskThread(InTaskThread)
+			, ExecuteLambda(MoveTemp(InExecuteLambda))
+		{
+		}
+
+		virtual EInterchangeTaskThread GetTaskThread() const override
+		{
+			return TaskThread;
+		}
+
+		virtual void Execute()
+		{
+			ExecuteLambda();
+		}
+	};
+
+	class FInterchangeTaskSystem
+	{
+	public:
+		/**
+		 * The constructor should not be called directly, an ensure will be set if its the case
+		 * Use the static Get function to acces the singleton.
+		 */
+		FInterchangeTaskSystem();
+
+		//The interchange task system is a singleton
+		static INTERCHANGEENGINE_API FInterchangeTaskSystem& Get();
+
+		/**
+		 * Add a task with some prerequisites. Return the task id.
+		 * If there is an issue the returned task id will be INTERCHANGE_INVALID_TASK_ID
+		 * The task will start when all prerequisites are done and some thread resource will be available.
+		 */
+		INTERCHANGEENGINE_API uint64 AddTask(TSharedPtr<FInterchangeTaskBase, ESPMode::ThreadSafe> Task, TArray<uint64>& TaskPrerequisites);
+
+		/**
+		 * Add a task without prerequisite. Return the task id.
+		 * If there is an issue the returned task id will be INTERCHANGE_INVALID_TASK_ID
+		 * The task will start when some thread resource will be available.
+		 */
+		INTERCHANGEENGINE_API uint64 AddTask(TSharedPtr<FInterchangeTaskBase, ESPMode::ThreadSafe> Task);
+
+		/** Return the current task status */
+		INTERCHANGEENGINE_API EInterchangeTaskStatus GetTaskStatus(const uint64 TaskId) const;
+
+		/**
+		 * Cancel a task, you can control if you cancel also the prerequisites, and subsequents.
+		 * Canceling a task that is waiting will simply push it state to Done which will allow subsequent to be executed.
+		 */
+		INTERCHANGEENGINE_API void CancelTask(const uint64 TaskId, const bool bCancelPrerequisites);
+
+		/**
+		 * Wait can be call only on game thread. Async thread wait can create dead lock by using all the thread resources.
+		 *
+		 * The caller have some restriction to follow
+		 *    - This must be a safe place for any game thread task that will be executing during the wait.
+		 *        - Wait should be call from a engine tick and no lock (like the global UObject locks) should be taken before calling wait.
+		 *        - Slate ui callback are safe
+		 *        - Engine tick is safe, use FTSTicker::GetCoreTicker() when you can
+		 *    - Wait should not happen from a package/asset load call stack
+		 *    - Wait should not happen from a save package/asset call stack
+		 *    - Wait should not happen from a garbage collect call stack
+		 *
+		 * If your system need to wait and you already have some locks or you are under a unsafe place you need to redesign your system to
+		 * move the wait before you get to a unsafe call stack place to wait.
+		 */
+		INTERCHANGEENGINE_API void WaitUntilTasksComplete(const TArray<uint64>& TasksToComplete);
+
+	private:
+
+		static bool bIsCreatingSingleton;
+
+		//Internal cancel request with no lock, lock is handle by the caller (CancelTask)
+		void InternalAddCancelRequestNoLock(const uint64 TaskId, bool bCancelPrerequisites);
+
+		/**
+		 * The tick is where we start task execution and where we update the tasks status
+		 */
+		void Tick();
+
+		//The ticker we use to tick in the engine tick
+		FTSTicker::FDelegateHandle	TickTickerHandle;
+
+		// Store the tasks per id
+		TMap<uint64, TSharedPtr<FInterchangeTaskBase, ESPMode::ThreadSafe>> TaskPerIdMap;
+		mutable FCriticalSection TaskPerIdMapCriticalSection;
+
+		//When a task is done, we add it here and release the task. Prerequisite are search in the list and here.
+		//A clean up is done when the TaskPerIdMap is empty.
+		TArray<uint64> ReleaseAndDoneTasks;
+
+		//The priority value is a counter that is increment for each task added to the system to create the task id
+		//The lower task id are execute earlier if possible.
+		uint64 PriorityValue = 0;
+		FCriticalSection PriorityValueCriticalSection;
+
+		TArray<uint64> CancelTaskRequests;
+		FCriticalSection CancelTaskRequestsCriticalSection;
+	};
+} //namespace UE::Interchange
