[Backout] - CL34239002
Original CL Desc
-----------------------------------------------------------------
Interchange: refactor the task to not be tick anywhere on the game thread. We tick the interchange gamethread task in the engine tick to be sure we can do anything we want on the game thread.

#rnx

diff --git a/Engine/Source/Runtime/Engine/Private/SkeletalMesh.cpp b/Engine/Source/Runtime/Engine/Private/SkeletalMesh.cpp
--- a/Engine/Source/Runtime/Engine/Private/SkeletalMesh.cpp
+++ b/Engine/Source/Runtime/Engine/Private/SkeletalMesh.cpp
@@ -1400,15 +1400,9 @@ void USkeletalMesh::PreEditChange(FProperty* PropertyAboutToChange)
 {
 	TRACE_CPUPROFILER_EVENT_SCOPE(USkeletalMesh::PreEditChange);
 
-	//Don't call finish compile if this skeletal mesh is compiling and we are in a FSkinnedAssetAsyncBuildScope for this skeletal mesh.
-	//If on the game thread we call LockPropertyUntil and we call PreEditChange after, in such a case a deadlock will happen if
-	//we call finish compile on this skeletal mesh.
-	if (FSkinnedAssetAsyncBuildScope::ShouldWaitOnLockedProperties(this))
-	{
-		// Tell the compiler to finish compiling us if we have a pending
-		// compilation ongoing plus any dependency (i.e. UGroomBindings).
-		FAssetCompilingManager::Get().FinishCompilationForObjects({ this });
-	}
+	// Tell the compiler to finish compiling us if we have a pending
+	// compilation ongoing plus any dependency (i.e. UGroomBindings).
+	FAssetCompilingManager::Get().FinishCompilationForObjects({ this });
 
 	Super::PreEditChange(PropertyAboutToChange);
 }
diff --git a/Engine/Source/Runtime/Engine/Private/SkeletalMeshComponentPhysics.cpp b/Engine/Source/Runtime/Engine/Private/SkeletalMeshComponentPhysics.cpp
--- a/Engine/Source/Runtime/Engine/Private/SkeletalMeshComponentPhysics.cpp
+++ b/Engine/Source/Runtime/Engine/Private/SkeletalMeshComponentPhysics.cpp
@@ -2169,12 +2169,8 @@ void USkeletalMeshComponent::UpdateHasValidBodies()
 		// For each body in physics asset..
 		for( int32 BodyIndex = 0; BodyIndex < PhysicsAsset->SkeletalBodySetups.Num(); BodyIndex++ )
 		{
-			int32 BoneIndex = INDEX_NONE;
 			// .. find the matching graphics bone index
-			if (TObjectPtr<USkeletalBodySetup> SkeletalBodySetup = PhysicsAsset->SkeletalBodySetups[BodyIndex])
-			{
-				BoneIndex = GetBoneIndex(SkeletalBodySetup->BoneName);
-			}
+			int32 BoneIndex = GetBoneIndex( PhysicsAsset->SkeletalBodySetups[ BodyIndex ]->BoneName );
 
 			// If we found a valid graphics bone, set the 'valid' flag
 			if(BoneIndex != INDEX_NONE)
diff --git a/Engine/Source/Runtime/Interchange/Engine/Private/InterchangeManager.cpp b/Engine/Source/Runtime/Interchange/Engine/Private/InterchangeManager.cpp
--- a/Engine/Source/Runtime/Interchange/Engine/Private/InterchangeManager.cpp
+++ b/Engine/Source/Runtime/Interchange/Engine/Private/InterchangeManager.cpp
@@ -23,7 +23,6 @@
 #include "InterchangeProjectSettings.h"
 #include "InterchangePythonPipelineBase.h"
 #include "InterchangeSourceData.h"
-#include "InterchangeTaskSystem.h"
 #include "InterchangeTranslatorBase.h"
 #include "InterchangeWriterBase.h"
 #include "Internationalization/Internationalization.h"
@@ -538,56 +537,47 @@ void UE::Interchange::FImportAsyncHelper::SendAnalyticImportEndData()
 
 void UE::Interchange::FImportAsyncHelper::ReleaseTranslatorsSource()
 {
-	//Make sure Executing translator task are done before releasing the resource
-	const int32 TranslatorCount = Translators.Num();
-	for (int32 TranslatorIndex = 0; TranslatorIndex < TranslatorCount; ++TranslatorIndex)
+	for (UInterchangeTranslatorBase* BaseTranslator : Translators)
 	{
-		const uint64 TaskId = TranslatorTasks.IsValidIndex(TranslatorIndex) ? TranslatorTasks[TranslatorIndex] : INTERCHANGE_INVALID_TASK_ID;
-		UInterchangeTranslatorBase* BaseTranslator = Translators[TranslatorIndex];
-
 		if (BaseTranslator)
 		{
-			//Wait until the translator is done before releasing the resource. This should not happen since resource are free during completion
-			const EInterchangeTaskStatus TaskStatus = UE::Interchange::FInterchangeTaskSystem::Get().GetTaskStatus(TaskId);
-			if (TaskStatus == EInterchangeTaskStatus::Executing)
-			{
-				ensure(TaskStatus != EInterchangeTaskStatus::Executing);
-				UE::Interchange::FInterchangeTaskSystem::Get().WaitUntilTasksComplete({ TaskId });
-			}
 			BaseTranslator->ReleaseSource();
 		}
 	}
 }
 
-TArray<uint64> UE::Interchange::FImportAsyncHelper::GetCompletionTaskGraphEvent()
+FGraphEventArray UE::Interchange::FImportAsyncHelper::GetCompletionTaskGraphEvent()
 {
-	TArray<uint64> TasksToComplete;
+	FGraphEventArray TasksToComplete;
 
 	TasksToComplete.Append(TranslatorTasks);
 	TasksToComplete.Append(PipelineTasks);
 	
-	if (ParsingTask != INTERCHANGE_INVALID_TASK_ID)
+	if (ParsingTask.GetReference())
 	{
 		TasksToComplete.Add(ParsingTask);
 	}
 
 	//Parsing task must be done before the other tasks get added
-	UE::Interchange::FInterchangeTaskSystem::Get().WaitUntilTasksComplete(TasksToComplete);
+	FTaskGraphInterface::Get().WaitUntilTasksComplete(TasksToComplete, ENamedThreads::GameThread);
 	TasksToComplete.Reset();
 
 	TasksToComplete.Append(BeginImportObjectTasks);
 	TasksToComplete.Append(ImportObjectTasks);
 	TasksToComplete.Append(FinalizeImportObjectTasks);
 	TasksToComplete.Append(SceneTasks);
-	TasksToComplete.Append(WaitAssetCompilationTasks);
+	if (WaitAssetCompilationTask.GetReference())
+	{
+		TasksToComplete.Add(WaitAssetCompilationTask);
+	}
 	TasksToComplete.Append(PostImportTasks);
 
-	if (PreCompletionTask != INTERCHANGE_INVALID_TASK_ID)
+	if (PreCompletionTask.GetReference())
 	{
 		TasksToComplete.Add(PreCompletionTask);
 	}
 	
-	if (CompletionTask != INTERCHANGE_INVALID_TASK_ID)
+	if (CompletionTask.GetReference())
 	{
 		//Completion task will make sure any created asset before canceling will be mark for delete
 		TasksToComplete.Add(CompletionTask);
@@ -599,6 +589,7 @@ TArray<uint64> UE::Interchange::FImportAsyncHelper::GetCompletionTaskGraphEvent(
 void UE::Interchange::FImportAsyncHelper::InitCancel()
 {
 	bCancel = true;
+	ReleaseTranslatorsSource();
 }
 
 void UE::Interchange::FImportAsyncHelper::CleanUp()
@@ -676,7 +667,10 @@ bool UE::Interchange::FImportResult::IsValid() const
 void UE::Interchange::FImportResult::SetInProgress()
 {
 	EStatus ExpectedStatus = EStatus::Invalid;
-	ImportStatus.compare_exchange_strong(ExpectedStatus, EStatus::InProgress);
+	if (ImportStatus.compare_exchange_strong(ExpectedStatus, EStatus::InProgress))
+	{
+		GraphEvent = FGraphEvent::CreateGraphEvent();
+	}
 }
 
 void UE::Interchange::FImportResult::SetDone()
@@ -707,12 +701,13 @@ void UE::Interchange::FImportResult::SetDone()
 				WeakObjects.Emplace(Object);
 			}
 
-			if(ensure(IsInGameThread()))
+			// call the callbacks on the game thread
+			Async(EAsyncExecution::TaskGraphMainThread, [InWeakObjects = MoveTemp(WeakObjects), ImportDoneNative = OnImportDoneNative, ImportDone = OnImportDone]()
 			{
 				TArray<UObject*> ValidObjects;
-				ValidObjects.Reserve(WeakObjects.Num());
+				ValidObjects.Reserve(InWeakObjects.Num());
 
-				for (const TWeakObjectPtr<UObject>& WeakObject : WeakObjects)
+				for (const TWeakObjectPtr<UObject>& WeakObject : InWeakObjects)
 				{
 					if (UObject* ValidObject = WeakObject.Get())
 					{
@@ -720,10 +715,13 @@ void UE::Interchange::FImportResult::SetDone()
 					}
 				}
 
-				OnImportDoneNative.ExecuteIfBound(ValidObjects);
-				OnImportDone.ExecuteIfBound(ValidObjects);
-			}
+				ImportDoneNative.ExecuteIfBound(ValidObjects);
+				ImportDone.ExecuteIfBound(ValidObjects);
+			});
 		}
+
+
+		GraphEvent->DispatchSubsequents();
 	}
 }
 
@@ -731,17 +729,14 @@ void UE::Interchange::FImportResult::WaitUntilDone(bool bSynchronous /*= false*/
 {
 	if (ImportStatus == EStatus::InProgress)
 	{
-		//Pin the weak ptr, do not hold the shared ptr until the end of the import, simply get the completion task id
-		TArray<uint64> TasksIds;
+		if (bSynchronous)
 		{
-			TSharedPtr<FImportAsyncHelper, ESPMode::ThreadSafe> PinAsyncHelper = AsyncHelper.Pin();
-			if (PinAsyncHelper.IsValid())
-			{
-				TasksIds = PinAsyncHelper->GetCompletionTaskGraphEvent();
-			}
+			FTaskGraphInterface::Get().ProcessThreadUntilIdle(ENamedThreads::GameThread_Local);
+		}
+		else
+		{
+			FTaskGraphInterface::Get().WaitUntilTaskCompletes(GraphEvent);
 		}
-		//Wait until the whole import is done
-		UE::Interchange::FInterchangeTaskSystem::Get().WaitUntilTasksComplete(TasksIds);
 	}
 }
 
@@ -814,11 +809,6 @@ void UE::Interchange::FImportResult::AddReferencedObjects(FReferenceCollector& C
 	Collector.AddReferencedObject(Results);
 }
 
-void UE::Interchange::FImportResult::SetAsyncHelper(TWeakPtr<FImportAsyncHelper> InAsyncHelper)
-{
-	AsyncHelper = InAsyncHelper;
-}
-
 UInterchangePipelineBase* UE::Interchange::GeneratePipelineInstance(const FSoftObjectPath& PipelineInstance)
 {
 	UPackage* PipelineInstancePackage = GetTransientPackage();
@@ -939,6 +929,22 @@ UInterchangeManager& UInterchangeManager::GetInterchangeManager()
 
 		bIsCreatingSingleton = false;
 
+		InterchangeManager->GCEndDelegate = FCoreUObjectDelegates::GetPostGarbageCollect().AddLambda([]()
+			{
+				if (IsInterchangeImportEnabled() && InterchangeManager.IsValid())
+				{
+					InterchangeManager->StartQueuedTasks(InterchangeManager->bGCEndDelegateCancellAllTask);
+				}
+			});
+		InterchangeManager->GCPreDelegate = FCoreUObjectDelegates::GetPreGarbageCollectDelegate().AddLambda([]()
+			{
+				if (IsInterchangeImportEnabled() && InterchangeManager.IsValid())
+				{
+					const bool bCancel = !GIsEditor;
+					InterchangeManager->WaitUntilAllTasksDone(bCancel);
+				}
+			});
+
 		//We cancel any running task when we pre exit the engine
 		FCoreDelegates::OnEnginePreExit.AddLambda([]()
 		{
@@ -958,6 +964,17 @@ UInterchangeManager& UInterchangeManager::GetInterchangeManager()
 				InterchangeManager->WaitUntilAllTasksDone(bCancel);
 			}
 
+			//Remove any delegate
+			if (InterchangeManager->GCEndDelegate.IsValid())
+			{
+				FCoreUObjectDelegates::GetPostGarbageCollect().Remove(InterchangeManager->GCEndDelegate);
+				InterchangeManager->GCEndDelegate.Reset();
+			}
+			if(InterchangeManager->GCPreDelegate.IsValid())
+			{
+				FCoreUObjectDelegates::GetPreGarbageCollectDelegate().Remove(InterchangeManager->GCPreDelegate);
+				InterchangeManager->GCPreDelegate.Reset();
+			}
 			//Task should have been cancel in the Engine pre exit callback
 			ensure(InterchangeManager->ImportTasks.Num() == 0);
 			InterchangeManager->OnPreDestroyInterchangeManager.Broadcast();
@@ -1453,7 +1470,7 @@ void UInterchangeManager::StartQueuedTasks(bool bCancelAllTasks /*= false*/)
 			}
 
 			//Create/Start import tasks
-			TArray<uint64> PipelinePrerequistes;
+			FGraphEventArray PipelinePrerequistes;
 			if(QueuedTaskData.AsyncHelper->TranslatorTasks.Num() == 0)
 			{
 				check(QueuedTaskData.AsyncHelper->Translators.Num() == QueuedTaskData.AsyncHelper->SourceDatas.Num());
@@ -1461,20 +1478,18 @@ void UInterchangeManager::StartQueuedTasks(bool bCancelAllTasks /*= false*/)
 				{
 					//Log the source we begin importing
 					UE_LOG(LogInterchangeEngine, Display, TEXT("Interchange start importing source [%s]"), *QueuedTaskData.AsyncHelper->SourceDatas[SourceDataIndex]->GetFilename());
-					TSharedPtr<UE::Interchange::FTaskTranslator, ESPMode::ThreadSafe> TaskTranslator = MakeShared<UE::Interchange::FTaskTranslator, ESPMode::ThreadSafe>(SourceDataIndex, WeakAsyncHelper);
-					int32 TranslatorTaskIndex = QueuedTaskData.AsyncHelper->TranslatorTasks.Add(UE::Interchange::FInterchangeTaskSystem::Get().AddTask(TaskTranslator));
+					int32 TranslatorTaskIndex = QueuedTaskData.AsyncHelper->TranslatorTasks.Add(TGraphTask<UE::Interchange::FTaskTranslator>::CreateTask().ConstructAndDispatchWhenReady(SourceDataIndex, WeakAsyncHelper));
 					PipelinePrerequistes.Add(QueuedTaskData.AsyncHelper->TranslatorTasks[TranslatorTaskIndex]);
 				}
 			}
 
-			TArray<uint64> GraphParsingPrerequistes;
+			FGraphEventArray GraphParsingPrerequistes;
 			for (int32 GraphPipelineIndex = 0; GraphPipelineIndex < QueuedTaskData.AsyncHelper->Pipelines.Num(); ++GraphPipelineIndex)
 			{
 				UInterchangePipelineBase* GraphPipeline = QueuedTaskData.AsyncHelper->Pipelines[GraphPipelineIndex];
 				TWeakObjectPtr<UInterchangePipelineBase> WeakPipelinePtr = GraphPipeline;
-
-				TSharedPtr<UE::Interchange::FTaskPipeline, ESPMode::ThreadSafe> TaskPipeline = MakeShared<UE::Interchange::FTaskPipeline, ESPMode::ThreadSafe>(WeakPipelinePtr, WeakAsyncHelper);
-				int32 GraphPipelineTaskIndex = QueuedTaskData.AsyncHelper->PipelineTasks.Add(UE::Interchange::FInterchangeTaskSystem::Get().AddTask(TaskPipeline, PipelinePrerequistes));
+				int32 GraphPipelineTaskIndex = INDEX_NONE;
+				GraphPipelineTaskIndex = QueuedTaskData.AsyncHelper->PipelineTasks.Add(TGraphTask<UE::Interchange::FTaskPipeline>::CreateTask(&PipelinePrerequistes).ConstructAndDispatchWhenReady(WeakPipelinePtr, WeakAsyncHelper));
 				//Ensure we run the pipeline in the same order we create the task, since pipeline modify the node container, its important that its not process in parallel, Adding the one we start to the prerequisites
 				//is the way to go here
 				PipelinePrerequistes.Add(QueuedTaskData.AsyncHelper->PipelineTasks[GraphPipelineTaskIndex]);
@@ -1483,15 +1498,16 @@ void UInterchangeManager::StartQueuedTasks(bool bCancelAllTasks /*= false*/)
 				GraphParsingPrerequistes.Add(QueuedTaskData.AsyncHelper->PipelineTasks[GraphPipelineTaskIndex]);
 			}
 
-			if (GraphParsingPrerequistes.Num() == 0)
+			if (GraphParsingPrerequistes.Num() > 0)
+			{
+				QueuedTaskData.AsyncHelper->ParsingTask = TGraphTask<UE::Interchange::FTaskParsing>::CreateTask(&GraphParsingPrerequistes).ConstructAndDispatchWhenReady(this, WeakAsyncHelper);
+			}
+			else
 			{
 				//Fallback on the translator pipeline prerequisites (translator must be done if there is no pipeline)
-				GraphParsingPrerequistes.Append(PipelinePrerequistes);
+				QueuedTaskData.AsyncHelper->ParsingTask = TGraphTask<UE::Interchange::FTaskParsing>::CreateTask(&PipelinePrerequistes).ConstructAndDispatchWhenReady(this, WeakAsyncHelper);
 			}
 
-			TSharedPtr<UE::Interchange::FTaskParsing, ESPMode::ThreadSafe> TaskParsing = MakeShared<UE::Interchange::FTaskParsing, ESPMode::ThreadSafe>(this, WeakAsyncHelper);
-			QueuedTaskData.AsyncHelper->ParsingTask = UE::Interchange::FInterchangeTaskSystem::Get().AddTask(TaskParsing, GraphParsingPrerequistes);
-
 			//The graph parsing task will create the FCreateAssetTask that will run after them, the FAssetImportTask will call the appropriate Post asset import pipeline when the asset is completed
 		}
 	}
@@ -1901,14 +1917,13 @@ UInterchangeManager::ImportInternal(const FString& ContentPath, const UInterchan
 				Progress.MakeDialog();
 				Progress.EnterProgressFrame(1.f);
 				//Translate the source
+				FGraphEventArray PipelinePrerequistes;
 				for (int32 SourceDataIndex = 0; SourceDataIndex < AsyncHelper->SourceDatas.Num(); ++SourceDataIndex)
 				{
-					constexpr uint64 PipelinePrerequistes = INTERCHANGE_INVALID_TASK_ID;
 					//Log the source we begin importing
 					UE_LOG(LogInterchangeEngine, Display, TEXT("Interchange start importing source [%s]"), *AsyncHelper->SourceDatas[SourceDataIndex]->GetFilename());
-					UE::Interchange::FTaskTranslator TaskTranslator(SourceDataIndex, AsyncHelper);
-					TaskTranslator.Execute();
-					AsyncHelper->TranslatorTasks.Add(PipelinePrerequistes);
+					int32 TranslatorTaskIndex = AsyncHelper->TranslatorTasks.Add(TGraphTask<UE::Interchange::FTaskTranslator>::CreateTask(nullptr, ENamedThreads::GameThread_Local).ConstructAndDispatchWhenReady(SourceDataIndex, AsyncHelper));
+					AsyncHelper->TranslatorTasks[TranslatorTaskIndex]->Wait();
 				}
 				Progress.EnterProgressFrame(1.f);
 			}
@@ -2375,13 +2390,11 @@ TSharedRef<UE::Interchange::FImportAsyncHelper, ESPMode::ThreadSafe> UInterchang
 	AsyncHelper->AssetImportResult->OnObjectDoneNative = ImportAssetParameters.OnAssetDoneNative;
 	AsyncHelper->AssetImportResult->OnImportDone = ImportAssetParameters.OnAssetsImportDone;
 	AsyncHelper->AssetImportResult->OnImportDoneNative = ImportAssetParameters.OnAssetsImportDoneNative;
-	AsyncHelper->AssetImportResult->SetAsyncHelper(AsyncHelper);
 
 	AsyncHelper->SceneImportResult->OnObjectDone = ImportAssetParameters.OnSceneObjectDone;
 	AsyncHelper->SceneImportResult->OnObjectDoneNative = ImportAssetParameters.OnSceneObjectDoneNative;
 	AsyncHelper->SceneImportResult->OnImportDone = ImportAssetParameters.OnSceneImportDone;
 	AsyncHelper->SceneImportResult->OnImportDoneNative = ImportAssetParameters.OnSceneImportDoneNative;
-	AsyncHelper->SceneImportResult->SetAsyncHelper(AsyncHelper);
 
 	AsyncHelper->AssetImportResult->SetInProgress();
 
@@ -2451,7 +2464,7 @@ void UInterchangeManager::ReleaseAsyncHelper(TWeakPtr<UE::Interchange::FImportAs
 		check(!AsyncHelperSharedPtr.IsValid() || AsyncHelperSharedPtr->bCancel);
 	}
 
-	int32 ImportTaskNumber = ImportTasks.Num() + QueueTaskCount;
+	int32 ImportTaskNumber = ImportTasks.Num();
 	FString ImportTaskNumberStr = TEXT(" (") + FString::FromInt(ImportTaskNumber) + TEXT(")");
 	if (ImportTaskNumber == 0)
 	{
@@ -2709,10 +2722,10 @@ void UInterchangeManager::WaitUntilAllTasksDone(bool bCancel)
 		if (AsyncHelper.IsValid())
 		{
 			TWeakPtr<UE::Interchange::FImportAsyncHelper, ESPMode::ThreadSafe> WeakAsyncHelper = AsyncHelper;
-			TArray<uint64> TasksToComplete = AsyncHelper->GetCompletionTaskGraphEvent();
+			FGraphEventArray TasksToComplete = AsyncHelper->GetCompletionTaskGraphEvent();
 			//Release the shared pointer before waiting to be sure the async helper can be destroy in the completion task
 			AsyncHelper = nullptr;
-			UE::Interchange::FInterchangeTaskSystem::Get().WaitUntilTasksComplete(TasksToComplete);
+			FTaskGraphInterface::Get().WaitUntilTasksComplete(TasksToComplete, ENamedThreads::GameThread);
 			//We verify that the weak pointer is invalid after the task completed
 			ensure(!WeakAsyncHelper.IsValid());
 		}
diff --git a/Engine/Source/Runtime/Interchange/Engine/Private/InterchangeTaskSystem.cpp b/Engine/Source/Runtime/Interchange/Engine/Private/InterchangeTaskSystem.cpp
deleted file mode 100644
diff --git a/Engine/Source/Runtime/Interchange/Engine/Private/Tasks/InterchangeTaskCompletion.cpp b/Engine/Source/Runtime/Interchange/Engine/Private/Tasks/InterchangeTaskCompletion.cpp
--- a/Engine/Source/Runtime/Interchange/Engine/Private/Tasks/InterchangeTaskCompletion.cpp
+++ b/Engine/Source/Runtime/Interchange/Engine/Private/Tasks/InterchangeTaskCompletion.cpp
@@ -20,16 +20,14 @@
 #include "UObject/ObjectMacros.h"
 #include "UObject/WeakObjectPtrTemplates.h"
 
-void UE::Interchange::FTaskPreCompletion_GameThread::Execute()
+void UE::Interchange::FTaskPreCompletion::DoTask(ENamedThreads::Type CurrentThread, const FGraphEventRef& MyCompletionGraphEvent)
 {
-	TRACE_CPUPROFILER_EVENT_SCOPE(UE::Interchange::FTaskPreCompletion_GameThread::DoTask)
+	TRACE_CPUPROFILER_EVENT_SCOPE(UE::Interchange::FTaskPreCompletion::DoTask)
 #if INTERCHANGE_TRACE_ASYNCHRONOUS_TASK_ENABLED
 	INTERCHANGE_TRACE_ASYNCHRONOUS_TASK(PreCompletion)
 #endif
 	
 	LLM_SCOPE_BYNAME(TEXT("Interchange"));
-	
-	check(IsInGameThread());
 
 	TSharedPtr<FImportAsyncHelper, ESPMode::ThreadSafe> AsyncHelper = WeakAsyncHelper.Pin();
 	check(AsyncHelper.IsValid());
@@ -177,15 +175,14 @@ void UE::Interchange::FTaskPreCompletion_GameThread::Execute()
 }
 
 
-void UE::Interchange::FTaskCompletion_GameThread::Execute()
+void UE::Interchange::FTaskCompletion::DoTask(ENamedThreads::Type CurrentThread, const FGraphEventRef& MyCompletionGraphEvent)
 {
-	TRACE_CPUPROFILER_EVENT_SCOPE(UE::Interchange::FTaskCompletion_GameThread::DoTask)
+	TRACE_CPUPROFILER_EVENT_SCOPE(UE::Interchange::FTaskCompletion::DoTask)
 #if INTERCHANGE_TRACE_ASYNCHRONOUS_TASK_ENABLED
 	INTERCHANGE_TRACE_ASYNCHRONOUS_TASK(Completion)
 #endif
 
 	LLM_SCOPE_BYNAME(TEXT("Interchange"));
-	check(IsInGameThread());
 
 	TSharedPtr<FImportAsyncHelper, ESPMode::ThreadSafe> AsyncHelper = WeakAsyncHelper.Pin();
 	check(AsyncHelper.IsValid());
diff --git a/Engine/Source/Runtime/Interchange/Engine/Private/Tasks/InterchangeTaskCompletion.h b/Engine/Source/Runtime/Interchange/Engine/Private/Tasks/InterchangeTaskCompletion.h
--- a/Engine/Source/Runtime/Interchange/Engine/Private/Tasks/InterchangeTaskCompletion.h
+++ b/Engine/Source/Runtime/Interchange/Engine/Private/Tasks/InterchangeTaskCompletion.h
@@ -2,9 +2,9 @@
 
 #pragma once
 
+#include "Async/TaskGraphInterfaces.h"
 #include "CoreMinimal.h"
 #include "InterchangeManager.h"
-#include "InterchangeTaskSystem.h"
 #include "Stats/Stats.h"
 #include "UObject/WeakObjectPtrTemplates.h"
 
@@ -12,44 +12,73 @@ namespace UE
 {
 	namespace Interchange
 	{
-		class FTaskPreCompletion_GameThread : public FInterchangeTaskBase
+		class FTaskPreCompletion
 		{
 		private:
 			UInterchangeManager* InterchangeManager;
 			TWeakPtr<FImportAsyncHelper, ESPMode::ThreadSafe> WeakAsyncHelper;
 		public:
-			FTaskPreCompletion_GameThread(UInterchangeManager* InInterchangeManager, TWeakPtr<FImportAsyncHelper, ESPMode::ThreadSafe> InAsyncHelper)
+			FTaskPreCompletion(UInterchangeManager* InInterchangeManager, TWeakPtr<FImportAsyncHelper, ESPMode::ThreadSafe> InAsyncHelper)
 				: InterchangeManager(InInterchangeManager)
 				, WeakAsyncHelper(InAsyncHelper)
 			{
 			}
 
-			virtual EInterchangeTaskThread GetTaskThread() const override
+			FORCEINLINE ENamedThreads::Type GetDesiredThread()
 			{
-				return EInterchangeTaskThread::GameThread;
+				TSharedPtr<FImportAsyncHelper, ESPMode::ThreadSafe> AsyncHelper = WeakAsyncHelper.Pin();
+				if (AsyncHelper.IsValid() && AsyncHelper->bRunSynchronous)
+				{
+					return ENamedThreads::GameThread_Local;
+				}
+				return ENamedThreads::GameThread;
+			}
+			static FORCEINLINE ESubsequentsMode::Type GetSubsequentsMode()
+			{
+				return ESubsequentsMode::TrackSubsequents;
+			}
+
+			FORCEINLINE TStatId GetStatId() const
+			{
+				RETURN_QUICK_DECLARE_CYCLE_STAT(FTaskPreCompletion, STATGROUP_TaskGraphTasks);
 			}
 
-			virtual void Execute() override;
+			void DoTask(ENamedThreads::Type CurrentThread, const FGraphEventRef& MyCompletionGraphEvent);
 		};
 
-		class FTaskCompletion_GameThread : public FInterchangeTaskBase
+		class FTaskCompletion
 		{
 		private:
 			UInterchangeManager* InterchangeManager;
 			TWeakPtr<FImportAsyncHelper, ESPMode::ThreadSafe> WeakAsyncHelper;
 		public:
-			FTaskCompletion_GameThread(UInterchangeManager* InInterchangeManager, TWeakPtr<FImportAsyncHelper, ESPMode::ThreadSafe> InAsyncHelper)
+			FTaskCompletion(UInterchangeManager* InInterchangeManager, TWeakPtr<FImportAsyncHelper, ESPMode::ThreadSafe> InAsyncHelper)
 				: InterchangeManager(InInterchangeManager)
 				, WeakAsyncHelper(InAsyncHelper)
 			{
 			}
 
-			virtual EInterchangeTaskThread GetTaskThread() const override
+			FORCEINLINE ENamedThreads::Type GetDesiredThread()
+			{
+				TSharedPtr<FImportAsyncHelper, ESPMode::ThreadSafe> AsyncHelper = WeakAsyncHelper.Pin();
+				if (AsyncHelper.IsValid() && AsyncHelper->bRunSynchronous)
+				{
+					return ENamedThreads::GameThread_Local;
+				}
+				return ENamedThreads::GameThread;
+			}
+			static FORCEINLINE ESubsequentsMode::Type GetSubsequentsMode()
+			{
+				//In case we need to know when the task is done we need track subsequent to get a valid FGraphEventRef when we create the task
+				return ESubsequentsMode::TrackSubsequents;
+			}
+
+			FORCEINLINE TStatId GetStatId() const
 			{
-				return EInterchangeTaskThread::GameThread;
+				RETURN_QUICK_DECLARE_CYCLE_STAT(FTaskCompletion, STATGROUP_TaskGraphTasks);
 			}
 
-			virtual void Execute() override;
+			void DoTask(ENamedThreads::Type CurrentThread, const FGraphEventRef& MyCompletionGraphEvent);
 		};
 
 	} //ns Interchange
diff --git a/Engine/Source/Runtime/Interchange/Engine/Private/Tasks/InterchangeTaskCreateSceneObjects.cpp b/Engine/Source/Runtime/Interchange/Engine/Private/Tasks/InterchangeTaskCreateSceneObjects.cpp
--- a/Engine/Source/Runtime/Interchange/Engine/Private/Tasks/InterchangeTaskCreateSceneObjects.cpp
+++ b/Engine/Source/Runtime/Interchange/Engine/Private/Tasks/InterchangeTaskCreateSceneObjects.cpp
@@ -23,7 +23,7 @@
 #include "UObject/UObjectGlobals.h"
 #include "UObject/WeakObjectPtrTemplates.h"
 
-UE::Interchange::FTaskCreateSceneObjects_GameThread::FTaskCreateSceneObjects_GameThread(const FString& InPackageBasePath, const int32 InSourceIndex, TWeakPtr<FImportAsyncHelper> InAsyncHelper, TArrayView<UInterchangeFactoryBaseNode*> InFactoryNodes, const UClass* InFactoryClass)
+UE::Interchange::FTaskCreateSceneObjects::FTaskCreateSceneObjects(const FString& InPackageBasePath, const int32 InSourceIndex, TWeakPtr<FImportAsyncHelper> InAsyncHelper, TArrayView<UInterchangeFactoryBaseNode*> InFactoryNodes, const UClass* InFactoryClass)
 	: PackageBasePath(InPackageBasePath)
 	, SourceIndex(InSourceIndex)
 	, WeakAsyncHelper(InAsyncHelper)
@@ -33,15 +33,32 @@ UE::Interchange::FTaskCreateSceneObjects_GameThread::FTaskCreateSceneObjects_Gam
 	check(FactoryClass);
 }
 
-void UE::Interchange::FTaskCreateSceneObjects_GameThread::Execute()
+ENamedThreads::Type UE::Interchange::FTaskCreateSceneObjects::GetDesiredThread()
 {
-	TRACE_CPUPROFILER_EVENT_SCOPE(UE::Interchange::FTaskCreateSceneObjects_GameThread::DoTask)
+	// We are creating the factories in this task so it must execute on the GameThread.
+	// Also, there are no "CreatePackage Task" equivalent for scene objects right now, so the factories must create those on the game thread.
+	TSharedPtr<FImportAsyncHelper, ESPMode::ThreadSafe> AsyncHelper = WeakAsyncHelper.Pin();
+	if (AsyncHelper.IsValid() && AsyncHelper->bRunSynchronous)
+	{
+		return ENamedThreads::GameThread_Local;
+	}
+	return ENamedThreads::GameThread;
+}
+
+void UE::Interchange::FTaskCreateSceneObjects::DoTask(ENamedThreads::Type CurrentThread, const FGraphEventRef& MyCompletionGraphEvent)
+{
+	TRACE_CPUPROFILER_EVENT_SCOPE(UE::Interchange::FTaskCreateSceneObjects::DoTask)
 #if INTERCHANGE_TRACE_ASYNCHRONOUS_TASK_ENABLED
 	INTERCHANGE_TRACE_ASYNCHRONOUS_TASK(SpawnActor)
 #endif
 	using namespace UE::Interchange;
 
-	check(IsInGameThread());
+	TOptional<FGCScopeGuard> GCScopeGuard;
+	if (!IsInGameThread())
+	{
+		GCScopeGuard.Emplace();
+	}
+
 	TSharedPtr<FImportAsyncHelper> AsyncHelper = WeakAsyncHelper.Pin();
 	check(WeakAsyncHelper.IsValid());
 
diff --git a/Engine/Source/Runtime/Interchange/Engine/Private/Tasks/InterchangeTaskCreateSceneObjects.h b/Engine/Source/Runtime/Interchange/Engine/Private/Tasks/InterchangeTaskCreateSceneObjects.h
--- a/Engine/Source/Runtime/Interchange/Engine/Private/Tasks/InterchangeTaskCreateSceneObjects.h
+++ b/Engine/Source/Runtime/Interchange/Engine/Private/Tasks/InterchangeTaskCreateSceneObjects.h
@@ -2,8 +2,8 @@
 
 #pragma once
 
+#include "Async/TaskGraphInterfaces.h"
 #include "CoreMinimal.h"
-#include "InterchangeTaskSystem.h"
 #include "Stats/Stats.h"
 #include "UObject/WeakObjectPtrTemplates.h"
 #include "Nodes/InterchangeBaseNode.h"
@@ -17,7 +17,7 @@ namespace UE
 	{
 		class FImportAsyncHelper;
 
-		class FTaskCreateSceneObjects_GameThread : public FInterchangeTaskBase
+		class FTaskCreateSceneObjects
 		{
 		private:
 			FString PackageBasePath;
@@ -27,14 +27,23 @@ namespace UE
 			const UClass* FactoryClass;
 
 		public:
-			explicit FTaskCreateSceneObjects_GameThread(const FString& InPackageBasePath, const int32 InSourceIndex, TWeakPtr<FImportAsyncHelper> InAsyncHelper, TArrayView<UInterchangeFactoryBaseNode*> InNodes, const UClass* InFactoryClass);
+			explicit FTaskCreateSceneObjects(const FString& InPackageBasePath, const int32 InSourceIndex, TWeakPtr<FImportAsyncHelper> InAsyncHelper, TArrayView<UInterchangeFactoryBaseNode*> InNodes, const UClass* InFactoryClass);
 
-			virtual EInterchangeTaskThread GetTaskThread() const override
+			ENamedThreads::Type GetDesiredThread();
+
+			static ESubsequentsMode::Type GetSubsequentsMode()
+			{
+				return ESubsequentsMode::TrackSubsequents;
+			}
+
+			TStatId GetStatId() const
 			{
-				return EInterchangeTaskThread::GameThread;
+				RETURN_QUICK_DECLARE_CYCLE_STAT(FTaskCreateSceneObjects, STATGROUP_TaskGraphTasks);
 			}
 
-			virtual void Execute() override;
+			void DoTask(ENamedThreads::Type CurrentThread, const FGraphEventRef& MyCompletionGraphEvent);
 		};
+
+
 	} //ns Interchange
 }//ns UE
diff --git a/Engine/Source/Runtime/Interchange/Engine/Private/Tasks/InterchangeTaskImportObject.cpp b/Engine/Source/Runtime/Interchange/Engine/Private/Tasks/InterchangeTaskImportObject.cpp
--- a/Engine/Source/Runtime/Interchange/Engine/Private/Tasks/InterchangeTaskImportObject.cpp
+++ b/Engine/Source/Runtime/Interchange/Engine/Private/Tasks/InterchangeTaskImportObject.cpp
@@ -239,7 +239,7 @@ namespace UE::Interchange::Private
 	}
 }//ns UE::Interchange::Private
 
-void UE::Interchange::FTaskImportObject_GameThread::Execute()
+void UE::Interchange::FTaskImportObject_GameThread::DoTask(ENamedThreads::Type CurrentThread, const FGraphEventRef& MyCompletionGraphEvent)
 {
 	TRACE_CPUPROFILER_EVENT_SCOPE(UE::Interchange::FTaskImportObject_GameThread::DoTask)
 #if INTERCHANGE_TRACE_ASYNCHRONOUS_TASK_ENABLED
@@ -587,7 +587,7 @@ void UE::Interchange::FTaskImportObject_GameThread::Execute()
 	}
 }
 
-void UE::Interchange::FTaskImportObject_Async::Execute()
+void UE::Interchange::FTaskImportObject_Async::DoTask(ENamedThreads::Type CurrentThread, const FGraphEventRef& MyCompletionGraphEvent)
 {
 #if INTERCHANGE_TRACE_ASYNCHRONOUS_TASK_ENABLED
 	INTERCHANGE_TRACE_ASYNCHRONOUS_TASK(TaskImportObject_Async)
@@ -597,12 +597,6 @@ void UE::Interchange::FTaskImportObject_Async::Execute()
 
 	using namespace UE::Interchange;
 
-	TOptional<FGCScopeGuard> GCScopeGuard;
-	if (!IsInGameThread())
-	{
-		GCScopeGuard.Emplace();
-	}
-
 	TSharedPtr<FImportAsyncHelper, ESPMode::ThreadSafe> AsyncHelper = WeakAsyncHelper.Pin();
 	check(AsyncHelper.IsValid());
 
@@ -629,7 +623,7 @@ void UE::Interchange::FTaskImportObject_Async::Execute()
 		});
 }
 
-void UE::Interchange::FTaskImportObjectFinalize_GameThread::Execute()
+void UE::Interchange::FTaskImportObjectFinalize_GameThread::DoTask(ENamedThreads::Type CurrentThread, const FGraphEventRef& MyCompletionGraphEvent)
 {
 #if INTERCHANGE_TRACE_ASYNCHRONOUS_TASK_ENABLED
 	INTERCHANGE_TRACE_ASYNCHRONOUS_TASK(TaskImportObjectFinalize_GameThread)
@@ -654,8 +648,6 @@ void UE::Interchange::FTaskImportObjectFinalize_GameThread::Execute()
 		return;
 	}
 
-	check(IsInGameThread());
-
 	UInterchangeFactoryBase::FImportAssetResult ImportAssetResult = Private::InternalImportObjectStartup(AsyncHelper
 		, FactoryNode
 		, SourceIndex
diff --git a/Engine/Source/Runtime/Interchange/Engine/Private/Tasks/InterchangeTaskImportObject.h b/Engine/Source/Runtime/Interchange/Engine/Private/Tasks/InterchangeTaskImportObject.h
--- a/Engine/Source/Runtime/Interchange/Engine/Private/Tasks/InterchangeTaskImportObject.h
+++ b/Engine/Source/Runtime/Interchange/Engine/Private/Tasks/InterchangeTaskImportObject.h
@@ -2,10 +2,10 @@
 
 #pragma once
 
+#include "Async/TaskGraphInterfaces.h"
 #include "CoreMinimal.h"
 #include "InterchangeFactoryBase.h"
 #include "InterchangeManager.h"
-#include "InterchangeTaskSystem.h"
 #include "Stats/Stats.h"
 #include "UObject/WeakObjectPtrTemplates.h"
 #include "Nodes/InterchangeFactoryBaseNode.h"
@@ -15,7 +15,7 @@ namespace UE::Interchange
 	/**
 		* This task create UPackage and UObject, Cook::PackageTracker::NotifyUObjectCreated is not thread safe, so we need to create the packages on the main thread
 		*/
-	class FTaskImportObject_GameThread : public FInterchangeTaskBase
+	class FTaskImportObject_GameThread
 	{
 	private:
 		FString PackageBasePath;
@@ -36,15 +36,30 @@ namespace UE::Interchange
 			check(FactoryClass);
 		}
 
-		virtual EInterchangeTaskThread GetTaskThread() const override
+		ENamedThreads::Type GetDesiredThread() const
 		{
-			return EInterchangeTaskThread::GameThread;
+			TSharedPtr<FImportAsyncHelper, ESPMode::ThreadSafe> AsyncHelper = WeakAsyncHelper.Pin();
+			if (AsyncHelper.IsValid() && AsyncHelper->bRunSynchronous)
+			{
+				return ENamedThreads::GameThread_Local;
+			}
+			return ENamedThreads::GameThread;
+		}
+
+		static ESubsequentsMode::Type GetSubsequentsMode()
+		{
+			return ESubsequentsMode::TrackSubsequents;
 		}
 
-		virtual void Execute() override;
+		TStatId GetStatId() const
+		{
+			RETURN_QUICK_DECLARE_CYCLE_STAT(FTaskImportObject_GameThread, STATGROUP_TaskGraphTasks);
+		}
+
+		void DoTask(ENamedThreads::Type CurrentThread, const FGraphEventRef& MyCompletionGraphEvent);
 	};
 
-	class FTaskImportObject_Async : public FInterchangeTaskBase
+	class FTaskImportObject_Async
 	{
 	private:
 		FString PackageBasePath;
@@ -62,24 +77,33 @@ namespace UE::Interchange
 			check(FactoryNode);
 		}
 
-		virtual EInterchangeTaskThread GetTaskThread() const override
+		ENamedThreads::Type GetDesiredThread() const
 		{
 			TSharedPtr<FImportAsyncHelper, ESPMode::ThreadSafe> AsyncHelper = WeakAsyncHelper.Pin();
 			if (AsyncHelper.IsValid() && AsyncHelper->bRunSynchronous)
 			{
-				return EInterchangeTaskThread::GameThread;
+				return ENamedThreads::GameThread_Local;
 			}
+			return ENamedThreads::AnyBackgroundThreadNormalTask;
+		}
 
-			return EInterchangeTaskThread::AsyncThread;
+		static ESubsequentsMode::Type GetSubsequentsMode()
+		{
+			return ESubsequentsMode::TrackSubsequents;
 		}
 
-		virtual void Execute() override;
+		TStatId GetStatId() const
+		{
+			RETURN_QUICK_DECLARE_CYCLE_STAT(FTaskImportObject_Async, STATGROUP_TaskGraphTasks);
+		}
+
+		void DoTask(ENamedThreads::Type CurrentThread, const FGraphEventRef& MyCompletionGraphEvent);
 	};
 
 	/**
 		* This task create UPackage and UObject, Cook::PackageTracker::NotifyUObjectCreated is not thread safe, so we need to create the packages on the main thread
 		*/
-	class FTaskImportObjectFinalize_GameThread : public FInterchangeTaskBase
+	class FTaskImportObjectFinalize_GameThread
 	{
 	private:
 		FString PackageBasePath;
@@ -98,12 +122,27 @@ namespace UE::Interchange
 			check(FactoryNode);
 		}
 
-		virtual EInterchangeTaskThread GetTaskThread() const override
+		ENamedThreads::Type GetDesiredThread() const
+		{
+			TSharedPtr<FImportAsyncHelper, ESPMode::ThreadSafe> AsyncHelper = WeakAsyncHelper.Pin();
+			if (AsyncHelper.IsValid() && AsyncHelper->bRunSynchronous)
+			{
+				return ENamedThreads::GameThread_Local;
+			}
+			return ENamedThreads::GameThread;
+		}
+
+		static ESubsequentsMode::Type GetSubsequentsMode()
+		{
+			return ESubsequentsMode::TrackSubsequents;
+		}
+
+		TStatId GetStatId() const
 		{
-			return EInterchangeTaskThread::GameThread;
+			RETURN_QUICK_DECLARE_CYCLE_STAT(FTaskImportObjectFinalize_GameThread, STATGROUP_TaskGraphTasks);
 		}
 
-		virtual void Execute() override;
+		void DoTask(ENamedThreads::Type CurrentThread, const FGraphEventRef& MyCompletionGraphEvent);
 	};
 
 }//ns UE::Interchange
diff --git a/Engine/Source/Runtime/Interchange/Engine/Private/Tasks/InterchangeTaskParsing.cpp b/Engine/Source/Runtime/Interchange/Engine/Private/Tasks/InterchangeTaskParsing.cpp
--- a/Engine/Source/Runtime/Interchange/Engine/Private/Tasks/InterchangeTaskParsing.cpp
+++ b/Engine/Source/Runtime/Interchange/Engine/Private/Tasks/InterchangeTaskParsing.cpp
@@ -87,7 +87,7 @@ private:
 };
 
 
-void UE::Interchange::FTaskParsing::Execute()
+void UE::Interchange::FTaskParsing::DoTask(ENamedThreads::Type CurrentThread, const FGraphEventRef& MyCompletionGraphEvent)
 {
 	TRACE_CPUPROFILER_EVENT_SCOPE(UE::Interchange::FTaskParsing::DoTask)
 #if INTERCHANGE_TRACE_ASYNCHRONOUS_TASK_ENABLED
@@ -95,11 +95,7 @@ void UE::Interchange::FTaskParsing::Execute()
 #endif
 	LLM_SCOPE_BYNAME(TEXT("Interchange"));
 
-	TOptional<FGCScopeGuard> GCScopeGuard;
-	if (!IsInGameThread())
-	{
-		GCScopeGuard.Emplace();
-	}
+	FGCScopeGuard GCScopeGuard;
 
 	TSharedPtr<FImportAsyncHelper, ESPMode::ThreadSafe> AsyncHelper = WeakAsyncHelper.Pin();
 	check(AsyncHelper.IsValid());
@@ -111,8 +107,8 @@ void UE::Interchange::FTaskParsing::Execute()
 		int32 SourceIndex = INDEX_NONE;
 		bool bIsSceneNode = false;
 		TArray<FString> Dependencies;
-		uint64 GraphEventRef;
-		TArray<uint64> Prerequisites;
+		FGraphEventRef GraphEventRef;
+		FGraphEventArray Prerequisites;
 		const UClass* FactoryClass;
 
 		TArray<UInterchangeFactoryBaseNode*, TInlineAllocator<1>> Nodes; // For scenes, we can group multiple nodes into a single task as they are usually very light
@@ -200,9 +196,9 @@ void UE::Interchange::FTaskParsing::Execute()
 		}
 	}
 
-	auto CreateTasksForEachTaskData = [](TArray<FTaskData>& TaskDatas, TFunction<uint64(FTaskData&)> CreateTasksFunc) -> TArray<uint64>
+	auto CreateTasksForEachTaskData = [](TArray<FTaskData>& TaskDatas, TFunction<FGraphEventRef(FTaskData&)> CreateTasksFunc) -> FGraphEventArray
 	{
-		TArray<uint64> GraphEvents;
+		FGraphEventArray GraphEvents;
 
 		for (int32 TaskIndex = 0; TaskIndex < TaskDatas.Num(); ++TaskIndex)
 		{
@@ -238,7 +234,7 @@ void UE::Interchange::FTaskParsing::Execute()
 
 	TMap<FString, FTaskParsingRenameInfo> RenameAssets;
 	TSet<FString> CreatedTasksAssetNames; // Tracks for which asset name we have created a task so that we don't have 2 tasks for the same asset name
-	TFunction<uint64(FTaskData&)> CreateTasksFromData = [this, &AsyncHelper, &RenameAssets, &CreatedTasksAssetNames](FTaskData& TaskData)
+	TFunction<FGraphEventRef(FTaskData&)> CreateTasksFromData = [this, &AsyncHelper, &RenameAssets, &CreatedTasksAssetNames](FTaskData& TaskData)
 	{
 		LLM_SCOPE_BYNAME(TEXT("Interchange"));
 		check(TaskData.Nodes.Num() == 1); //We expect 1 node per asset task
@@ -249,13 +245,9 @@ void UE::Interchange::FTaskParsing::Execute()
 
 		if (TaskData.bIsSceneNode)
 		{
-			TSharedPtr<UE::Interchange::FTaskCreateSceneObjects_GameThread, ESPMode::ThreadSafe> TaskCreateSceneObjects = MakeShared<UE::Interchange::FTaskCreateSceneObjects_GameThread, ESPMode::ThreadSafe>(AsyncHelper->ContentBasePath
-				, SourceIndex
-				, WeakAsyncHelper
-				, TaskData.Nodes
-				, FactoryClass);
-
-			return AsyncHelper->SceneTasks.Add_GetRef(UE::Interchange::FInterchangeTaskSystem::Get().AddTask(TaskCreateSceneObjects, TaskData.Prerequisites));
+			return AsyncHelper->SceneTasks.Add_GetRef(
+				TGraphTask<FTaskCreateSceneObjects>::CreateTask(&(TaskData.Prerequisites))
+				.ConstructAndDispatchWhenReady(AsyncHelper->ContentBasePath, SourceIndex, WeakAsyncHelper, TaskData.Nodes, FactoryClass));
 		}
 		else
 		{
@@ -298,31 +290,22 @@ void UE::Interchange::FTaskParsing::Execute()
 				TEXT("Found multiple task data with the same asset name (%s). Only one will be executed."), *AssetFullPath))
 			{
 				LLM_SCOPE_BYNAME(TEXT("Interchange"));
-				TArray<uint64> ImportObjectTasksPrerequistes;
-				TSharedPtr<FTaskImportObject_GameThread, ESPMode::ThreadSafe> TaskImportObject_GameThread = MakeShared<FTaskImportObject_GameThread, ESPMode::ThreadSafe>(AsyncHelper->ContentBasePath
-					, SourceIndex
-					, WeakAsyncHelper
-					, FactoryNode
-					, FactoryClass);
-				int32 BeginImportObjectTaskIndex = AsyncHelper->BeginImportObjectTasks.Add(FInterchangeTaskSystem::Get().AddTask(TaskImportObject_GameThread, TaskData.Prerequisites));
+				FGraphEventArray ImportObjectTasksPrerequistes;
+				int32 BeginImportObjectTaskIndex = AsyncHelper->BeginImportObjectTasks.Add(
+					TGraphTask<FTaskImportObject_GameThread>::CreateTask(&(TaskData.Prerequisites)).ConstructAndDispatchWhenReady(AsyncHelper->ContentBasePath, SourceIndex, WeakAsyncHelper, FactoryNode, FactoryClass)
+				);
 				ImportObjectTasksPrerequistes.Add(AsyncHelper->BeginImportObjectTasks[BeginImportObjectTaskIndex]);
 
-				TSharedPtr<FTaskImportObject_Async, ESPMode::ThreadSafe> TaskImportObject_Async = MakeShared<FTaskImportObject_Async, ESPMode::ThreadSafe>(AsyncHelper->ContentBasePath
-					, SourceIndex
-					, WeakAsyncHelper
-					, FactoryNode);
-				int32 ImportObjectTaskIndex = AsyncHelper->ImportObjectTasks.Add(FInterchangeTaskSystem::Get().AddTask(TaskImportObject_Async, ImportObjectTasksPrerequistes));
-				
-				TArray<uint64> FinalizeImportObjectTasksPrerequistes;
+				int32 ImportObjectTaskIndex = AsyncHelper->ImportObjectTasks.Add(
+					TGraphTask<FTaskImportObject_Async>::CreateTask(&(ImportObjectTasksPrerequistes)).ConstructAndDispatchWhenReady(AsyncHelper->ContentBasePath, SourceIndex, WeakAsyncHelper, FactoryNode)
+				);
+
+				FGraphEventArray FinalizeImportObjectTasksPrerequistes;
 				FinalizeImportObjectTasksPrerequistes.Add(AsyncHelper->ImportObjectTasks[ImportObjectTaskIndex]);
 
-				
-				
-				TSharedPtr<FTaskImportObjectFinalize_GameThread, ESPMode::ThreadSafe> TaskImportObjectFinalize_GameThread = MakeShared<FTaskImportObjectFinalize_GameThread, ESPMode::ThreadSafe>(AsyncHelper->ContentBasePath
-					, SourceIndex
-					, WeakAsyncHelper
-					, FactoryNode);
-				int32 FinalizeCreateTaskIndex = AsyncHelper->FinalizeImportObjectTasks.Add(FInterchangeTaskSystem::Get().AddTask(TaskImportObjectFinalize_GameThread, FinalizeImportObjectTasksPrerequistes));
+				int32 FinalizeCreateTaskIndex = AsyncHelper->FinalizeImportObjectTasks.Add(
+					TGraphTask<FTaskImportObjectFinalize_GameThread>::CreateTask(&(FinalizeImportObjectTasksPrerequistes)).ConstructAndDispatchWhenReady(AsyncHelper->ContentBasePath, SourceIndex, WeakAsyncHelper, FactoryNode)
+				);
 
 				//Only add the name if the factory node is enabled
 				if (FactoryNode->IsEnabled())
@@ -334,7 +317,8 @@ void UE::Interchange::FTaskParsing::Execute()
 			}
 			else
 			{
-				constexpr uint64 EmptyGraphEvent = 0;
+				FGraphEventRef EmptyGraphEvent = FGraphEvent::CreateGraphEvent();
+				EmptyGraphEvent->DispatchSubsequents();
 				return EmptyGraphEvent;
 			}
 		}
@@ -357,10 +341,10 @@ void UE::Interchange::FTaskParsing::Execute()
 		}
 	}
 
-	TArray<uint64> CompletionPrerequistes;
+	FGraphEventArray CompletionPrerequistes;
 	const int32 PoolWorkerThreadCount = FTaskGraphInterface::Get().GetNumWorkerThreads() / 2;
 	const int32 MaxNumWorker = FMath::Max(PoolWorkerThreadCount, 1);
-	TArray<uint64> GroupPrerequistes;
+	FGraphEventArray GroupPrerequistes;
 	for (int32 TaskIndex = 0; TaskIndex < TaskDatas.Num(); ++TaskIndex)
 	{
 		FTaskData& TaskData = TaskDatas[TaskIndex];
@@ -424,32 +408,28 @@ void UE::Interchange::FTaskParsing::Execute()
 	}
 
 	//Add an async task for pre completion
-	TArray<uint64> PreCompletionPrerequistes;
-	TSharedPtr<FTaskPreCompletion_GameThread, ESPMode::ThreadSafe> TaskPreCompletion = MakeShared<FTaskPreCompletion_GameThread, ESPMode::ThreadSafe>(InterchangeManager, WeakAsyncHelper);
-	AsyncHelper->PreCompletionTask = FInterchangeTaskSystem::Get().AddTask(TaskPreCompletion, CompletionPrerequistes);
+	FGraphEventArray PreCompletionPrerequistes;
+	AsyncHelper->PreCompletionTask = TGraphTask<FTaskPreCompletion>::CreateTask(&CompletionPrerequistes).ConstructAndDispatchWhenReady(InterchangeManager, WeakAsyncHelper);
 	PreCompletionPrerequistes.Add(AsyncHelper->PreCompletionTask);
 
-	TArray<uint64> AssetCompilationPrerequistes;
+	FGraphEventArray AssetCompilationPrerequistes;
 	//Start the wait for asset compilation task on a async task
 	for (int32 SourceIndex = 0; SourceIndex < AsyncHelper->SourceDatas.Num(); ++SourceIndex)
 	{
-		TSharedPtr<FTaskWaitAssetCompilation_GameThread, ESPMode::ThreadSafe> TaskWaitAssetCompilation = MakeShared<FTaskWaitAssetCompilation_GameThread, ESPMode::ThreadSafe>(SourceIndex, WeakAsyncHelper);
-		int32 WaitAssetCompilationTaskIndex = AsyncHelper->WaitAssetCompilationTasks.Add(FInterchangeTaskSystem::Get().AddTask(TaskWaitAssetCompilation, PreCompletionPrerequistes));
-		AssetCompilationPrerequistes.Add(AsyncHelper->WaitAssetCompilationTasks[WaitAssetCompilationTaskIndex]);
+		AsyncHelper->WaitAssetCompilationTask = TGraphTask<FTaskWaitAssetCompilation>::CreateTask(&PreCompletionPrerequistes).ConstructAndDispatchWhenReady(SourceIndex, WeakAsyncHelper);
+		AssetCompilationPrerequistes.Add(AsyncHelper->WaitAssetCompilationTask);
 	}
 
 	//Start the Post pipeline task
 	for (int32 SourceIndex = 0; SourceIndex < AsyncHelper->SourceDatas.Num(); ++SourceIndex)
 	{
-
-		TSharedPtr<FTaskPostImport_GameThread, ESPMode::ThreadSafe> TaskPostImport = MakeShared<FTaskPostImport_GameThread, ESPMode::ThreadSafe>(SourceIndex, WeakAsyncHelper);
-		int32 PostImportTaskIndex = AsyncHelper->PostImportTasks.Add(FInterchangeTaskSystem::Get().AddTask(TaskPostImport, AssetCompilationPrerequistes));
-
+		int32 GraphTaskIndex = AsyncHelper->PostImportTasks.Add(
+			TGraphTask<FTaskPostImport>::CreateTask(&(AssetCompilationPrerequistes)).ConstructAndDispatchWhenReady(SourceIndex, WeakAsyncHelper)
+		);
 		//Ensure we run the pipeline in the same order we create the task, since the pipeline modifies the node container, its important that its not processed in parallel, Adding the one we start to the prerequisites
 		//is the way to go here
-		AssetCompilationPrerequistes.Add(AsyncHelper->PostImportTasks[PostImportTaskIndex]);
+		AssetCompilationPrerequistes.Add(AsyncHelper->PostImportTasks[GraphTaskIndex]);
 	}
 
-	TSharedPtr<FTaskCompletion_GameThread, ESPMode::ThreadSafe> TaskCompletion = MakeShared<FTaskCompletion_GameThread, ESPMode::ThreadSafe>(InterchangeManager, WeakAsyncHelper);
-	AsyncHelper->CompletionTask = FInterchangeTaskSystem::Get().AddTask(TaskCompletion, AssetCompilationPrerequistes);
+	AsyncHelper->CompletionTask = TGraphTask<FTaskCompletion>::CreateTask(&AssetCompilationPrerequistes).ConstructAndDispatchWhenReady(InterchangeManager, WeakAsyncHelper);
 }
diff --git a/Engine/Source/Runtime/Interchange/Engine/Private/Tasks/InterchangeTaskParsing.h b/Engine/Source/Runtime/Interchange/Engine/Private/Tasks/InterchangeTaskParsing.h
--- a/Engine/Source/Runtime/Interchange/Engine/Private/Tasks/InterchangeTaskParsing.h
+++ b/Engine/Source/Runtime/Interchange/Engine/Private/Tasks/InterchangeTaskParsing.h
@@ -2,9 +2,9 @@
 
 #pragma once
 
+#include "Async/TaskGraphInterfaces.h"
 #include "CoreMinimal.h"
 #include "InterchangeManager.h"
-#include "InterchangeTaskSystem.h"
 #include "Stats/Stats.h"
 #include "UObject/WeakObjectPtrTemplates.h"
 
@@ -13,7 +13,7 @@ namespace UE
 	namespace Interchange
 	{
 
-		class FTaskParsing : public FInterchangeTaskBase
+		class FTaskParsing
 		{
 		private:
 			UInterchangeManager* InterchangeManager;
@@ -26,18 +26,27 @@ namespace UE
 				check(InterchangeManager);
 			}
 
-			virtual EInterchangeTaskThread GetTaskThread() const override
+			FORCEINLINE ENamedThreads::Type GetDesiredThread()
 			{
 				TSharedPtr<FImportAsyncHelper, ESPMode::ThreadSafe> AsyncHelper = WeakAsyncHelper.Pin();
 				if (AsyncHelper.IsValid() && AsyncHelper->bRunSynchronous)
 				{
-					return EInterchangeTaskThread::GameThread;
+					return ENamedThreads::GameThread_Local;
 				}
+				return ENamedThreads::AnyBackgroundThreadNormalTask;
+			}
+			static FORCEINLINE ESubsequentsMode::Type GetSubsequentsMode()
+			{
 
-				return EInterchangeTaskThread::AsyncThread;
+				return ESubsequentsMode::TrackSubsequents;
+			}
+
+			FORCEINLINE TStatId GetStatId() const
+			{
+				RETURN_QUICK_DECLARE_CYCLE_STAT(FTaskParsing, STATGROUP_TaskGraphTasks);
 			}
 
-			virtual void Execute() override;
+			void DoTask(ENamedThreads::Type CurrentThread, const FGraphEventRef& MyCompletionGraphEvent);
 		};
 
 	} //ns Interchange
diff --git a/Engine/Source/Runtime/Interchange/Engine/Private/Tasks/InterchangeTaskPipeline.cpp b/Engine/Source/Runtime/Interchange/Engine/Private/Tasks/InterchangeTaskPipeline.cpp
--- a/Engine/Source/Runtime/Interchange/Engine/Private/Tasks/InterchangeTaskPipeline.cpp
+++ b/Engine/Source/Runtime/Interchange/Engine/Private/Tasks/InterchangeTaskPipeline.cpp
@@ -3,6 +3,7 @@
 
 #include "AssetCompilingManager.h"
 #include "Async/Async.h"
+#include "Async/TaskGraphInterfaces.h"
 #include "CoreMinimal.h"
 #include "GenericPlatform/GenericPlatformProcess.h"
 #include "InterchangeEngineLogPrivate.h"
@@ -19,7 +20,8 @@
 #include "UObject/WeakObjectPtrTemplates.h"
 
 
-void UE::Interchange::FTaskPipeline::Execute()
+
+void UE::Interchange::FTaskPipeline::DoTask(ENamedThreads::Type CurrentThread, const FGraphEventRef& MyCompletionGraphEvent)
 {
 	TRACE_CPUPROFILER_EVENT_SCOPE(UE::Interchange::FTaskPipeline::DoTask)
 #if INTERCHANGE_TRACE_ASYNCHRONOUS_TASK_ENABLED
@@ -57,16 +59,17 @@ void UE::Interchange::FTaskPipeline::Execute()
 	}
 }
 
-void UE::Interchange::FTaskWaitAssetCompilation_GameThread::Execute()
+void UE::Interchange::FTaskWaitAssetCompilation::DoTask(ENamedThreads::Type CurrentThread, const FGraphEventRef& MyCompletionGraphEvent)
 {
-	TRACE_CPUPROFILER_EVENT_SCOPE(UE::Interchange::FTaskWaitAssetCompilation_GameThread::DoTask)
+	TRACE_CPUPROFILER_EVENT_SCOPE(UE::Interchange::FTaskWaitAssetCompilation::DoTask)
 #if INTERCHANGE_TRACE_ASYNCHRONOUS_TASK_ENABLED
-	INTERCHANGE_TRACE_ASYNCHRONOUS_TASK(WaitAssetCompilation)
+		INTERCHANGE_TRACE_ASYNCHRONOUS_TASK(WaitAssetCompilation)
 #endif
+
 	LLM_SCOPE_BYNAME(TEXT("Interchange"));
 
 #if WITH_EDITOR
-	check(IsInGameThread());
+
 	TSharedPtr<FImportAsyncHelper, ESPMode::ThreadSafe> AsyncHelper = WeakAsyncHelper.Pin();
 	if (!ensure(AsyncHelper.IsValid()) || AsyncHelper->bCancel)
 	{
@@ -92,45 +95,74 @@ void UE::Interchange::FTaskWaitAssetCompilation_GameThread::Execute()
 	AsyncHelper->IterateImportedAssets(SourceIndex, FillImportedObjectsFromSource);
 	AsyncHelper->IterateImportedSceneObjects(SourceIndex, FillImportedObjectsFromSource);
 
-	bool bCompilationFinish = true;
-	for (int32 ObjectIndex = 0; ObjectIndex < ImportedObjects.Num(); ++ObjectIndex)
-	{
-		UObject* ImportObject = ImportedObjects[ObjectIndex];
-		if (UMaterialInterface* MaterialInterface = Cast<UMaterialInterface>(ImportObject))
+	//Make sure all assets compilation are done before calling the pipeline post import task, let other thread execute if assets are not compile yet and wait 50ms before a new query
+	FPlatformProcess::ConditionalSleep([&ImportedObjects]()
 		{
-			if (MaterialInterface->IsCompiling())
+			LLM_SCOPE_BYNAME(TEXT("Interchange"));
+			//Compilation status cannot be ask in async thread, query the compile status on the main thread with a small fast function
+			//This ensure we don't stall the main thread until all assets are compile.
+			bool bCompilationFinish = false;
+
+			auto VerifyAssetCompiling = [&bCompilationFinish, &ImportedObjects]()
+				{
+					TRACE_CPUPROFILER_EVENT_SCOPE(UE::Interchange::FTaskWaitAssetCompilation::DoTask::IsCompilingLambda_GameThread);
+					LLM_SCOPE_BYNAME(TEXT("Interchange"));
+					//Make sure all asset compiling managers are up to date, In case the game thread is waiting for the import to finish (like automation test or synchronous import)
+					FAssetCompilingManager::Get().ProcessAsyncTasks();
+
+					bCompilationFinish = true;
+					for (int32 ObjectIndex = 0; ObjectIndex < ImportedObjects.Num(); ++ObjectIndex)
+					{
+						UObject* ImportObject = ImportedObjects[ObjectIndex];
+						if (UMaterialInterface* MaterialInterface = Cast<UMaterialInterface>(ImportObject))
+						{
+							if (MaterialInterface->IsCompiling())
+							{
+								bCompilationFinish = false;
+								break;
+							}
+						}
+						if (IInterface_AsyncCompilation* AssetCompilationInterface = Cast<IInterface_AsyncCompilation>(ImportObject))
+						{
+							if (AssetCompilationInterface->IsCompiling())
+							{
+								bCompilationFinish = false;
+								break;
+							}
+						}
+					}
+				};
+
+			if (IsInGameThread())
 			{
-				bCompilationFinish = false;
-				break;
+				VerifyAssetCompiling();
 			}
-		}
-		if (IInterface_AsyncCompilation* AssetCompilationInterface = Cast<IInterface_AsyncCompilation>(ImportObject))
-		{
-			if (AssetCompilationInterface->IsCompiling())
+			else
 			{
-				bCompilationFinish = false;
-				break;
+				Async(EAsyncExecution::TaskGraphMainThread, [&VerifyAssetCompiling]()
+					{
+						VerifyAssetCompiling();
+					}).Wait();
 			}
-		}
-	}
-
-	if (!bCompilationFinish)
-	{
-		//re-enqueue this task so it doesn't execute the subsequent tasks
-		SetTaskStatus(EInterchangeTaskStatus::Waiting);
-	}
+			return bCompilationFinish;
+		}, 0.05f);
 #endif //WITH_EDITOR
 }
 
-void UE::Interchange::FTaskPostImport_GameThread::Execute()
+void UE::Interchange::FTaskPostImport::DoTask(ENamedThreads::Type CurrentThread, const FGraphEventRef& MyCompletionGraphEvent)
 {
-	TRACE_CPUPROFILER_EVENT_SCOPE(UE::Interchange::FTaskPostImport_GameThread::DoTask)
+	TRACE_CPUPROFILER_EVENT_SCOPE(UE::Interchange::FTaskPostImport::DoTask)
 #if INTERCHANGE_TRACE_ASYNCHRONOUS_TASK_ENABLED
 	INTERCHANGE_TRACE_ASYNCHRONOUS_TASK(PipelinePostImport)
 #endif
 
 	LLM_SCOPE_BYNAME(TEXT("Interchange"));
-	check(IsInGameThread());
+
+	TOptional<FGCScopeGuard> GCScopeGuard;
+	if (!IsInGameThread())
+	{
+		GCScopeGuard.Emplace();
+	}
 
 	TSharedPtr<FImportAsyncHelper, ESPMode::ThreadSafe> AsyncHelper = WeakAsyncHelper.Pin();
 	if (!ensure(AsyncHelper.IsValid()) || AsyncHelper->bCancel)
diff --git a/Engine/Source/Runtime/Interchange/Engine/Private/Tasks/InterchangeTaskPipeline.h b/Engine/Source/Runtime/Interchange/Engine/Private/Tasks/InterchangeTaskPipeline.h
--- a/Engine/Source/Runtime/Interchange/Engine/Private/Tasks/InterchangeTaskPipeline.h
+++ b/Engine/Source/Runtime/Interchange/Engine/Private/Tasks/InterchangeTaskPipeline.h
@@ -2,10 +2,10 @@
 
 #pragma once
 
+#include "Async/TaskGraphInterfaces.h"
 #include "CoreMinimal.h"
 #include "InterchangeManager.h"
 #include "InterchangePipelineBase.h"
-#include "InterchangeTaskSystem.h"
 #include "Stats/Stats.h"
 #include "UObject/WeakObjectPtrTemplates.h"
 
@@ -14,7 +14,7 @@ namespace UE
 	namespace Interchange
 	{
 
-		class FTaskPipeline : public FInterchangeTaskBase
+		class FTaskPipeline
 		{
 		private:
 			TWeakObjectPtr<UInterchangePipelineBase> PipelineBase;
@@ -26,73 +26,114 @@ namespace UE
 			{
 			}
 
-			virtual EInterchangeTaskThread GetTaskThread() const override
+			ENamedThreads::Type GetDesiredThread()
 			{
 				TSharedPtr<FImportAsyncHelper, ESPMode::ThreadSafe> AsyncHelper = WeakAsyncHelper.Pin();
 				if (AsyncHelper.IsValid() && AsyncHelper->bRunSynchronous)
 				{
-					return EInterchangeTaskThread::GameThread;
+					return ENamedThreads::GameThread_Local;
 				}
 
 				if (!ensure(PipelineBase.IsValid()))
 				{
-					return EInterchangeTaskThread::GameThread;
+					return ENamedThreads::GameThread;
 				}
-
+				
 				//Scripted (python) cannot run outside of the game thread, it will lock forever if we do this
 				if (PipelineBase.Get()->IsScripted())
 				{
-					return EInterchangeTaskThread::GameThread;
+					return ENamedThreads::GameThread;
 				}
 
-				return PipelineBase.Get()->CanExecuteOnAnyThread(EInterchangePipelineTask::PostTranslator) ? EInterchangeTaskThread::AsyncThread : EInterchangeTaskThread::GameThread;
+				return PipelineBase.Get()->CanExecuteOnAnyThread(EInterchangePipelineTask::PostTranslator) ? ENamedThreads::AnyBackgroundThreadNormalTask : ENamedThreads::GameThread;
+			}
+
+			static ESubsequentsMode::Type GetSubsequentsMode()
+			{
+				return ESubsequentsMode::TrackSubsequents;
+			}
+
+			FORCEINLINE TStatId GetStatId() const
+			{
+				RETURN_QUICK_DECLARE_CYCLE_STAT(FTaskPipeline, STATGROUP_TaskGraphTasks);
 			}
 
-			virtual void Execute() override;
+			void DoTask(ENamedThreads::Type CurrentThread, const FGraphEventRef& MyCompletionGraphEvent);
 		};
 
-		//We want to be sure any asset compilation is finish before calling subsequents tasks
-		class FTaskWaitAssetCompilation_GameThread : public FInterchangeTaskBase
+		//We want to be sure any asset compilation is finish before calling FTaskPostImport, we use a async task to wait until they are done
+		class FTaskWaitAssetCompilation
 		{
 		private:
 			int32 SourceIndex;
 			TWeakPtr<FImportAsyncHelper, ESPMode::ThreadSafe> WeakAsyncHelper;
 
 		public:
-			FTaskWaitAssetCompilation_GameThread(int32 InSourceIndex, TWeakPtr<FImportAsyncHelper, ESPMode::ThreadSafe> InAsyncHelper)
+			FTaskWaitAssetCompilation(int32 InSourceIndex, TWeakPtr<FImportAsyncHelper, ESPMode::ThreadSafe> InAsyncHelper)
 				: SourceIndex(InSourceIndex)
 				, WeakAsyncHelper(InAsyncHelper)
 			{
 			}
 
-			virtual EInterchangeTaskThread GetTaskThread() const override
+			ENamedThreads::Type GetDesiredThread()
+			{
+				TSharedPtr<FImportAsyncHelper, ESPMode::ThreadSafe> AsyncHelper = WeakAsyncHelper.Pin();
+				if (AsyncHelper.IsValid() && AsyncHelper->bRunSynchronous)
+				{
+					return ENamedThreads::GameThread_Local;
+				}
+
+				return ENamedThreads::AnyBackgroundThreadNormalTask;
+			}
+
+			static ESubsequentsMode::Type GetSubsequentsMode()
 			{
-				//This task is re-enqueue and don't stall the game thread
-				return EInterchangeTaskThread::GameThread;
+				return ESubsequentsMode::TrackSubsequents;
 			}
 
-			virtual void Execute() override;
+			FORCEINLINE TStatId GetStatId() const
+			{
+				RETURN_QUICK_DECLARE_CYCLE_STAT(FTaskWaitAssetCompilation, STATGROUP_TaskGraphTasks);
+			}
+
+			void DoTask(ENamedThreads::Type CurrentThread, const FGraphEventRef& MyCompletionGraphEvent);
 		};
 
-		class FTaskPostImport_GameThread : public FInterchangeTaskBase
+		class FTaskPostImport
 		{
 		private:
 			int32 SourceIndex;
 			TWeakPtr<FImportAsyncHelper, ESPMode::ThreadSafe> WeakAsyncHelper;
 
 		public:
-			FTaskPostImport_GameThread(int32 InSourceIndex, TWeakPtr<FImportAsyncHelper, ESPMode::ThreadSafe> InAsyncHelper)
+			FTaskPostImport(int32 InSourceIndex, TWeakPtr<FImportAsyncHelper, ESPMode::ThreadSafe> InAsyncHelper)
 				: SourceIndex(InSourceIndex)
 				, WeakAsyncHelper(InAsyncHelper)
 			{
 			}
 
-			virtual EInterchangeTaskThread GetTaskThread() const override
+			ENamedThreads::Type GetDesiredThread()
+			{
+				TSharedPtr<FImportAsyncHelper, ESPMode::ThreadSafe> AsyncHelper = WeakAsyncHelper.Pin();
+				if (AsyncHelper.IsValid() && AsyncHelper->bRunSynchronous)
+				{
+					return ENamedThreads::GameThread_Local;
+				}
+
+				return ENamedThreads::GameThread;
+			}
+
+			static ESubsequentsMode::Type GetSubsequentsMode()
+			{
+				return ESubsequentsMode::TrackSubsequents;
+			}
+
+			FORCEINLINE TStatId GetStatId() const
 			{
-				return EInterchangeTaskThread::GameThread;
+				RETURN_QUICK_DECLARE_CYCLE_STAT(FTaskPostImport, STATGROUP_TaskGraphTasks);
 			}
 
-			virtual void Execute() override;
+			void DoTask(ENamedThreads::Type CurrentThread, const FGraphEventRef& MyCompletionGraphEvent);
 		};
 	} //ns Interchange
 }//ns UE
diff --git a/Engine/Source/Runtime/Interchange/Engine/Private/Tasks/InterchangeTaskTranslator.cpp b/Engine/Source/Runtime/Interchange/Engine/Private/Tasks/InterchangeTaskTranslator.cpp
--- a/Engine/Source/Runtime/Interchange/Engine/Private/Tasks/InterchangeTaskTranslator.cpp
+++ b/Engine/Source/Runtime/Interchange/Engine/Private/Tasks/InterchangeTaskTranslator.cpp
@@ -1,11 +1,11 @@
 // Copyright Epic Games, Inc. All Rights Reserved.
 #include "InterchangeTaskTranslator.h"
 
+#include "Async/TaskGraphInterfaces.h"
 #include "CoreMinimal.h"
 #include "InterchangeEngineLogPrivate.h"
 #include "InterchangeManager.h"
 #include "InterchangeSourceData.h"
-#include "InterchangeTaskSystem.h"
 #include "InterchangeTranslatorBase.h"
 #include "Nodes/InterchangeBaseNodeContainer.h"
 #include "Stats/Stats.h"
@@ -13,19 +13,13 @@
 #include "UObject/GarbageCollection.h"
 #include "UObject/WeakObjectPtrTemplates.h"
 
-void UE::Interchange::FTaskTranslator::Execute()
+void UE::Interchange::FTaskTranslator::DoTask(ENamedThreads::Type CurrentThread, const FGraphEventRef& MyCompletionGraphEvent)
 {
 	TRACE_CPUPROFILER_EVENT_SCOPE(UE::Interchange::FTaskTranslator::DoTask)
 #if INTERCHANGE_TRACE_ASYNCHRONOUS_TASK_ENABLED
-	INTERCHANGE_TRACE_ASYNCHRONOUS_TASK(TaskTranslator)
+	INTERCHANGE_TRACE_ASYNCHRONOUS_TASK(Translator)
 #endif
-
-	
-	TOptional<FGCScopeGuard> GCScopeGuard;
-	if (!IsInGameThread())
-	{
-		GCScopeGuard.Emplace();
-	}
+	FGCScopeGuard GCScopeGuard;
 
 	TSharedPtr<FImportAsyncHelper, ESPMode::ThreadSafe> AsyncHelper = WeakAsyncHelper.Pin();
 	check(AsyncHelper.IsValid());
diff --git a/Engine/Source/Runtime/Interchange/Engine/Private/Tasks/InterchangeTaskTranslator.h b/Engine/Source/Runtime/Interchange/Engine/Private/Tasks/InterchangeTaskTranslator.h
--- a/Engine/Source/Runtime/Interchange/Engine/Private/Tasks/InterchangeTaskTranslator.h
+++ b/Engine/Source/Runtime/Interchange/Engine/Private/Tasks/InterchangeTaskTranslator.h
@@ -2,9 +2,9 @@
 
 #pragma once
 
+#include "Async/TaskGraphInterfaces.h"
 #include "CoreMinimal.h"
 #include "InterchangeManager.h"
-#include "InterchangeTaskSystem.h"
 #include "Stats/Stats.h"
 #include "UObject/WeakObjectPtrTemplates.h"
 
@@ -12,7 +12,8 @@ namespace UE
 {
 	namespace Interchange
 	{
-		class FTaskTranslator : public FInterchangeTaskBase
+
+		class FTaskTranslator
 		{
 		private:
 			int32 SourceIndex = INDEX_NONE;
@@ -25,18 +26,29 @@ namespace UE
 			{
 			}
 
-			virtual EInterchangeTaskThread GetTaskThread() const override
+			FORCEINLINE ENamedThreads::Type GetDesiredThread()
 			{
 				TSharedPtr<FImportAsyncHelper, ESPMode::ThreadSafe> AsyncHelper = WeakAsyncHelper.Pin();
 				if (AsyncHelper.IsValid() && AsyncHelper->bRunSynchronous)
 				{
-					return EInterchangeTaskThread::GameThread;
+					return ENamedThreads::GameThread_Local;
 				}
 
-				return EInterchangeTaskThread::AsyncThread;
+				return ENamedThreads::AnyBackgroundThreadNormalTask;
+			}
+			static FORCEINLINE ESubsequentsMode::Type GetSubsequentsMode()
+			{
+				return ESubsequentsMode::TrackSubsequents;
+			}
+
+			FORCEINLINE TStatId GetStatId() const
+			{
+				RETURN_QUICK_DECLARE_CYCLE_STAT(FTaskTranslator, STATGROUP_TaskGraphTasks);
 			}
 
-			virtual void Execute();
+			void DoTask(ENamedThreads::Type CurrentThread, const FGraphEventRef& MyCompletionGraphEvent);
 		};
+
+
 	} //ns Interchange
 } //ns UE
diff --git a/Engine/Source/Runtime/Interchange/Engine/Public/InterchangeManager.h b/Engine/Source/Runtime/Interchange/Engine/Public/InterchangeManager.h
--- a/Engine/Source/Runtime/Interchange/Engine/Public/InterchangeManager.h
+++ b/Engine/Source/Runtime/Interchange/Engine/Public/InterchangeManager.h
@@ -17,7 +17,6 @@
 #include "InterchangePipelineConfigurationBase.h"
 #include "InterchangeResultsContainer.h"
 #include "InterchangeSourceData.h"
-#include "InterchangeTaskSystem.h"
 #include "InterchangeTranslatorBase.h"
 #include "InterchangeWriterBase.h"
 #include "Nodes/InterchangeBaseNodeContainer.h"
@@ -173,9 +172,6 @@ namespace UE
 			// Callback when the status switches to done.
 			INTERCHANGEENGINE_API void OnDone(TFunction< void(FImportResult&) > Callback);
 
-			//Set the async helper that own this import result
-			INTERCHANGEENGINE_API void SetAsyncHelper(TWeakPtr<class FImportAsyncHelper> InAsyncHelper);
-
 			// Internal delegates. To set these, use the FImportAssetParameters when calling the Interchange import functions.
 			FOnObjectImportDoneDynamic OnObjectDone;
 			FOnObjectImportDoneNative OnObjectDoneNative;
@@ -198,7 +194,7 @@ namespace UE
 			mutable FRWLock ImportedObjectsRWLock;
 			TObjectPtr<UInterchangeResultsContainer> Results;
 
-			TWeakPtr<class FImportAsyncHelper> AsyncHelper = nullptr;
+			FGraphEventRef GraphEvent; // WaitUntilDone waits for this event to be triggered.
 
 			TFunction< void(FImportResult&) > DoneCallback;
 		};
@@ -241,18 +237,18 @@ namespace UE
 			//Python class instanced assets cannot be saved, so we have to serialize in JSON the data to restore it when we do a reimport.
 			TArray<UObject*> OriginalPipelines;
 
-			TArray<uint64> TranslatorTasks;
-			TArray<uint64> PipelineTasks;
-			TArray<uint64> WaitAssetCompilationTasks;
-			TArray<uint64> PostImportTasks;
-			uint64 ParsingTask;
-			TArray<uint64> BeginImportObjectTasks;
-			TArray<uint64> ImportObjectTasks;
-			TArray<uint64> FinalizeImportObjectTasks;
-			TArray<uint64> SceneTasks;
+			TArray<FGraphEventRef> TranslatorTasks;
+			TArray<FGraphEventRef> PipelineTasks;
+			FGraphEventRef WaitAssetCompilationTask;
+			TArray<FGraphEventRef> PostImportTasks;
+			FGraphEventRef ParsingTask;
+			TArray<FGraphEventRef> BeginImportObjectTasks;
+			TArray<FGraphEventRef> ImportObjectTasks;
+			TArray<FGraphEventRef> FinalizeImportObjectTasks;
+			TArray<FGraphEventRef> SceneTasks;
 
-			uint64 PreCompletionTask;
-			uint64 CompletionTask;
+			FGraphEventRef PreCompletionTask;
+			FGraphEventRef CompletionTask;
 
 			//Return true if we can import this class, or false otherwise.
 			bool IsClassImportAllowed(UClass* Class);
@@ -301,7 +297,7 @@ namespace UE
 			/**
 			 * Wait synchronously after the graph parsing task is done, and return the GraphEventArray up to the completion TaskGraphEvent.
 			 */
-			TArray<uint64> GetCompletionTaskGraphEvent();
+			FGraphEventArray GetCompletionTaskGraphEvent();
 
 			void InitCancel();
 
diff --git a/Engine/Source/Runtime/Interchange/Engine/Public/InterchangeTaskSystem.h b/Engine/Source/Runtime/Interchange/Engine/Public/InterchangeTaskSystem.h
deleted file mode 100644
