Interchange:
- Fix actors not correctly filtered out

diff --git a/Engine/Source/Runtime/Interchange/Engine/Private/Tasks/InterchangeTaskCreateSceneObjects.cpp b/Engine/Source/Runtime/Interchange/Engine/Private/Tasks/InterchangeTaskCreateSceneObjects.cpp
--- a/Engine/Source/Runtime/Interchange/Engine/Private/Tasks/InterchangeTaskCreateSceneObjects.cpp
+++ b/Engine/Source/Runtime/Interchange/Engine/Private/Tasks/InterchangeTaskCreateSceneObjects.cpp
@@ -60,7 +60,8 @@ void UE::Interchange::FTaskCreateSceneObjects_GameThread::Execute()
 
 	for (UInterchangeFactoryBaseNode* FactoryNode : FactoryNodes)
 	{
-		if (!FactoryNode)
+		// Don't create actors on a first import if they're disabled
+		if (!FactoryNode || !FactoryNode->IsEnabled())
 		{
 			continue;
 		}
