[BUGFIX] Interchange - Canceling the import on the way isn't working and isn't recorded in the output logs.

diff --git a/Engine/Source/Runtime/Interchange/Engine/Private/InterchangeManager.cpp b/Engine/Source/Runtime/Interchange/Engine/Private/InterchangeManager.cpp
--- a/Engine/Source/Runtime/Interchange/Engine/Private/InterchangeManager.cpp
+++ b/Engine/Source/Runtime/Interchange/Engine/Private/InterchangeManager.cpp
@@ -2470,6 +2470,7 @@ void UInterchangeManager::ReleaseAsyncHelper(TWeakPtr<UE::Interchange::FImportAs
 			{
 				TitleText = NSLOCTEXT("Interchange", "Asynchronous_import_canceled", "Import Canceled");
 				bSucceeded = true; // Mark the "cancelation" as a success so that the notification goes away
+				Notification->SetProgressText(FText::FromString("Operation was canceled by user."));
 			}
 			else
 			{
@@ -2481,8 +2482,10 @@ void UInterchangeManager::ReleaseAsyncHelper(TWeakPtr<UE::Interchange::FImportAs
 				{
 					TitleText = NSLOCTEXT("Interchange", "Asynchronous_import_failed", "Import Failed");
 				}
-			}
 
+				Notification->SetProgressText(FText::FromString("Operation completed."));
+			}
+									
 			Notification->SetComplete(TitleText, FText::GetEmpty(), bSucceeded);
 			Notification = nullptr; //This should delete the notification
 		}
@@ -2697,6 +2700,7 @@ void UInterchangeManager::CancelAllTasks()
 			QueuedTaskData.AsyncHelper->InitCancel();
 		}
 	}
+	bImportCanceled = true;
 	//Tasks should all finish quite fast now
 };
 
