[ZenLoader/Misc]
 - The loader is the sole owner of the AsyncLoading flag and it should never be removed by other systems

#rnx

diff --git a/Engine/Source/Runtime/Engine/Private/MorphTools.cpp b/Engine/Source/Runtime/Engine/Private/MorphTools.cpp
--- a/Engine/Source/Runtime/Engine/Private/MorphTools.cpp
+++ b/Engine/Source/Runtime/Engine/Private/MorphTools.cpp
@@ -390,7 +390,7 @@ void FFinishBuildMorphTargetData::ApplyEditorData(USkeletalMesh * SkeletalMesh,
 	//Clear any async flags after the morphtargets have been set to the skeletalmesh
 	for (UMorphTarget* MorphTarget : SkeletalMesh->GetMorphTargets())
 	{
-		const EInternalObjectFlags AsyncFlags = EInternalObjectFlags::Async | EInternalObjectFlags::AsyncLoading;
+		const EInternalObjectFlags AsyncFlags = EInternalObjectFlags::Async;
 		MorphTarget->ClearInternalFlags(AsyncFlags);
 	}
 
diff --git a/Engine/Source/Runtime/Interchange/Engine/Private/Tasks/InterchangeTaskCompletion.cpp b/Engine/Source/Runtime/Interchange/Engine/Private/Tasks/InterchangeTaskCompletion.cpp
--- a/Engine/Source/Runtime/Interchange/Engine/Private/Tasks/InterchangeTaskCompletion.cpp
+++ b/Engine/Source/Runtime/Interchange/Engine/Private/Tasks/InterchangeTaskCompletion.cpp
@@ -90,7 +90,7 @@ void UE::Interchange::FTaskPreCompletion_GameThread::Execute()
 			Message->AssetType = ImportedObject->GetClass();
 
 			//Clear any async flag from the created asset and all its subobjects
-			const EInternalObjectFlags AsyncFlags = EInternalObjectFlags::Async | EInternalObjectFlags::AsyncLoading;
+			const EInternalObjectFlags AsyncFlags = EInternalObjectFlags::Async;
 			ImportedObject->ClearInternalFlags(AsyncFlags);
 
 			TArray<UObject*> ImportedSubobjects;
