[Interchange] Prevent user to pick wrong file extension on reimport

diff --git a/Engine/Source/Editor/UnrealEd/Private/Editor.cpp b/Engine/Source/Editor/UnrealEd/Private/Editor.cpp
--- a/Engine/Source/Editor/UnrealEd/Private/Editor.cpp
+++ b/Engine/Source/Editor/UnrealEd/Private/Editor.cpp
@@ -859,7 +859,7 @@ void FReimportManager::GetNewReimportPath(UObject* Obj, TArray<FString>& InOutFi
 	if (UInterchangeManager::IsInterchangeImportEnabled())
 	{
 		//Get the extension interchange can translate for this object
-		TArray<FString> TranslatorFormats = UInterchangeManager::GetInterchangeManager().GetSupportedFormatsForObject(Obj);
+		TArray<FString> TranslatorFormats = UInterchangeManager::GetInterchangeManager().GetSupportedFormatsForObject(Obj, RealSourceFileIndex);
 		ObjectTools::AppendFormatsFileExtensions(TranslatorFormats, FileTypes, AllExtensions, DummyFilterIndexToFactory);
 	}
 
diff --git a/Engine/Source/Runtime/Interchange/Engine/Private/InterchangeManager.cpp b/Engine/Source/Runtime/Interchange/Engine/Private/InterchangeManager.cpp
--- a/Engine/Source/Runtime/Interchange/Engine/Private/InterchangeManager.cpp
+++ b/Engine/Source/Runtime/Interchange/Engine/Private/InterchangeManager.cpp
@@ -1189,7 +1189,7 @@ TArray<FString> UInterchangeManager::GetSupportedFormats(const EInterchangeTrans
 	return FileExtensions;
 }
 
-TArray<FString> UInterchangeManager::GetSupportedAssetTypeFormats(const EInterchangeTranslatorAssetType ForTranslatorAssetType) const
+TArray<FString> UInterchangeManager::GetSupportedAssetTypeFormats(const EInterchangeTranslatorAssetType ForTranslatorAssetType, const EInterchangeTranslatorType ForTranslatorType /*= EInterchangeTranslatorType::Invalid*/, bool bStrictMatchTranslatorType /*= false*/) const
 {
 	TArray<FString> FileExtensions;
 	if (!IsInterchangeImportEnabled())
@@ -1203,14 +1203,25 @@ TArray<FString> UInterchangeManager::GetSupportedAssetTypeFormats(const EInterch
 
 		if (TranslatorBaseCDO->DoesSupportAssetType(ForTranslatorAssetType))
 		{
-			FileExtensions.Append(TranslatorBaseCDO->GetSupportedFormats());
+			bool bAddTranslatorSupportedFormats = true;
+			if (ForTranslatorType != EInterchangeTranslatorType::Invalid)
+			{
+				bAddTranslatorSupportedFormats = bStrictMatchTranslatorType
+					? EnumHasAllFlags(TranslatorBaseCDO->GetTranslatorType(), ForTranslatorType)
+					: EnumHasAnyFlags(TranslatorBaseCDO->GetTranslatorType(), ForTranslatorType);
+			}
+
+			if (bAddTranslatorSupportedFormats)
+			{
+				FileExtensions.Append(TranslatorBaseCDO->GetSupportedFormats());
+			}
 		}
 	}
 
 	return FileExtensions;
 }
 
-TArray<FString> UInterchangeManager::GetSupportedFormatsForObject(const UObject* Object) const
+TArray<FString> UInterchangeManager::GetSupportedFormatsForObject(const UObject* Object, int32 SourceFileIndex) const
 {
 	TArray<FString> FileExtensions;
 	if (!IsInterchangeImportEnabled())
@@ -1225,34 +1236,48 @@ TArray<FString> UInterchangeManager::GetSupportedFormatsForObject(const UObject*
 	}
 
 	UInterchangeFactoryBase* Factory = RegisteredFactoryClass->GetDefaultObject<UInterchangeFactoryBase>();
-	TArray<FString> TempFilenames;
+	TArray<FString> SourceFilenames;
 	//GetSourceFilenames verify we have a valid UInterchangeAssetImportData for this Object
 	//This ensure we do not allow re-import
-	if (!Factory->GetSourceFilenames(Object, TempFilenames))
+	if (!Factory->GetSourceFilenames(Object, SourceFilenames))
 	{
 		return FileExtensions;
 	}
 
-	switch (Factory->GetFactoryAssetType())
+	// SourceFileIndex is ensured to be a valid array index.
+	if (SourceFileIndex < SourceFilenames.Num())
 	{
-	case EInterchangeFactoryAssetType::Animations:
-		FileExtensions = GetSupportedAssetTypeFormats(EInterchangeTranslatorAssetType::Animations);
-		break;
-	case EInterchangeFactoryAssetType::Materials:
-		FileExtensions = GetSupportedAssetTypeFormats(EInterchangeTranslatorAssetType::Materials);
-		break;
-	case EInterchangeFactoryAssetType::Meshes:
-	case EInterchangeFactoryAssetType::Physics:
-		FileExtensions = GetSupportedAssetTypeFormats(EInterchangeTranslatorAssetType::Meshes);
-		break;
-	case EInterchangeFactoryAssetType::Textures:
-		FileExtensions = GetSupportedAssetTypeFormats(EInterchangeTranslatorAssetType::Textures);
-		break;
-	case EInterchangeFactoryAssetType::None: //Actor factories return None
-		FileExtensions = GetSupportedFormats(EInterchangeTranslatorType::Actors);
-		break;
-	}
+		UE::Interchange::FScopedSourceData SourceData(SourceFilenames[SourceFileIndex]);
+		{
+			UE::Interchange::FScopedTranslator ScopedTranslator(SourceData.GetSourceData());
+			if (UInterchangeTranslatorBase* Translator = ScopedTranslator.GetTranslator())
+			{
+				const EInterchangeTranslatorType TranslatorType = Translator->GetTranslatorType();
+				constexpr bool bStrictMatchTranslatorType = false;
 
+				switch (Factory->GetFactoryAssetType())
+				{
+				case EInterchangeFactoryAssetType::Animations:
+					FileExtensions = GetSupportedAssetTypeFormats(EInterchangeTranslatorAssetType::Animations, TranslatorType, bStrictMatchTranslatorType);
+					break;
+				case EInterchangeFactoryAssetType::Materials:
+					FileExtensions = GetSupportedAssetTypeFormats(EInterchangeTranslatorAssetType::Materials, TranslatorType, bStrictMatchTranslatorType);
+					break;
+				case EInterchangeFactoryAssetType::Meshes:
+				case EInterchangeFactoryAssetType::Physics:
+					FileExtensions = GetSupportedAssetTypeFormats(EInterchangeTranslatorAssetType::Meshes, TranslatorType, bStrictMatchTranslatorType);
+					break;
+				case EInterchangeFactoryAssetType::Textures:
+					FileExtensions = GetSupportedAssetTypeFormats(EInterchangeTranslatorAssetType::Textures, TranslatorType, bStrictMatchTranslatorType);
+					break;
+				case EInterchangeFactoryAssetType::None: //Actor factories return None
+					FileExtensions = GetSupportedFormats(EInterchangeTranslatorType::Actors);
+					break;
+				}
+			}
+		}
+	}
+	
 	//Make sure we return lower case extensions
 	for (FString& Extension : FileExtensions)
 	{
diff --git a/Engine/Source/Runtime/Interchange/Engine/Public/InterchangeManager.h b/Engine/Source/Runtime/Interchange/Engine/Public/InterchangeManager.h
--- a/Engine/Source/Runtime/Interchange/Engine/Public/InterchangeManager.h
+++ b/Engine/Source/Runtime/Interchange/Engine/Public/InterchangeManager.h
@@ -548,12 +548,13 @@ public:
 	/**
 	 * Returns the list of formats supporting the specified translator asset type.
 	 */
-	INTERCHANGEENGINE_API TArray<FString> GetSupportedAssetTypeFormats(const EInterchangeTranslatorAssetType ForTranslatorAssetType) const;
+	INTERCHANGEENGINE_API TArray<FString> GetSupportedAssetTypeFormats(const EInterchangeTranslatorAssetType ForTranslatorAssetType, const EInterchangeTranslatorType ForTranslatorType = EInterchangeTranslatorType::Invalid, bool bStrictMatchTranslatorType = false) const;
 
 	/**
 	 * Returns the list of supported formats for a given Object.
+	 * 
 	 */
-	INTERCHANGEENGINE_API TArray<FString> GetSupportedFormatsForObject(const UObject* Object) const;
+	INTERCHANGEENGINE_API TArray<FString> GetSupportedFormatsForObject(const UObject* Object, int32 SourceFileIndex) const;
 
 	/**
 	 * Check whether there is a registered translator for this source data.
