Interchange Telemetry Enhancement

- Sending the warning and error messages even if empty
- Array to string serialization using JSON
- Custom PrettyPrintJsonPolicy that will not print line endings or tabs

diff --git a/Engine/Source/Runtime/Interchange/Engine/Private/InterchangeManager.cpp b/Engine/Source/Runtime/Interchange/Engine/Private/InterchangeManager.cpp
--- a/Engine/Source/Runtime/Interchange/Engine/Private/InterchangeManager.cpp
+++ b/Engine/Source/Runtime/Interchange/Engine/Private/InterchangeManager.cpp
@@ -49,6 +49,10 @@
 #include "UObject/WeakObjectPtrTemplates.h"
 #include "Widgets/Notifications/SNotificationList.h"
 
+#include "Serialization/JsonSerializer.h"
+#include "Serialization/JsonWriter.h"
+#include "Serialization/JsonSerializerWriter.h"
+
 #include UE_INLINE_GENERATED_CPP_BY_NAME(InterchangeManager)
 
 static bool GInterchangeImportEnable = false;
@@ -178,6 +182,26 @@ namespace UE::Interchange::Private
 		FString EventString = TEXT("Interchange.Usage.Import.Pipeline");
 		FEngineAnalytics::GetProvider().RecordEvent(EventString, PipelineAttribs);
 	}
+
+	template <class CharType>
+	struct TInterchangeJsonPrintPolicy
+		: public TPrettyJsonPrintPolicy<CharType>
+	{
+		static inline void WriteLineTerminator(FArchive* Stream)
+		{
+			// No line terminators
+		}
+
+		static inline void WriteTabs(FArchive* Stream, int32 Count)
+		{
+			// No Tabs
+		}
+
+		static inline void WriteSpace(FArchive* Stream)
+		{
+			TJsonPrintPolicy<CharType>::WriteChar(Stream, CharType(' '));
+		}
+	};
 }
 
 UE::Interchange::FScopedInterchangeImportEnableState::FScopedInterchangeImportEnableState(const bool bScopeValue)
@@ -477,6 +501,28 @@ void UE::Interchange::FImportAsyncHelper::SendAnalyticImportEndData()
 		return;
 	}
 
+	// Helper to create JSON Serialized Array Strings
+	auto AddArrayAttribute = [](TArray<FAnalyticsEventAttribute>& Attributes, const FString& AttributeName, TArray<FString>& Strings) 
+	{
+		using namespace UE::Interchange::Private;
+
+		FString JsonString;
+		TSharedRef < TJsonWriter <TCHAR, TInterchangeJsonPrintPolicy<TCHAR> > > JSONWriter = TJsonStringWriter<TInterchangeJsonPrintPolicy<TCHAR>>::Create(&JsonString);
+		
+		{
+			// A wrapper that serializes array and takes in FJsonSerializableArrays
+			FJsonSerializerWriter< TCHAR, TInterchangeJsonPrintPolicy<TCHAR> > SerializerWriter = FJsonSerializerWriter< TCHAR, TInterchangeJsonPrintPolicy<TCHAR> >(JSONWriter);
+			// Serializes array into JSON Writer's byte buffer
+			SerializerWriter.SerializeArray(Strings);
+		}
+
+		// This writes the bytes to the JsonString variable
+		JSONWriter->Close();
+
+		Attributes.Add(FAnalyticsEventAttribute(AttributeName, JsonString));
+	};
+
+
 	int32 ImportedObjectCount = 0;
 	for (const TPair<int32, TArray<FImportedObjectInfo>>& SourceIndexAndImportedAssets : ImportedAssetsPerSourceIndex)
 	{
@@ -504,14 +550,13 @@ void UE::Interchange::FImportAsyncHelper::SendAnalyticImportEndData()
 			case EInterchangeResultType::Success:
 				break;
 			case EInterchangeResultType::Warning:
-				WarningMessages.Add(TEXT("{") + InterchangeResult->GetText().ToString() + TEXT("}"));
+				WarningMessages.Add(InterchangeResult->GetText().BuildSourceString());
 				break;
 			case EInterchangeResultType::Error:
-				ErrorMessages.Add(TEXT("{") + InterchangeResult->GetText().ToString() + TEXT("}"));
+				ErrorMessages.Add(InterchangeResult->GetText().BuildSourceString());
 				break;
 			}
 		}
-		
 	};
 
 	if (const UInterchangeResultsContainer* ResultContainer = AssetImportResult->GetResults())
@@ -523,15 +568,9 @@ void UE::Interchange::FImportAsyncHelper::SendAnalyticImportEndData()
 		CollectResultContainer(ResultContainer);
 	}
 
-	if (WarningMessages.Num() > 0)
-	{
-		Attribs.Add(FAnalyticsEventAttribute(TEXT("WarningMessages"), WarningMessages));
-	}
-	if (ErrorMessages.Num() > 0)
-	{
-		Attribs.Add(FAnalyticsEventAttribute(TEXT("ErrorMessages"), ErrorMessages));
-	}
-
+	AddArrayAttribute(Attribs, TEXT("WarningMessages"), WarningMessages);
+	AddArrayAttribute(Attribs, TEXT("ErrorMessages"), ErrorMessages);
+	
 	FString EventString = TEXT("Interchange.Usage.ImportResult");
 	FEngineAnalytics::GetProvider().RecordEvent(EventString, Attribs);
 }
