Fix CIS build on mac
#rnx

diff --git a/Engine/Source/Runtime/Interchange/Engine/Public/InterchangeManager.h b/Engine/Source/Runtime/Interchange/Engine/Public/InterchangeManager.h
--- a/Engine/Source/Runtime/Interchange/Engine/Public/InterchangeManager.h
+++ b/Engine/Source/Runtime/Interchange/Engine/Public/InterchangeManager.h
@@ -13,6 +13,7 @@
 #include "HAL/Thread.h"
 #include "HAL/ThreadSafeBool.h"
 #include "InterchangeAssetImportData.h"
+#include "InterchangeEditorUtilitiesBase.h"
 #include "InterchangeFactoryBase.h"
 #include "InterchangePipelineConfigurationBase.h"
 #include "InterchangeResultsContainer.h"
@@ -32,7 +33,6 @@
 
 class FAsyncTaskNotification;
 class UInterchangeBlueprintPipelineBase;
-class UInterchangeEditorUtilitiesBase;
 class UInterchangeFactoryBaseNode;
 class UInterchangePipelineBase;
 class UInterchangePythonPipelineBase;
