Why Import Failed: Add detailed telemetry for failed Interchange import(s)
- Interchange Event data is serialized as JSON and sent as a JSON Fragment to analytics event attribute.
- TCondensedJsonPrintPolicy is used

diff --git a/Engine/Source/Runtime/Interchange/Engine/Private/InterchangeManager.cpp b/Engine/Source/Runtime/Interchange/Engine/Private/InterchangeManager.cpp
--- a/Engine/Source/Runtime/Interchange/Engine/Private/InterchangeManager.cpp
+++ b/Engine/Source/Runtime/Interchange/Engine/Private/InterchangeManager.cpp
@@ -37,6 +37,7 @@
 #include "Misc/DateTime.h"
 #include "Nodes/InterchangeBaseNodeContainer.h"
 #include "PackageUtils/PackageUtils.h"
+#include "Policies/CondensedJsonPrintPolicy.h"
 #include "Serialization/JsonSerializer.h"
 #include "Serialization/JsonWriter.h"
 #include "Serialization/JsonSerializerWriter.h"
@@ -183,25 +184,56 @@ namespace UE::Interchange::Private
 		FEngineAnalytics::GetProvider().RecordEvent(EventString, PipelineAttribs);
 	}
 
-	template <class CharType>
-	struct TInterchangeJsonPrintPolicy
-		: public TPrettyJsonPrintPolicy<CharType>
+	// Json writer subclass to allow us to avoid using a SharedPtr to write basic Json.
+	class FAnalyticsJsonWriter : public TJsonStringWriter<TCondensedJsonPrintPolicy<TCHAR>>
 	{
-		static inline void WriteLineTerminator(FArchive* Stream)
+	public:
+		explicit FAnalyticsJsonWriter(FString* Out) : TJsonStringWriter<TCondensedJsonPrintPolicy<TCHAR>>(Out, 0)
 		{
-			// No line terminators
 		}
+	};
 
-		static inline void WriteTabs(FArchive* Stream, int32 Count)
+	FJsonFragment AnalyticsConvertMapToJsonFragment(const TMap<FString, int32>& InFreqMap)
+	{
+		FString ReturnValue;
+		FAnalyticsJsonWriter JsonWriter(&ReturnValue);
+		JsonWriter.WriteArrayStart();
+		for (const TPair<FString,int32>& FreqMapPair : InFreqMap)
 		{
-			// No Tabs
+			JsonWriter.WriteObjectStart();
+			JsonWriter.WriteValue(TEXT("MessageKey"), FreqMapPair.Key);
+			JsonWriter.WriteValue(TEXT("MessageCount"), FreqMapPair.Value);
+			JsonWriter.WriteObjectEnd();
 		}
+		JsonWriter.WriteArrayEnd();
+		JsonWriter.Close();
+		return FJsonFragment(MoveTemp(ReturnValue));
+	}
 
-		static inline void WriteSpace(FArchive* Stream)
+	bool ExtractNamespace(const FText& Text, FString& OutTextNamespaceId)
+	{
+		FText TextToUse = Text;
+
+		TArray<FHistoricTextFormatData> TextHistory;
+		FTextInspector::GetHistoricFormatData(Text, TextHistory);
+		if (TextHistory.Num() > 0)
 		{
-			TJsonPrintPolicy<CharType>::WriteChar(Stream, CharType(' '));
+			const FHistoricTextFormatData& FmtData = TextHistory[0];
+			TextToUse = FmtData.SourceFmt.GetSourceText();
 		}
-	};
+
+		const TOptional<FString> TextNamespace = FTextInspector::GetNamespace(TextToUse);
+		const TOptional<FString> TextKey = FTextInspector::GetKey(TextToUse);
+
+		if (TextNamespace.IsSet() && TextKey.IsSet())
+		{
+			OutTextNamespaceId =  FString::Printf(TEXT("%s_%s"), *TextNamespace.GetValue(), *TextKey.GetValue());
+			return true;
+		}
+		
+		OutTextNamespaceId = FString(TEXT("UnknownError"));
+		return false;
+	}
 }
 
 UE::Interchange::FScopedInterchangeImportEnableState::FScopedInterchangeImportEnableState(const bool bScopeValue)
@@ -501,28 +533,6 @@ void UE::Interchange::FImportAsyncHelper::SendAnalyticImportEndData()
 		return;
 	}
 
-	// Helper to create JSON Serialized Array Strings
-	auto AddArrayAttribute = [](TArray<FAnalyticsEventAttribute>& Attributes, const FString& AttributeName, TArray<FString>& Strings) 
-	{
-		using namespace UE::Interchange::Private;
-
-		FString JsonString;
-		TSharedRef < TJsonWriter <TCHAR, TInterchangeJsonPrintPolicy<TCHAR> > > JSONWriter = TJsonStringWriter<TInterchangeJsonPrintPolicy<TCHAR>>::Create(&JsonString);
-		
-		{
-			// A wrapper that serializes array and takes in FJsonSerializableArrays
-			FJsonSerializerWriter< TCHAR, TInterchangeJsonPrintPolicy<TCHAR> > SerializerWriter = FJsonSerializerWriter< TCHAR, TInterchangeJsonPrintPolicy<TCHAR> >(JSONWriter);
-			// Serializes array into JSON Writer's byte buffer
-			SerializerWriter.SerializeArray(Strings);
-		}
-
-		// This writes the bytes to the JsonString variable
-		JSONWriter->Close();
-
-		Attributes.Add(FAnalyticsEventAttribute(AttributeName, JsonString));
-	};
-
-
 	int32 ImportedObjectCount = 0;
 	for (const TPair<int32, TArray<FImportedObjectInfo>>& SourceIndexAndImportedAssets : ImportedAssetsPerSourceIndex)
 	{
@@ -537,25 +547,49 @@ void UE::Interchange::FImportAsyncHelper::SendAnalyticImportEndData()
 	Attribs.Add(FAnalyticsEventAttribute(TEXT("ImportObjectCount"), ImportedObjectCount));
 
 	//Report any warning or error message
-	TArray<FString> WarningMessages;
-	TArray<FString> ErrorMessages;
+	TMap<FString, int32> WarningMessages;
+	TMap<FString, int32> ErrorMessages;
 	auto CollectResultContainer = [&WarningMessages, &ErrorMessages](const UInterchangeResultsContainer* ResultContainer)
 	{
 		TArray<UInterchangeResult*> InterchangeResults = ResultContainer->GetResults();
 		
 		for (const UInterchangeResult* InterchangeResult : InterchangeResults)
 		{
+			using namespace UE::Interchange::Private;
+
 			switch (InterchangeResult->GetResultType() )
 			{
 			case EInterchangeResultType::Success:
 				break;
 			case EInterchangeResultType::Warning:
-				WarningMessages.Add(InterchangeResult->GetText().BuildSourceString());
+			{
+				FString OutWarningAttribValue;
+				if (ExtractNamespace(InterchangeResult->GetText(), OutWarningAttribValue))
+				{
+					int32& Frequency = WarningMessages.FindOrAdd(OutWarningAttribValue);
+					Frequency++;
+				}
+				else
+				{
+					UE_LOG(LogInterchangeEngine, Error, TEXT("Failed to extract Analytic Attribute Value from %s"), *(InterchangeResult->GetText().ToString()));
+				}
 				break;
+			}
 			case EInterchangeResultType::Error:
-				ErrorMessages.Add(InterchangeResult->GetText().BuildSourceString());
+			{
+				FString OutErrorAttribValue;
+				if (ExtractNamespace(InterchangeResult->GetText(), OutErrorAttribValue))
+				{
+					int32& Frequency = ErrorMessages.FindOrAdd(OutErrorAttribValue);
+					Frequency++;
+				}
+				else
+				{
+					UE_LOG(LogInterchangeEngine, Error, TEXT("Failed to extract Analytic Attribute Value from %s"), *(InterchangeResult->GetText().ToString()));
+				}
 				break;
 			}
+			}
 		}
 	};
 
@@ -568,9 +602,12 @@ void UE::Interchange::FImportAsyncHelper::SendAnalyticImportEndData()
 		CollectResultContainer(ResultContainer);
 	}
 
-	AddArrayAttribute(Attribs, TEXT("WarningMessages"), WarningMessages);
-	AddArrayAttribute(Attribs, TEXT("ErrorMessages"), ErrorMessages);
-	
+	{
+		using namespace UE::Interchange::Private;
+		Attribs.Add(FAnalyticsEventAttribute(TEXT("WarningMessages"), AnalyticsConvertMapToJsonFragment(WarningMessages)));
+		Attribs.Add(FAnalyticsEventAttribute(TEXT("ErrorMessages"), AnalyticsConvertMapToJsonFragment(ErrorMessages)));
+	}
+
 	FString EventString = TEXT("Interchange.Usage.ImportResult");
 	FEngineAnalytics::GetProvider().RecordEvent(EventString, Attribs);
 }
