Interchange: fix static analysis error.

#rnx

diff --git a/Engine/Source/Runtime/Interchange/Engine/Private/InterchangeManager.cpp b/Engine/Source/Runtime/Interchange/Engine/Private/InterchangeManager.cpp
--- a/Engine/Source/Runtime/Interchange/Engine/Private/InterchangeManager.cpp
+++ b/Engine/Source/Runtime/Interchange/Engine/Private/InterchangeManager.cpp
@@ -1879,21 +1879,32 @@ UInterchangeManager::ImportInternal(const FString& ContentPath, const UInterchan
 	UInterchangeSourceData* DuplicateSourceData = Cast<UInterchangeSourceData>(StaticDuplicateObject(SourceData, GetTransientPackage()));
 	//Array of source data to build one graph per source
 	AsyncHelper->SourceDatas.Add(DuplicateSourceData);
+
+	//Get the first source data translator, we currently do not support more then one source.
 	constexpr int32 SourceIndex = 0;
 	UInterchangeTranslatorBase* AsyncTranslator = nullptr;
-	//Get all the translators for the source datas
+	//Add all source data translator
 	for (int32 SourceDataIndex = 0; SourceDataIndex < AsyncHelper->SourceDatas.Num(); ++SourceDataIndex)
 	{
-		AsyncTranslator = GetTranslatorForSourceData(AsyncHelper->SourceDatas[SourceDataIndex]);
+		UInterchangeTranslatorBase* SourceTranslator = GetTranslatorForSourceData(AsyncHelper->SourceDatas[SourceDataIndex]);
+		if (!ensureMsgf(SourceTranslator, TEXT("Each interchange source data should have a valid translator")))
+		{
+			return EarlyExit();
+		}
 		if (bIsReimport)
 		{
 			//Set translator settings if we are doing a reimport
 			if (const UInterchangeTranslatorSettings* InterchangeTranslatorSettings = OriginalAssetImportData->GetTranslatorSettings())
 			{
-				AsyncTranslator->SetSettings(InterchangeTranslatorSettings);
+				SourceTranslator->SetSettings(InterchangeTranslatorSettings);
 			}
 		}
-		ensure(AsyncHelper->Translators.Add(AsyncTranslator) == SourceDataIndex);
+		//Get the expected source index translator
+		if (SourceDataIndex == SourceIndex)
+		{
+			AsyncTranslator = SourceTranslator;
+		}
+		ensure(AsyncHelper->Translators.Add(SourceTranslator) == SourceDataIndex);
 	}
 
 	//Create the node graphs for each source data (StrongObjectPtr has to be created on the main thread)
