[Interchange] Duplicate actors when re-importing a file with non-unique actors' name fixing.

#rnx

diff --git a/Engine/Source/Runtime/Interchange/Engine/Private/InterchangeSceneImportAsset.cpp b/Engine/Source/Runtime/Interchange/Engine/Private/InterchangeSceneImportAsset.cpp
--- a/Engine/Source/Runtime/Interchange/Engine/Private/InterchangeSceneImportAsset.cpp
+++ b/Engine/Source/Runtime/Interchange/Engine/Private/InterchangeSceneImportAsset.cpp
@@ -204,17 +204,20 @@ UObject* UInterchangeSceneImportAsset::GetSceneObject(const FString& PackageName
 {
 #if WITH_EDITORONLY_DATA
 	const FSoftObjectPath ObjectPath(FName(PackageName), FName(AssetName), SubPathString);
+	const FName ObjectSubPathBaseName = FActorSpawnUtils::GetBaseName(*SubPathString);
 
 	UObject* SceneObject = nullptr;
 	for (const TPair<FSoftObjectPath, FString>& SceneObjectPair : SceneObjects)
 	{
 		const FSoftObjectPath& SceneObjectPath = SceneObjectPair.Key;
+
 		if (SceneObjectPath.GetLongPackageFName() == ObjectPath.GetLongPackageFName()
 			&& SceneObjectPath.GetAssetFName() == ObjectPath.GetAssetFName())
 		{
 			//World partition actor have a guid that we need to remove before comparison
-			FName SceneObjectSubPathFName = FActorSpawnUtils::GetBaseName(*SceneObjectPath.GetSubPathString());
-			if (SceneObjectSubPathFName == *ObjectPath.GetSubPathString())
+			FString SceneObjectSubPathString = SceneObjectPath.GetSubPathString();
+			if (SceneObjectSubPathString.Contains(SubPathString)
+				&& FActorSpawnUtils::GetBaseName(*SceneObjectSubPathString) == ObjectSubPathBaseName)
 			{
 				SceneObject = SceneObjectPath.TryLoad();
 				break;
