[Interchange] Fixing TaskOrdering.

#rnx

diff --git a/Engine/Source/Runtime/Interchange/Engine/Private/Tasks/InterchangeTaskParsing.cpp b/Engine/Source/Runtime/Interchange/Engine/Private/Tasks/InterchangeTaskParsing.cpp
--- a/Engine/Source/Runtime/Interchange/Engine/Private/Tasks/InterchangeTaskParsing.cpp
+++ b/Engine/Source/Runtime/Interchange/Engine/Private/Tasks/InterchangeTaskParsing.cpp
@@ -175,21 +175,14 @@ void UE::Interchange::FTaskParsing::Execute()
 					[&BaseNodeContainer, &DependencyCache](const FTaskData& A, const FTaskData& B)
 				{
 					const TSet<FString>& BDependencies = DependencyCache.GetAccumulatedDependencies(BaseNodeContainer, B.UniqueID);
-					//if A is a dependency of B then return true to do A before B
-					if (BDependencies.Contains(A.UniqueID))
-					{
-						return true;
-					}
-					// Cache number of B's dependencies as reference on TSet can become stale
-					const int32 BDependenciesNum = BDependencies.Num();
-
 					const TSet<FString>& ADependencies = DependencyCache.GetAccumulatedDependencies(BaseNodeContainer, A.UniqueID);
-					if (ADependencies.Contains(B.UniqueID))
+
+					if (ADependencies.Num() == BDependencies.Num())
 					{
-						return false;
+						return A.UniqueID < B.UniqueID;
 					}
 
-					return ADependencies.Num() <= BDependenciesNum;
+					return ADependencies.Num() < BDependencies.Num();
 				};
 
 				// Nodes cannot depend on a node from another source, so it's faster to sort the dependencies per-source and then append those to the TaskData arrays.
