Major optimizations to FSkeletonHelper::RecursiveAddSkeletonMetaDataValues (importing SK_BlackAlder_L_01.FBX would take ~36s in this function, and now it takes ~86ms)

diff --git a/Engine/Plugins/Interchange/Runtime/Source/Import/Private/Mesh/InterchangeSkeletonHelper.cpp b/Engine/Plugins/Interchange/Runtime/Source/Import/Private/Mesh/InterchangeSkeletonHelper.cpp
--- a/Engine/Plugins/Interchange/Runtime/Source/Import/Private/Mesh/InterchangeSkeletonHelper.cpp
+++ b/Engine/Plugins/Interchange/Runtime/Source/Import/Private/Mesh/InterchangeSkeletonHelper.cpp
@@ -170,8 +170,11 @@ namespace UE::Interchange::Private
 		constexpr bool bAddSourceNodeName = true;
 		UInterchangeUserDefinedAttributesAPI::DuplicateAllUserDefinedAttribute(SceneNode, DestinationNode, bAddSourceNodeName);
 
-		//Iterate childrens
-		const TArray<FString> ChildrenIds = NodeContainer->GetNodeChildrenUids(JointUid);
+		// Iterate children
+		TArray<FString>* CachedChildren = NodeContainer->GetCachedNodeChildrenUids(JointUid);
+		check(CachedChildren); // Call ComputeChildrenCache() prior to calling this function
+
+		const TArray<FString>& ChildrenIds = *CachedChildren;
 		for (int32 ChildIndex = 0; ChildIndex < ChildrenIds.Num(); ++ChildIndex)
 		{
 			RecursiveAddSkeletonMetaDataValues(NodeContainer, DestinationNode, ChildrenIds[ChildIndex]);
diff --git a/Engine/Plugins/Interchange/Runtime/Source/Pipelines/Private/InterchangeGenericAnimationPipeline.cpp b/Engine/Plugins/Interchange/Runtime/Source/Pipelines/Private/InterchangeGenericAnimationPipeline.cpp
--- a/Engine/Plugins/Interchange/Runtime/Source/Pipelines/Private/InterchangeGenericAnimationPipeline.cpp
+++ b/Engine/Plugins/Interchange/Runtime/Source/Pipelines/Private/InterchangeGenericAnimationPipeline.cpp
@@ -729,6 +729,7 @@ void UInterchangeGenericAnimationPipeline::CreateAnimSequenceFactoryNode(UInterc
 		// NOTE: Could this be added as an array of FString attributes on the UInterchangeSkeletalAnimationTrackNode
 #if WITH_EDITOR
 		//Iterate all joints to set the meta data value in the anim sequence factory node
+		BaseNodeContainer->ComputeChildrenCache();
 		UE::Interchange::Private::FSkeletonHelper::RecursiveAddSkeletonMetaDataValues(BaseNodeContainer, AnimSequenceFactoryNode, RootJointUid);
 #endif //WITH_EDITOR
 
diff --git a/Engine/Plugins/Interchange/Runtime/Source/Pipelines/Private/InterchangeGenericCommonSkeletalMeshesAndAnimationsProperties.cpp b/Engine/Plugins/Interchange/Runtime/Source/Pipelines/Private/InterchangeGenericCommonSkeletalMeshesAndAnimationsProperties.cpp
--- a/Engine/Plugins/Interchange/Runtime/Source/Pipelines/Private/InterchangeGenericCommonSkeletalMeshesAndAnimationsProperties.cpp
+++ b/Engine/Plugins/Interchange/Runtime/Source/Pipelines/Private/InterchangeGenericCommonSkeletalMeshesAndAnimationsProperties.cpp
@@ -63,6 +63,7 @@ UInterchangeSkeletonFactoryNode* UInterchangeGenericCommonSkeletalMeshesAndAnima
 	}
 #if WITH_EDITOR
 	//Iterate all joints to set the meta data value in the skeleton node
+	BaseNodeContainer->ComputeChildrenCache();
 	UE::Interchange::Private::FSkeletonHelper::RecursiveAddSkeletonMetaDataValues(BaseNodeContainer, SkeletonFactoryNode, RootJointUid);
 #endif //WITH_EDITOR
 
diff --git a/Engine/Source/Runtime/Interchange/Core/Private/Nodes/InterchangeBaseNodeContainer.cpp b/Engine/Source/Runtime/Interchange/Core/Private/Nodes/InterchangeBaseNodeContainer.cpp
--- a/Engine/Source/Runtime/Interchange/Core/Private/Nodes/InterchangeBaseNodeContainer.cpp
+++ b/Engine/Source/Runtime/Interchange/Core/Private/Nodes/InterchangeBaseNodeContainer.cpp
@@ -164,18 +164,23 @@ bool UInterchangeBaseNodeContainer::SetNodeParentUid(const FString& NodeUniqueID
 
 int32 UInterchangeBaseNodeContainer::GetNodeChildrenCount(const FString& NodeUniqueID) const
 {
+	if (TArray<FString>* CacheChildrenPtr = GetCachedNodeChildrenUids(NodeUniqueID))
+	{
+		return CacheChildrenPtr->Num();
+	}
+
 	TArray<FString> ChildrenUids = GetNodeChildrenUids(NodeUniqueID);
 	return ChildrenUids.Num();
 }
 
 TArray<FString> UInterchangeBaseNodeContainer::GetNodeChildrenUids(const FString& NodeUniqueID) const
 {
-	if(TArray<FString>* CacheChildrenPtr = ChildrenCache.Find(NodeUniqueID))
+	if (TArray<FString>* CacheChildrenPtr = GetCachedNodeChildrenUids(NodeUniqueID))
 	{
 		return *CacheChildrenPtr;
 	}
 
-	//Update the cache
+	// Update the cache
 	TArray<FString>& CacheChildren = ChildrenCache.Add(NodeUniqueID);
 	for (const auto& NodeKeyValue : Nodes)
 	{
@@ -187,6 +192,16 @@ TArray<FString> UInterchangeBaseNodeContainer::GetNodeChildrenUids(const FString
 	return CacheChildren;
 }
 
+TArray<FString>* UInterchangeBaseNodeContainer::GetCachedNodeChildrenUids(const FString& NodeUniqueID) const
+{
+	if (TArray<FString>* CacheChildrenPtr = ChildrenCache.Find(NodeUniqueID))
+	{
+		return CacheChildrenPtr;
+	}
+
+	return nullptr;
+}
+
 UInterchangeBaseNode* UInterchangeBaseNodeContainer::GetNodeChildren(const FString& NodeUniqueID, int32 ChildIndex)
 {
 	return GetNodeChildrenInternal(NodeUniqueID, ChildIndex);
diff --git a/Engine/Source/Runtime/Interchange/Core/Private/Types/AttributeStorage.cpp b/Engine/Source/Runtime/Interchange/Core/Private/Types/AttributeStorage.cpp
--- a/Engine/Source/Runtime/Interchange/Core/Private/Types/AttributeStorage.cpp
+++ b/Engine/Source/Runtime/Interchange/Core/Private/Types/AttributeStorage.cpp
@@ -416,6 +416,21 @@ namespace UE
 			FMemory::Memcpy((uint8*)(*OutValue), &StorageData[AttributeAllocationInfo->Offset], AttributeAllocationInfo->Size);
 		}
 
+		const FStringView FAttributeStorage::GetFStringViewAttributeFromStorage(const uint8* StorageData, const FAttributeAllocationInfo* AttributeAllocationInfo) const
+		{
+			const uint64 NumberOfChar = AttributeAllocationInfo->Size / sizeof(TCHAR);
+			check(NumberOfChar > 0);
+
+			// Null terminator is included in the size
+			if (NumberOfChar <= 1)
+			{
+				return FStringView();
+			}
+
+			const TCHAR* StringContents = reinterpret_cast<const TCHAR*>(&StorageData[AttributeAllocationInfo->Offset]);
+			return FStringView(StringContents, int32(NumberOfChar - 1));
+		}
+
 		EAttributeStorageResult FAttributeStorage::GetAttribute(const FAttributeKey& ElementAttributeKey, FString& OutValue, TSpecializeType<FString >) const
 		{
 			static_assert(TAttributeTypeTraits<FString>::GetType() != EAttributeTypes::None, "Not a supported type for the attributes. Check EAttributeTypes for the supported types.");
@@ -448,7 +463,8 @@ namespace UE
 				return EAttributeStorageResult::Operation_Success;
 			}
 
-			ExtractFStringAttributeFromStorage(AttributeStorage.GetData(), AttributeAllocationInfo, OutValue);
+			const FStringView ValueView = GetFStringViewAttributeFromStorage(AttributeStorage.GetData(), AttributeAllocationInfo);
+			OutValue = ValueView;
 
 			return EAttributeStorageResult::Operation_Success;
 		}
@@ -488,13 +504,8 @@ namespace UE
 				return EAttributeStorageResult::Operation_Success;
 			}
 
-			FString ValueStr;
-
-			//Share the code with FString from here
-			ExtractFStringAttributeFromStorage(AttributeStorage.GetData(), AttributeAllocationInfo, ValueStr);
-
-			//Create the FName and copy it to OutValue
-			OutValue = FName(*ValueStr);
+			const FStringView ValueStr = GetFStringViewAttributeFromStorage(AttributeStorage.GetData(), AttributeAllocationInfo);
+			OutValue = FName(ValueStr);
 
 			return EAttributeStorageResult::Operation_Success;
 		}
@@ -534,12 +545,7 @@ namespace UE
 				return EAttributeStorageResult::Operation_Success;
 			}
 
-			FString ValueStr;
-
-			//Share the code with FString from here
-			ExtractFStringAttributeFromStorage(AttributeStorage.GetData(), AttributeAllocationInfo, ValueStr);
-
-			//Create the FSoftObjectPath and copy it to OutValue
+			const FStringView ValueStr = GetFStringViewAttributeFromStorage(AttributeStorage.GetData(), AttributeAllocationInfo);
 			OutValue = FSoftObjectPath(ValueStr);
 
 			return EAttributeStorageResult::Operation_Success;
diff --git a/Engine/Source/Runtime/Interchange/Core/Public/Nodes/InterchangeBaseNodeContainer.h b/Engine/Source/Runtime/Interchange/Core/Public/Nodes/InterchangeBaseNodeContainer.h
--- a/Engine/Source/Runtime/Interchange/Core/Public/Nodes/InterchangeBaseNodeContainer.h
+++ b/Engine/Source/Runtime/Interchange/Core/Public/Nodes/InterchangeBaseNodeContainer.h
@@ -173,6 +173,8 @@ public:
 	UFUNCTION(BlueprintCallable, Category = "Interchange | Node Container")
 	INTERCHANGECORE_API TArray<FString> GetNodeChildrenUids(const FString& NodeUniqueID) const;
 
+	INTERCHANGECORE_API TArray<FString>* GetCachedNodeChildrenUids(const FString& NodeUniqueID) const;
+
 	/** Get the nth const child of the node */
 	UFUNCTION(BlueprintCallable, Category = "Interchange | Node Container")
 	INTERCHANGECORE_API UInterchangeBaseNode* GetNodeChildren(const FString& NodeUniqueID, int32 ChildIndex);
diff --git a/Engine/Source/Runtime/Interchange/Core/Public/Types/AttributeStorage.h b/Engine/Source/Runtime/Interchange/Core/Public/Types/AttributeStorage.h
--- a/Engine/Source/Runtime/Interchange/Core/Public/Types/AttributeStorage.h
+++ b/Engine/Source/Runtime/Interchange/Core/Public/Types/AttributeStorage.h
@@ -921,6 +921,8 @@ namespace UE
 	
 			INTERCHANGECORE_API void ExtractFStringAttributeFromStorage(const uint8* StorageData, const FAttributeAllocationInfo* AttributeAllocationInfo, FString& OutValue) const;
 
+			INTERCHANGECORE_API const FStringView GetFStringViewAttributeFromStorage(const uint8* StorageData, const FAttributeAllocationInfo* AttributeAllocationInfo) const;
+
 			INTERCHANGECORE_API EAttributeStorageResult GetAttribute(const FAttributeKey& ElementAttributeKey, FString& OutValue, TSpecializeType<FString >) const;
 
 			INTERCHANGECORE_API EAttributeStorageResult GetAttribute(const FAttributeKey& ElementAttributeKey, FName& OutValue, TSpecializeType<FName >) const;
