Significant optimizations to UInterchangeUserDefinedAttributesAPI::GetUserDefinedAttributeInfos and MakeUserDefinedPropertyKey

diff --git a/Engine/Plugins/Interchange/Runtime/Source/Parsers/Fbx/Private/FbxHelper.cpp b/Engine/Plugins/Interchange/Runtime/Source/Parsers/Fbx/Private/FbxHelper.cpp
--- a/Engine/Plugins/Interchange/Runtime/Source/Parsers/Fbx/Private/FbxHelper.cpp
+++ b/Engine/Plugins/Interchange/Runtime/Source/Parsers/Fbx/Private/FbxHelper.cpp
@@ -105,12 +105,15 @@ namespace UE
 
 			FString FFbxHelper::GetFbxPropertyName(const FbxProperty& Property) const
 			{
-				FString PropertyName = UE::Interchange::MakeName(FFbxConvert::MakeString(Property.GetName()));
+				FString PropertyName = FFbxConvert::MakeString(Property.GetName());
+				UE::Interchange::SanitizeName(PropertyName);
+
 				if (PropertyName.Equals(TEXT("none"), ESearchCase::IgnoreCase))
 				{
 					//Replace None by Null because None clash with NAME_None and the create asset will instead call the object ClassName_X
 					PropertyName = TEXT("Null");
 				}
+
 				return PropertyName;
 			}
 
@@ -120,7 +123,10 @@ namespace UE
 				{
 					return FString();
 				}
-				FString ObjName = UE::Interchange::MakeName(FFbxConvert::MakeString(Object->GetName()), bIsJoint);
+
+				FString ObjName = FFbxConvert::MakeString(Object->GetName());
+				UE::Interchange::SanitizeName(ObjName, bIsJoint);
+
 				if (ObjName.Equals(TEXT("none"), ESearchCase::IgnoreCase))
 				{
 					//Replace None by Null because None clash with NAME_None and the create asset will instead call the object ClassName_X
diff --git a/Engine/Source/Runtime/Interchange/Core/Private/InterchangeHelper.cpp b/Engine/Source/Runtime/Interchange/Core/Private/InterchangeHelper.cpp
--- a/Engine/Source/Runtime/Interchange/Core/Private/InterchangeHelper.cpp
+++ b/Engine/Source/Runtime/Interchange/Core/Private/InterchangeHelper.cpp
@@ -4,17 +4,21 @@
 
 namespace UE::Interchange
 {
-	FString MakeName(const FString& InName, bool bIsJoint)
+	void SanitizeName(FString& OutName, bool bIsJoint)
 	{
 		const TCHAR* InvalidChar = bIsJoint ? INVALID_OBJECTNAME_CHARACTERS TEXT("+ ") : INVALID_OBJECTNAME_CHARACTERS;
-		FString TmpName = InName;
 
 		while (*InvalidChar)
 		{
-			TmpName.ReplaceCharInline(*InvalidChar, TCHAR('_'), ESearchCase::CaseSensitive);
+			OutName.ReplaceCharInline(*InvalidChar, TCHAR('_'), ESearchCase::CaseSensitive);
 			++InvalidChar;
 		}
+	}
 
+	FString MakeName(const FString& InName, bool bIsJoint)
+	{
+		FString TmpName = InName;
+		SanitizeName(TmpName, bIsJoint);
 		return TmpName;
 	}
 };
\ No newline at end of file
diff --git a/Engine/Source/Runtime/Interchange/Core/Private/Nodes/InterchangeUserDefinedAttribute.cpp b/Engine/Source/Runtime/Interchange/Core/Private/Nodes/InterchangeUserDefinedAttribute.cpp
--- a/Engine/Source/Runtime/Interchange/Core/Private/Nodes/InterchangeUserDefinedAttribute.cpp
+++ b/Engine/Source/Runtime/Interchange/Core/Private/Nodes/InterchangeUserDefinedAttribute.cpp
@@ -152,25 +152,28 @@ TArray<FInterchangeUserDefinedAttributeInfo> UInterchangeUserDefinedAttributesAP
 	int32 LeftChopIndex = UserDefinedAttributeValuePostKey.Len();
 	int32 AddDelegateRightChopIndex = UserDefinedAttributeDelegateKey.Len();
 
-	for (UE::Interchange::FAttributeKey& AttributeKey : AttributeKeys)
+	for (const UE::Interchange::FAttributeKey& AttributeKey : AttributeKeys)
 	{
-		const FString AttributeKeyString = AttributeKey.ToString();
-		if (AttributeKeyString.StartsWith(UserDefinedAttributeBaseKey) && AttributeKeyString.EndsWith(UserDefinedAttributeValuePostKey))
+		if (AttributeKey.Key.StartsWith(UserDefinedAttributeBaseKey) && AttributeKey.Key.EndsWith(UserDefinedAttributeValuePostKey))
 		{
 			bool RequiresDelegate = false;
-			FString UserDefinedAttributeName = AttributeKeyString.RightChop(RightChopIndex).LeftChop(LeftChopIndex);
+
+			FStringView UserDefinedAttributeName = FStringView(AttributeKey.Key);
+			UserDefinedAttributeName.RemoveSuffix(RightChopIndex);
+			UserDefinedAttributeName.RemovePrefix(LeftChopIndex);
+
 			if (UserDefinedAttributeName.StartsWith(UserDefinedAttributeDelegateKey))
 			{
-				UserDefinedAttributeName = UserDefinedAttributeName.RightChop(AddDelegateRightChopIndex);
+				UserDefinedAttributeName.RemoveSuffix(AddDelegateRightChopIndex);
 				RequiresDelegate = true;
 			}
+
 			FInterchangeUserDefinedAttributeInfo& UserDefinedAttributeInfo = UserDefinedAttributeInfos.AddDefaulted_GetRef();
 			UserDefinedAttributeInfo.Type = InterchangeNode->GetAttributeType(AttributeKey);
 			UserDefinedAttributeInfo.Name = UserDefinedAttributeName;
 			UserDefinedAttributeInfo.RequiresDelegate = RequiresDelegate;
 
-			//Get the optional payload key
-			const FString StorageBaseKey = UserDefinedAttributeBaseKey + UserDefinedAttributeName;
+			// Get the optional payload key
 			const UE::Interchange::FAttributeKey UserDefinedPayloadKey = MakeUserDefinedPropertyPayloadKey(UserDefinedAttributeName, RequiresDelegate);
 			if (InterchangeNode->HasAttribute(UserDefinedPayloadKey))
 			{
@@ -180,6 +183,7 @@ TArray<FInterchangeUserDefinedAttributeInfo> UInterchangeUserDefinedAttributesAP
 			}
 		}
 	}
+
 	return UserDefinedAttributeInfos;
 }
 
@@ -380,11 +384,22 @@ void UInterchangeUserDefinedAttributesAPI::AddApplyAndFillDelegatesToFactory(UIn
 	}
 }
 
+UE::Interchange::FAttributeKey UInterchangeUserDefinedAttributesAPI::MakeUserDefinedPropertyValueKey(const FStringView UserDefinedAttributeName, bool RequiresDelegate)
+{
+	return MakeUserDefinedPropertyKey(UserDefinedAttributeName, RequiresDelegate);
+}
+
 UE::Interchange::FAttributeKey UInterchangeUserDefinedAttributesAPI::MakeUserDefinedPropertyValueKey(const FString& UserDefinedAttributeName, bool RequiresDelegate)
 {
 	return MakeUserDefinedPropertyKey(UserDefinedAttributeName, RequiresDelegate);
 }
 
+UE::Interchange::FAttributeKey UInterchangeUserDefinedAttributesAPI::MakeUserDefinedPropertyPayloadKey(const FStringView UserDefinedAttributeName, bool RequiresDelegate)
+{
+	const bool bGeneratePayloadKey = true;
+	return MakeUserDefinedPropertyKey(UserDefinedAttributeName, RequiresDelegate, bGeneratePayloadKey);
+}
+
 UE::Interchange::FAttributeKey UInterchangeUserDefinedAttributesAPI::MakeUserDefinedPropertyPayloadKey(const FString& UserDefinedAttributeName, bool RequiresDelegate)
 {
 	const bool bGeneratePayloadKey = true;
@@ -409,16 +424,33 @@ bool UInterchangeUserDefinedAttributesAPI::HasAttribute(const UInterchangeBaseNo
 	return false;
 }
 
-UE::Interchange::FAttributeKey UInterchangeUserDefinedAttributesAPI::MakeUserDefinedPropertyKey(const FString& UserDefinedAttributeName, bool RequiresDelegate, bool GeneratePayloadKey /*= false*/)
-{ 
-	//Create a unique Key for this user defined attribute
-	FString Prefix = UserDefinedAttributeBaseKey;
+UE::Interchange::FAttributeKey UInterchangeUserDefinedAttributesAPI::MakeUserDefinedPropertyKey(const FStringView UserDefinedAttributeName, bool RequiresDelegate, bool GeneratePayloadKey /*= false*/)
+{
+	// Create a unique Key for this user defined attribute
+	FStringBuilderBase StringBuilder;
+	
+	StringBuilder.Append(UserDefinedAttributeBaseKey);
+
 	if (RequiresDelegate)
 	{
-		Prefix = Prefix + UserDefinedAttributeDelegateKey;
+		StringBuilder.Append(UserDefinedAttributeDelegateKey);
 	}
 
-	const FString Suffix = GeneratePayloadKey ? UserDefinedAttributePayLoadPostKey : UserDefinedAttributeValuePostKey;
+	StringBuilder.Append(UserDefinedAttributeName);
 
-	return UE::Interchange::FAttributeKey(Prefix + UserDefinedAttributeName + Suffix);
+	if (GeneratePayloadKey)
+	{
+		StringBuilder.Append(UserDefinedAttributePayLoadPostKey);
+	}
+	else
+	{
+		StringBuilder.Append(UserDefinedAttributeValuePostKey);
+	}
+
+	return UE::Interchange::FAttributeKey(StringBuilder.ToView());
+}
+
+UE::Interchange::FAttributeKey UInterchangeUserDefinedAttributesAPI::MakeUserDefinedPropertyKey(const FString& UserDefinedAttributeName, bool RequiresDelegate, bool GeneratePayloadKey /*= false*/)
+{
+	return MakeUserDefinedPropertyKey(FStringView(UserDefinedAttributeName), RequiresDelegate, GeneratePayloadKey);
 }
\ No newline at end of file
diff --git a/Engine/Source/Runtime/Interchange/Core/Public/InterchangeHelper.h b/Engine/Source/Runtime/Interchange/Core/Public/InterchangeHelper.h
--- a/Engine/Source/Runtime/Interchange/Core/Public/InterchangeHelper.h
+++ b/Engine/Source/Runtime/Interchange/Core/Public/InterchangeHelper.h
@@ -11,6 +11,7 @@ namespace UE::Interchange
 	 * Replaces any unsupported characters with "_" character, and removes namespace indicator ":" character
 	 * #todo_interchange: Replace with SanitizeObjectName from InterchangeManager.cpp
 	 */
+	INTERCHANGECORE_API void SanitizeName(FString& OutName, bool bIsJoint = false);
 	INTERCHANGECORE_API FString MakeName(const FString& InName, bool bIsJoint = false);
 
 	class FScopedLambda
diff --git a/Engine/Source/Runtime/Interchange/Core/Public/Nodes/InterchangeUserDefinedAttribute.h b/Engine/Source/Runtime/Interchange/Core/Public/Nodes/InterchangeUserDefinedAttribute.h
--- a/Engine/Source/Runtime/Interchange/Core/Public/Nodes/InterchangeUserDefinedAttribute.h
+++ b/Engine/Source/Runtime/Interchange/Core/Public/Nodes/InterchangeUserDefinedAttribute.h
@@ -119,12 +119,15 @@ public:
 
 	static INTERCHANGECORE_API void AddApplyAndFillDelegatesToFactory(UInterchangeFactoryBaseNode* InterchangeFactoryNode, UClass* ParentClass);
 
+	static INTERCHANGECORE_API UE::Interchange::FAttributeKey MakeUserDefinedPropertyValueKey(const FStringView UserDefinedAttributeName, bool RequiresDelegate);
 	static INTERCHANGECORE_API UE::Interchange::FAttributeKey MakeUserDefinedPropertyValueKey(const FString& UserDefinedAttributeName,bool RequiresDelegate);
 	
+	static INTERCHANGECORE_API UE::Interchange::FAttributeKey MakeUserDefinedPropertyPayloadKey(const FStringView UserDefinedAttributeName, bool RequiresDelegate);
 	static INTERCHANGECORE_API UE::Interchange::FAttributeKey MakeUserDefinedPropertyPayloadKey(const FString& UserDefinedAttributeName, bool RequiresDelegate);
 
 private:
 	static INTERCHANGECORE_API bool HasAttribute(const UInterchangeBaseNode* InterchangeSourceNode, const FString& InUserDefinedAttributeName, bool GeneratePayloadKey, bool& OutRequiresDelegate);
+	static INTERCHANGECORE_API UE::Interchange::FAttributeKey MakeUserDefinedPropertyKey(const FStringView UserDefinedAttributeName, bool RequiresDelegate, bool GeneratePayloadKey = false); 
 	static INTERCHANGECORE_API UE::Interchange::FAttributeKey MakeUserDefinedPropertyKey(const FString& UserDefinedAttributeName,bool RequiresDelegate, bool GeneratePayloadKey = false);
 	
 private:
diff --git a/Engine/Source/Runtime/Interchange/Core/Public/Types/AttributeStorage.h b/Engine/Source/Runtime/Interchange/Core/Public/Types/AttributeStorage.h
--- a/Engine/Source/Runtime/Interchange/Core/Public/Types/AttributeStorage.h
+++ b/Engine/Source/Runtime/Interchange/Core/Public/Types/AttributeStorage.h
@@ -88,6 +88,11 @@ namespace UE
 				Key = Other.ToString();
 			}
 
+			explicit FAttributeKey(const FStringView& Other)
+			{
+				Key = Other;
+			}
+
 			explicit FAttributeKey(const FString& Other)
 			{
 				Key = Other;
