[Interchange] UserDefined Attribute fixes

#rnx

diff --git a/Engine/Source/Runtime/Interchange/Core/Private/Nodes/InterchangeUserDefinedAttribute.cpp b/Engine/Source/Runtime/Interchange/Core/Private/Nodes/InterchangeUserDefinedAttribute.cpp
--- a/Engine/Source/Runtime/Interchange/Core/Private/Nodes/InterchangeUserDefinedAttribute.cpp
+++ b/Engine/Source/Runtime/Interchange/Core/Private/Nodes/InterchangeUserDefinedAttribute.cpp
@@ -159,12 +159,12 @@ TArray<FInterchangeUserDefinedAttributeInfo> UInterchangeUserDefinedAttributesAP
 			bool RequiresDelegate = false;
 
 			FStringView UserDefinedAttributeName = FStringView(AttributeKey.Key);
-			UserDefinedAttributeName.RemoveSuffix(RightChopIndex);
-			UserDefinedAttributeName.RemovePrefix(LeftChopIndex);
+			UserDefinedAttributeName.RemovePrefix(RightChopIndex);
+			UserDefinedAttributeName.RemoveSuffix(LeftChopIndex);
 
 			if (UserDefinedAttributeName.StartsWith(UserDefinedAttributeDelegateKey))
 			{
-				UserDefinedAttributeName.RemoveSuffix(AddDelegateRightChopIndex);
+				UserDefinedAttributeName.RemovePrefix(AddDelegateRightChopIndex);
 				RequiresDelegate = true;
 			}
 
