[Interchange] Add user group to project settings.

#rnx

diff --git a/Engine/Plugins/Interchange/Editor/Source/Pipelines/Private/SInterchangePipelineConfigurationDialog.cpp b/Engine/Plugins/Interchange/Editor/Source/Pipelines/Private/SInterchangePipelineConfigurationDialog.cpp
--- a/Engine/Plugins/Interchange/Editor/Source/Pipelines/Private/SInterchangePipelineConfigurationDialog.cpp
+++ b/Engine/Plugins/Interchange/Editor/Source/Pipelines/Private/SInterchangePipelineConfigurationDialog.cpp
@@ -35,6 +35,7 @@
 #include "Widgets/Layout/SSeparator.h"
 #include "Widgets/Input/STextComboBox.h"
 #include "Widgets/Layout/SUniformGridPanel.h"
+#include "Widgets/Input/SEditableTextBox.h"
 
 #define LOCTEXT_NAMESPACE "InterchangePipelineConfiguration"
 
@@ -431,6 +432,113 @@ TSharedRef<SBox> SInterchangePipelineConfigurationDialog::SpawnPipelineConfigura
 		];
 	}
 
+	TSharedPtr<SWidget> StackAndGroupWidget;
+
+	//Groups
+	FText GroupUsedText = LOCTEXT("GroupUsedText", "Group Used:");
+
+	FInterchangeGroup::EUsedGroupStatus UsedGroupStatus;
+	const FInterchangeGroup& UsedInterchangeGroup = FInterchangeProjectSettingsUtils::GetUsedGroup(UsedGroupStatus);
+
+	switch (UsedGroupStatus)
+	{
+		case FInterchangeGroup::NotSet:
+			StackAndGroupWidget = SNew(SBox)
+				[
+					SNew(SHorizontalBox)
+						+ SHorizontalBox::Slot()
+						.VAlign(VAlign_Center)
+						.AutoWidth()
+						[
+							StackTextComboBox.ToSharedRef()
+						]
+				];
+			break;
+		case FInterchangeGroup::SetAndValid:
+			{
+				FText GroupComboBoxTooltip = LOCTEXT("GroupComboBoxTooltip", "Group usage can be set in Editor Preferences > Interchange > Groups.");
+
+				StackAndGroupWidget = SNew(SBox)
+					[
+						SNew(SHorizontalBox)
+							+ SHorizontalBox::Slot()
+							.VAlign(VAlign_Center)
+							.AutoWidth()
+							[
+								StackTextComboBox.ToSharedRef()
+							]
+							+ SHorizontalBox::Slot()
+							.VAlign(VAlign_Center)
+							.AutoWidth()
+							[
+								SNew(SBox)
+									[
+										SNew(SHorizontalBox)
+											+ SHorizontalBox::Slot()
+											.VAlign(VAlign_Center)
+											.Padding(16.0f, 0.0f, 4.0f, 0.0f)
+											.AutoWidth()
+											[
+												SNew(SBox)
+													[
+														SNew(STextBlock)
+															.Text(GroupUsedText)
+													]
+											]
+											+ SHorizontalBox::Slot()
+											.VAlign(VAlign_Center)
+											.AutoWidth()
+											[
+												SNew(SEditableTextBox)
+													.Text(FText::FromName(UsedInterchangeGroup.DisplayName))
+													.IsEnabled(false)
+													.ToolTipText(GroupComboBoxTooltip)
+											]
+									]
+							]
+					];
+			}
+			break;
+		case FInterchangeGroup::SetAndInvalid:
+			{
+				//invalid Group usage:
+				FText InvalidGroupText = LOCTEXT("InvalidGroupText", "Invalid Group setup for usage!");
+				FText InvalidGroupTooltip = LOCTEXT("InvalidGroupTooltip", "Please review Group usage in Editor Preferences > Interchange > Groups.");
+
+				StackAndGroupWidget = SNew(SBox)
+					[
+						SNew(SHorizontalBox)
+							+ SHorizontalBox::Slot()
+							.VAlign(VAlign_Center)
+							.AutoWidth()
+							[
+								StackTextComboBox.ToSharedRef()
+							]
+							+ SHorizontalBox::Slot()
+							.VAlign(VAlign_Center)
+							.AutoWidth()
+							[
+								SNew(SHorizontalBox)
+									+ SHorizontalBox::Slot()
+									.VAlign(VAlign_Center)
+									.Padding(16.0f, 0.0f, 0.0f, 0.0f)
+									.AutoWidth()
+									[
+										SNew(SBox)
+											[
+												SNew(STextBlock)
+													.Text(InvalidGroupText)
+													.ToolTipText(InvalidGroupTooltip)
+											]
+									]
+							]
+					];
+			}
+			break;
+		default:
+			break;
+	}
+
 	TSharedPtr<SBox> InspectorBox;
 	TSharedRef<SBox> PipelineConfigurationPanelBox = SNew(SBox)
 	[
@@ -447,7 +555,7 @@ TSharedRef<SBox> SInterchangePipelineConfigurationDialog::SpawnPipelineConfigura
 				+ SHorizontalBox::Slot()
 				.AutoWidth()
 				[
-					StackTextComboBox.ToSharedRef()
+					StackAndGroupWidget.ToSharedRef()
 				]
 				+ SHorizontalBox::Slot()
 				.FillWidth(1.0f)
diff --git a/Engine/Source/Runtime/Interchange/Engine/Private/InterchangeProjectSettings.cpp b/Engine/Source/Runtime/Interchange/Engine/Private/InterchangeProjectSettings.cpp
--- a/Engine/Source/Runtime/Interchange/Engine/Private/InterchangeProjectSettings.cpp
+++ b/Engine/Source/Runtime/Interchange/Engine/Private/InterchangeProjectSettings.cpp
@@ -4,6 +4,24 @@
 #include "InterchangeManager.h"
 #include "InterchangeSourceData.h"
 
+#if WITH_EDITOR
+void UInterchangeProjectSettings::PostEditChangeProperty(FPropertyChangedEvent& PropertyChangedEvent)
+{
+	Super::PostEditChangeProperty(PropertyChangedEvent);
+
+	FProperty* PropertyThatChanged = PropertyChangedEvent.MemberProperty;
+	const FString PropertyName = PropertyThatChanged ? PropertyThatChanged->GetName() : TEXT("");
+
+	if (PropertyName == GET_MEMBER_NAME_STRING_CHECKED(UInterchangeProjectSettings, InterchangeGroups))
+	{
+		if (UInterchangeEditorSettings* InterchangeEditorSettings = GetMutableDefault<UInterchangeEditorSettings>())
+		{
+			InterchangeEditorSettings->UpdateUsedGroupName();
+		}
+	}
+}
+#endif
+
 const FInterchangeImportSettings& FInterchangeProjectSettingsUtils::GetImportSettings(const UInterchangeProjectSettings& InterchangeProjectSettings, const bool bIsSceneImport)
 {
 	if (bIsSceneImport)
@@ -41,8 +59,12 @@ FInterchangeImportSettings& FInterchangeProjectSettingsUtils::GetMutableDefaultI
 FName FInterchangeProjectSettingsUtils::GetDefaultPipelineStackName(const bool bIsSceneImport, const UInterchangeSourceData& SourceData)
 {
 	const FInterchangeImportSettings& ImportSettings = GetDefaultImportSettings(bIsSceneImport);
+	
+	FInterchangeGroup::EUsedGroupStatus UsedGroupStatus;
+	const FInterchangeGroup& UsedInterchangeGroup = FInterchangeProjectSettingsUtils::GetUsedGroup(UsedGroupStatus);
+	bool bInterchangeGroupUsed = (UsedGroupStatus == FInterchangeGroup::EUsedGroupStatus::SetAndValid);
 
-	FName DefaultPipelineStack = ImportSettings.DefaultPipelineStack;
+	FName DefaultPipelineStack = bInterchangeGroupUsed ? ImportSettings.DefaultPipelineStack : UsedInterchangeGroup.DefaultPipelineStack;
 
 	if (!bIsSceneImport)
 	{
@@ -52,8 +74,9 @@ FName FInterchangeProjectSettingsUtils::GetDefaultPipelineStackName(const bool b
 		{
 			EInterchangeTranslatorAssetType SupportedAssetTypes = Translator->GetSupportedAssetTypes();
 
-			const FInterchangeContentImportSettings& ContentImportSettings = GetDefault<UInterchangeProjectSettings>()->ContentImportSettings;
-			for (TMap<EInterchangeTranslatorAssetType, FName>::TConstIterator StackOverridesIt = ContentImportSettings.DefaultPipelineStackOverride.CreateConstIterator(); StackOverridesIt; ++StackOverridesIt)
+			const TMap<EInterchangeTranslatorAssetType, FName>& DefaultPipelineStackOverride = bInterchangeGroupUsed ? UsedInterchangeGroup.DefaultPipelineStackOverride : GetDefault<UInterchangeProjectSettings>()->ContentImportSettings.DefaultPipelineStackOverride;
+
+			for (TMap<EInterchangeTranslatorAssetType, FName>::TConstIterator StackOverridesIt = DefaultPipelineStackOverride.CreateConstIterator(); StackOverridesIt; ++StackOverridesIt)
 			{
 				if ((SupportedAssetTypes ^ StackOverridesIt->Key) < StackOverridesIt->Key)
 				{
@@ -108,9 +131,14 @@ void FInterchangeProjectSettingsUtils::SetDefaultPipelineStackName(const bool bI
 
 bool FInterchangeProjectSettingsUtils::ShouldShowPipelineStacksConfigurationDialog(const bool bIsSceneImport, const bool bReImport, const UInterchangeSourceData& SourceData)
 {
-	const FInterchangeImportSettings& ImportSettings = GetDefaultImportSettings(bIsSceneImport);
+	bool bShowImportDialog = true; 
 
-	bool bShowImportDialog = ImportSettings.bShowImportDialog;
+	FInterchangeGroup::EUsedGroupStatus UsedGroupStatus;
+	const FInterchangeGroup& UsedInterchangeGroup = FInterchangeProjectSettingsUtils::GetUsedGroup(UsedGroupStatus);
+	bool bInterchangeGroupUsed = (UsedGroupStatus == FInterchangeGroup::EUsedGroupStatus::SetAndValid);
+
+	bShowImportDialog = bInterchangeGroupUsed ? UsedInterchangeGroup.bShowImportDialog : GetDefaultImportSettings(bIsSceneImport).bShowImportDialog;
+	
 
 	if (!bIsSceneImport)
 	{
@@ -124,8 +152,9 @@ bool FInterchangeProjectSettingsUtils::ShouldShowPipelineStacksConfigurationDial
 			//Iterate all override, if there is at least one override that show the imnport dialog we will show it.
 			bool bFoundOverride = false;
 			bool bShowFromOverrideStack = false;
-			const FInterchangeContentImportSettings& ContentImportSettings = GetDefault<UInterchangeProjectSettings>()->ContentImportSettings;
-			for (const TPair<EInterchangeTranslatorAssetType, FInterchangeDialogOverride>& ShowImportDialog : ContentImportSettings.ShowImportDialogOverride)
+			const TMap<EInterchangeTranslatorAssetType, FInterchangeDialogOverride>& ShowImportDialogOverride = bInterchangeGroupUsed ? UsedInterchangeGroup.ShowImportDialogOverride : GetDefault<UInterchangeProjectSettings>()->ContentImportSettings.ShowImportDialogOverride;
+
+			for (const TPair<EInterchangeTranslatorAssetType, FInterchangeDialogOverride>& ShowImportDialog : ShowImportDialogOverride)
 			{
 				if ((ShowImportDialog.Key == EInterchangeTranslatorAssetType::None && SupportedAssetTypes == EInterchangeTranslatorAssetType::None)
 					|| (static_cast<uint8>(ShowImportDialog.Key & SupportedAssetTypes) > 0))
@@ -161,4 +190,118 @@ bool FInterchangeProjectSettingsUtils::ShouldShowPipelineStacksConfigurationDial
 	}
 
 	return bShowImportDialog;
-}
\ No newline at end of file
+}
+
+const FInterchangeGroup& FInterchangeProjectSettingsUtils::GetUsedGroup(FInterchangeGroup::EUsedGroupStatus& UsedGroupStatus)
+{
+	static const FInterchangeGroup InterchangeGroupNone;
+
+	UsedGroupStatus = FInterchangeGroup::EUsedGroupStatus::NotSet;
+
+	if (const UInterchangeEditorSettings* InterchangeEditorSettings = GetDefault<UInterchangeEditorSettings>())
+	{
+		const FGuid& InterchangeUsedGroupUID = InterchangeEditorSettings->GetUsedGroupUID();
+		if (InterchangeUsedGroupUID.IsValid())
+		{
+			const UInterchangeProjectSettings* ImportSettings = GetDefault<UInterchangeProjectSettings>();
+			for (const FInterchangeGroup& Group : ImportSettings->InterchangeGroups)
+			{
+				if (Group.UniqueID == InterchangeUsedGroupUID)
+				{
+					UsedGroupStatus = FInterchangeGroup::EUsedGroupStatus::SetAndValid;
+					return Group;
+				}
+			}
+
+			UsedGroupStatus = FInterchangeGroup::EUsedGroupStatus::SetAndInvalid;
+		}
+	}
+
+	return InterchangeGroupNone;
+}
+
+TArray<FName> FInterchangeProjectSettingsUtils::GetGroupNames()
+{
+	TArray<FName> GroupNames;
+	GroupNames.Add(FName());
+
+	const TArray<FInterchangeGroup>& Groups = GetDefault<UInterchangeProjectSettings>()->InterchangeGroups;
+	for (const FInterchangeGroup& Group : Groups)
+	{
+		GroupNames.Add(Group.DisplayName);
+	}
+
+	return GroupNames;
+}
+
+void UInterchangeEditorSettings::SetUsedGroupName(const FName& InUsedGroupName)
+{
+	UsedGroupName = InUsedGroupName;
+
+	UpdateUsedGroupUIDFromGroupName();
+}
+
+TArray<FName> UInterchangeEditorSettings::GetSelectableItems() const
+{
+	return FInterchangeProjectSettingsUtils::GetGroupNames();
+}
+
+void UInterchangeEditorSettings::UpdateUsedGroupName()
+{
+	FInterchangeGroup::EUsedGroupStatus UsedGroupStatus;
+	const FInterchangeGroup& UsedInterchangeGroup = FInterchangeProjectSettingsUtils::GetUsedGroup(UsedGroupStatus);
+
+	switch (UsedGroupStatus)
+	{
+		case FInterchangeGroup::NotSet:
+			UsedGroupName = FName();
+			break;
+		case FInterchangeGroup::SetAndValid:
+			UsedGroupName = UsedInterchangeGroup.DisplayName;
+			break;
+		case FInterchangeGroup::SetAndInvalid:
+			UsedGroupName = FName("Invalid Group Used, Defaulting to No Group usage.");
+			break;
+		default:
+			break;
+	}
+}
+
+void UInterchangeEditorSettings::UpdateUsedGroupUIDFromGroupName()
+{
+	UsedGroupUID = FGuid();
+
+	const UInterchangeProjectSettings* ImportSettings = GetDefault<UInterchangeProjectSettings>();
+	for (const FInterchangeGroup& Group : ImportSettings->InterchangeGroups)
+	{
+		if (Group.DisplayName == UsedGroupName)
+		{
+			UsedGroupUID = Group.UniqueID;
+		}
+	}
+
+	UpdateUsedGroupName();
+}
+
+#if WITH_EDITOR
+void UInterchangeEditorSettings::PostEditChangeProperty(FPropertyChangedEvent& PropertyChangedEvent)
+{
+	Super::PostEditChangeProperty(PropertyChangedEvent);
+
+	FProperty* PropertyThatChanged = PropertyChangedEvent.MemberProperty;
+	const FString PropertyName = PropertyThatChanged ? PropertyThatChanged->GetName() : TEXT("");
+
+	if (PropertyName == GET_MEMBER_NAME_STRING_CHECKED(UInterchangeEditorSettings, UsedGroupName))
+	{
+		UpdateUsedGroupUIDFromGroupName();
+	}
+}
+
+void UInterchangeEditorSettings::PostInitProperties()
+{
+	Super::PostInitProperties();
+
+	UpdateUsedGroupName();
+}
+
+#endif
\ No newline at end of file
diff --git a/Engine/Source/Runtime/Interchange/Engine/Public/InterchangeProjectSettings.h b/Engine/Source/Runtime/Interchange/Engine/Public/InterchangeProjectSettings.h
--- a/Engine/Source/Runtime/Interchange/Engine/Public/InterchangeProjectSettings.h
+++ b/Engine/Source/Runtime/Interchange/Engine/Public/InterchangeProjectSettings.h
@@ -99,6 +99,42 @@ struct FInterchangeContentImportSettings : public FInterchangeImportSettings
 	TMap<EInterchangeTranslatorAssetType, FInterchangeDialogOverride> ShowImportDialogOverride;
 };
 
+USTRUCT()
+struct FInterchangeGroup
+{
+	GENERATED_BODY()
+
+	enum EUsedGroupStatus : uint8
+	{
+		NotSet,
+		SetAndValid,
+		SetAndInvalid
+	};
+
+	/** Specifies a different pipeline stack for Interchange to use by default when importing specific types of assets. */
+	UPROPERTY(EditAnywhere, Category = "InterchangeGroup")
+	FName DisplayName;
+
+	/** This tell interchange if the import dialog should show or not when importing a particular type of asset.*/
+	UPROPERTY(VisibleAnywhere, Category = "InterchangeGroup", meta = (IgnoreForMemberInitializationTest, EditCondition = "false", EditConditionHides))
+	FGuid UniqueID = FGuid::NewGuid();
+
+	/** Specifies which pipeline stack Interchange should use by default. */
+	UPROPERTY(EditAnywhere, Category = "Pipeline")
+	FName DefaultPipelineStack = NAME_None;
+
+	/** Specifies a different pipeline stack for Interchange to use by default when importing specific types of assets. */
+	UPROPERTY(EditAnywhere, Category = "Pipeline", Meta = (DisplayAfter = "DefaultPipelineStack"))
+	TMap<EInterchangeTranslatorAssetType, FName> DefaultPipelineStackOverride;
+
+	UPROPERTY(EditAnywhere, Category = "DialogOverride")
+	bool bShowImportDialog = true;
+
+	/** This tell interchange if the import dialog should show or not when importing a particular type of asset.*/
+	UPROPERTY(EditAnywhere, Category = "Pipeline", Meta = (DisplayAfter = "bShowImportDialog"))
+	TMap<EInterchangeTranslatorAssetType, FInterchangeDialogOverride> ShowImportDialogOverride;
+};
+
 UCLASS(config=Engine, meta=(DisplayName=Interchange), MinimalAPI)
 class UInterchangeProjectSettings : public UDeveloperSettings
 {
@@ -141,6 +177,17 @@ public:
 	 */
 	UPROPERTY(EditAnywhere, config, Category = "Converters", meta = (AllowedClasses = "/Script/InterchangeCore.InterchangePipelineBase, /Script/InterchangeEngine.InterchangeBlueprintPipelineBase, /Script/InterchangeEngine.InterchangePythonPipelineAsset"))
 	FSoftObjectPath ConverterDefaultPipeline;
+
+	/**
+	* Groups that define PerTransalatorPipelines that user can select to use.
+	*/
+	UPROPERTY(EditAnywhere, config, Category = "Groups")
+	TArray<FInterchangeGroup> InterchangeGroups;
+
+private:
+#if WITH_EDITOR
+	virtual void PostEditChangeProperty(FPropertyChangedEvent& PropertyChangedEvent) override;
+#endif
 };
 
 class FInterchangeProjectSettingsUtils
@@ -155,6 +202,10 @@ public:
 	static INTERCHANGEENGINE_API void SetDefaultPipelineStackName(const bool bIsSceneImport, const UInterchangeSourceData& SourceData, const FName StackName);
 
 	static INTERCHANGEENGINE_API bool ShouldShowPipelineStacksConfigurationDialog(const bool bIsSceneImport, const bool bReImport, const UInterchangeSourceData& SourceData);
+
+	static INTERCHANGEENGINE_API const FInterchangeGroup& GetUsedGroup(FInterchangeGroup::EUsedGroupStatus& UsedGroupStatus);
+
+	static INTERCHANGEENGINE_API TArray<FName> GetGroupNames();
 };
 
 UCLASS(config = EditorPerProjectUserSettings, meta = (DisplayName = Interchange), MinimalAPI)
@@ -166,4 +217,32 @@ public:
 	/** If enabled, the import option dialog will show when interchange re-import.*/
 	UPROPERTY(EditAnywhere, config, Category = "Show Dialog")
 	bool bShowImportDialogAtReimport = false;
+
+	const FGuid& GetUsedGroupUID() const { return UsedGroupUID; }
+
+	UFUNCTION(BlueprintCallable, Category = "Interchange | Groups")
+	const FName& GetUsedGroupName() const { return UsedGroupName; }
+
+	UFUNCTION(BlueprintCallable, Category = "Interchange | Groups")
+	void SetUsedGroupName(const FName& InUsedGroupName);
+
+	void UpdateUsedGroupName();
+
+private:
+	/** If enabled, the import option dialog will show when interchange re-import.*/
+	UPROPERTY(EditAnywhere, Transient, Category = "Group Used", meta = (GetOptions = "GetSelectableItems", AllowPrivateAccess = "true"))
+	FName UsedGroupName;
+
+	UPROPERTY(EditAnywhere, Category = "Group Used", config, meta = (IgnoreForMemberInitializationTest, EditCondition = "false", EditConditionHides))
+	FGuid UsedGroupUID;
+
+	UFUNCTION()
+	TArray<FName> GetSelectableItems() const;
+
+	void UpdateUsedGroupUIDFromGroupName();
+
+#if WITH_EDITOR
+	virtual void PostEditChangeProperty(FPropertyChangedEvent& PropertyChangedEvent) override;
+	virtual void PostInitProperties() override;
+#endif
 };
