Make sure we can re-import an engine created skeletalmesh from a staticmesh. We now copy the static mesh import data into the created skeletalmesh.

#rnx

diff --git a/Engine/Plugins/Interchange/Editor/Source/InterchangeEditor/Private/InterchangeFbxAssetImportDataConverter.cpp b/Engine/Plugins/Interchange/Editor/Source/InterchangeEditor/Private/InterchangeFbxAssetImportDataConverter.cpp
--- a/Engine/Plugins/Interchange/Editor/Source/InterchangeEditor/Private/InterchangeFbxAssetImportDataConverter.cpp
+++ b/Engine/Plugins/Interchange/Editor/Source/InterchangeEditor/Private/InterchangeFbxAssetImportDataConverter.cpp
@@ -721,7 +721,7 @@ namespace UE::Interchange::Private
 		InterchangeFbxTranslatorSettings->bConvertSceneUnit = FbxAssetImportData->bConvertSceneUnit;
 		InterchangeFbxTranslatorSettings->bKeepFbxNamespace = GetDefault<UEditorPerProjectUserSettings>()->bKeepFbxNamespace;
 		DestinationData->SetTranslatorSettings(InterchangeFbxTranslatorSettings);
-
+		bool bConvertToNewType = false;
 		if (const UFbxStaticMeshImportData* LegacyStaticMeshImportData = Cast<UFbxStaticMeshImportData>(FbxAssetImportData))
 		{
 			UInterchangeStaticMeshFactoryNode* MeshNode = NewObject<UInterchangeStaticMeshFactoryNode>(DestinationContainer);
@@ -729,8 +729,14 @@ namespace UE::Interchange::Private
 			DestinationContainer->AddNode(MeshNode);
 
 			GenericAssetPipeline->MeshPipeline->CommonMeshesProperties->ForceAllMeshAsType = EInterchangeForceMeshType::IFMT_StaticMesh;
-			check(Obj->IsA<UStaticMesh>());
-			FillInterchangeGenericAssetsPipelineFromFbxStaticMesh(GenericAssetPipeline, Cast<UStaticMesh>(Obj));
+			if (Obj->IsA<UStaticMesh>())
+			{
+				FillInterchangeGenericAssetsPipelineFromFbxStaticMesh(GenericAssetPipeline, Cast<UStaticMesh>(Obj));
+			}
+			else
+			{
+				bConvertToNewType = true;
+			}
 			FillInterchangeGenericAssetsPipelineFromFbxStaticMeshImportData(GenericAssetPipeline
 				, LegacyStaticMeshImportData);
 		}
@@ -741,9 +747,10 @@ namespace UE::Interchange::Private
 			DestinationContainer->AddNode(MeshNode);
 
 			GenericAssetPipeline->MeshPipeline->CommonMeshesProperties->ForceAllMeshAsType = EInterchangeForceMeshType::IFMT_SkeletalMesh;
-			check(Obj->IsA<USkeletalMesh>());
 			FillInterchangeGenericAssetsPipelineFromFbxSkeletalMeshImportData(GenericAssetPipeline
 				, LegacySkeletalMeshImportData);
+			
+			bConvertToNewType = (!Obj->IsA<USkeletalMesh>());
 		}
 		else if (const UFbxAnimSequenceImportData* LegacyAnimSequenceImportData = Cast<UFbxAnimSequenceImportData>(FbxAssetImportData))
 		{
@@ -751,9 +758,10 @@ namespace UE::Interchange::Private
 			AnimationNode->InitializeAnimSequenceNode(NodeUniqueId, NodeDisplayLabel);
 			DestinationContainer->AddNode(AnimationNode);
 
-			check(Obj->IsA<UAnimSequence>());
 			FillInterchangeGenericAssetsPipelineFromFbxAnimSequenceImportData(GenericAssetPipeline
 				, LegacyAnimSequenceImportData);
+
+			bConvertToNewType = (!Obj->IsA<UAnimSequence>());
 		}
 
 		if (UInterchangeFactoryBaseNode* DestinationFactoryNode = DestinationContainer->GetFactoryNode(NodeUniqueId))
@@ -763,6 +771,13 @@ namespace UE::Interchange::Private
 			DestinationData->SetNodeContainer(DestinationContainer);
 			DestinationData->NodeUniqueID = NodeUniqueId;
 		}
+#if WITH_EDITOR
+		//If the type of asset has change we must convert the options
+		if (bConvertToNewType)
+		{
+			DestinationData->ConvertAssetImportDataToNewOwner(Obj);
+		}
+#endif
 		return DestinationData;
 	}
 
diff --git a/Engine/Plugins/Interchange/Runtime/Source/Pipelines/Private/InterchangeGenericMeshPipeline.cpp b/Engine/Plugins/Interchange/Runtime/Source/Pipelines/Private/InterchangeGenericMeshPipeline.cpp
--- a/Engine/Plugins/Interchange/Runtime/Source/Pipelines/Private/InterchangeGenericMeshPipeline.cpp
+++ b/Engine/Plugins/Interchange/Runtime/Source/Pipelines/Private/InterchangeGenericMeshPipeline.cpp
@@ -173,7 +173,9 @@ void UInterchangeGenericMeshPipeline::AdjustSettingsForContext(const FInterchang
 			}
 			bImportStaticMeshes = false;
 			HideCategories.Add(StaticMeshesCategory);
-			if(!bContainSkeletalMesh || SkeletalMeshImportContentType == EInterchangeSkeletalMeshContentType::Geometry)
+			if(!bContainSkeletalMesh
+				|| SkeletalMeshImportContentType == EInterchangeSkeletalMeshContentType::Geometry
+				|| CommonMeshesProperties->ForceAllMeshAsType == EInterchangeForceMeshType::IFMT_StaticMesh)
 			{
 				CommonMeshesProperties->ForceAllMeshAsType = EInterchangeForceMeshType::IFMT_SkeletalMesh;
 			}
@@ -183,7 +185,8 @@ void UInterchangeGenericMeshPipeline::AdjustSettingsForContext(const FInterchang
 			HideCategories.Add(SkeletalMeshesCategory);
 			HideCategories.Add(CommonSkeletalMeshesAndAnimationCategory);
 			bImportSkeletalMeshes = false;
-			if (!bContainStaticMesh)
+			if (!bContainStaticMesh
+				|| CommonMeshesProperties->ForceAllMeshAsType == EInterchangeForceMeshType::IFMT_SkeletalMesh)
 			{
 				CommonMeshesProperties->ForceAllMeshAsType = EInterchangeForceMeshType::IFMT_StaticMesh;
 			}
diff --git a/Engine/Source/Developer/SkeletalMeshUtilitiesCommon/Private/StaticToSkeletalMeshConverter.cpp b/Engine/Source/Developer/SkeletalMeshUtilitiesCommon/Private/StaticToSkeletalMeshConverter.cpp
--- a/Engine/Source/Developer/SkeletalMeshUtilitiesCommon/Private/StaticToSkeletalMeshConverter.cpp
+++ b/Engine/Source/Developer/SkeletalMeshUtilitiesCommon/Private/StaticToSkeletalMeshConverter.cpp
@@ -5,9 +5,11 @@
 #if WITH_EDITOR
 
 #include "Animation/Skeleton.h"
+#include "EditorFramework/AssetImportData.h"
 #include "Engine/SkeletalMesh.h"
 #include "Engine/SkinnedAssetCommon.h"
 #include "Engine/StaticMesh.h"
+#include "InterchangeHelper.h"
 #include "Interfaces/ITargetPlatformManagerModule.h"
 #include "LODUtilities.h"
 #include "MeshDescription.h"
@@ -429,7 +431,8 @@ bool FStaticToSkeletalMeshConverter::InitializeSkeletalMeshFromStaticMesh(
 	{
 		FSkeletalMaterial Material(
 			StaticMaterial.MaterialInterface,
-			StaticMaterial.MaterialSlotName);
+			StaticMaterial.MaterialSlotName,
+			StaticMaterial.ImportedMaterialSlotName);
 		
 		Materials.Add(Material);
 	}
@@ -447,6 +450,14 @@ bool FStaticToSkeletalMeshConverter::InitializeSkeletalMeshFromStaticMesh(
 	InSkeletalMesh->SetPositiveBoundsExtension(InStaticMesh->GetPositiveBoundsExtension());
 	InSkeletalMesh->SetNegativeBoundsExtension(InStaticMesh->GetNegativeBoundsExtension());
 
+	//Create some import data so we can re-import this new skeletalmesh
+	UAssetImportData* OriginalAssetImportData = InStaticMesh->GetAssetImportData();
+	if (OriginalAssetImportData)
+	{
+		UAssetImportData* DuplicateAssetImportData = DuplicateObject<UAssetImportData>(OriginalAssetImportData, InSkeletalMesh);
+		DuplicateAssetImportData->ConvertAssetImportDataToNewOwner(InSkeletalMesh);
+		InSkeletalMesh->SetAssetImportData(DuplicateAssetImportData);
+	}
 	return true;
 }
 
diff --git a/Engine/Source/Developer/SkeletalMeshUtilitiesCommon/SkeletalMeshUtilitiesCommon.Build.cs b/Engine/Source/Developer/SkeletalMeshUtilitiesCommon/SkeletalMeshUtilitiesCommon.Build.cs
--- a/Engine/Source/Developer/SkeletalMeshUtilitiesCommon/SkeletalMeshUtilitiesCommon.Build.cs
+++ b/Engine/Source/Developer/SkeletalMeshUtilitiesCommon/SkeletalMeshUtilitiesCommon.Build.cs
@@ -18,6 +18,7 @@ namespace UnrealBuildTool.Rules
 					"CoreUObject",
 					"Engine",
 					"ImageCore",
+					"InterchangeCore",
 					"MeshDescription",
 					"MeshUtilitiesCommon",
 					"RenderCore",
diff --git a/Engine/Source/Runtime/Engine/Classes/EditorFramework/AssetImportData.h b/Engine/Source/Runtime/Engine/Classes/EditorFramework/AssetImportData.h
--- a/Engine/Source/Runtime/Engine/Classes/EditorFramework/AssetImportData.h
+++ b/Engine/Source/Runtime/Engine/Classes/EditorFramework/AssetImportData.h
@@ -163,6 +163,18 @@ public:
 	/** Convert an absolute import path so that it's relative to either this object's package, BaseDir() or leave it absolute */
 	static ENGINE_API FString SanitizeImportFilename(const FString& InPath, const UPackage* Outermost);
 	static ENGINE_API FString SanitizeImportFilename(const FString& InPath, const FString& PackagePath);
+#if WITH_EDITOR
+	/**
+	 * This function is call when we duplicate asset import data and assign the new copy to a new asset type.
+	 * 
+	 * @note - A use case is when we convert static mesh to skeletal mesh, Interchange need to change the pipeline context
+	 * for the new asset type so we show the correct properties in the re-import dialog.
+	 */
+	ENGINE_API virtual bool ConvertAssetImportDataToNewOwner(UObject* Owner)
+	{
+		return false;
+	}
+#endif
 
 	ENGINE_API virtual void PostLoad() override;
 
diff --git a/Engine/Source/Runtime/Interchange/Engine/Private/InterchangeAssetImportData.cpp b/Engine/Source/Runtime/Interchange/Engine/Private/InterchangeAssetImportData.cpp
--- a/Engine/Source/Runtime/Interchange/Engine/Private/InterchangeAssetImportData.cpp
+++ b/Engine/Source/Runtime/Interchange/Engine/Private/InterchangeAssetImportData.cpp
@@ -1,15 +1,14 @@
 // Copyright Epic Games, Inc. All Rights Reserved.
 
 #include "InterchangeAssetImportData.h"
+#include "InterchangeCustomVersion.h"
 #include "InterchangeEngineLogPrivate.h"
 #include "InterchangeManager.h"
 #include "InterchangePipelineBase.h"
-
-#include "InterchangeCustomVersion.h"
-
+#include "InterchangeProjectSettings.h"
 #include "JsonObjectConverter.h"
-#include "Serialization/LargeMemoryWriter.h"
 #include "Serialization/LargeMemoryReader.h"
+#include "Serialization/LargeMemoryWriter.h"
 #include "UObject/CoreRedirects.h"
 #include "UObject/UObjectIterator.h"
 
@@ -138,6 +137,79 @@ void UInterchangeAssetImportData::Serialize(FArchive& Ar)
 	}
 }
 
+#if WITH_EDITOR
+bool UInterchangeAssetImportData::ConvertAssetImportDataToNewOwner(UObject* Owner)
+{
+	if (!Owner)
+	{
+		return false;
+	}
+
+	ProcessPipelinesCache();
+	if (TransientPipelines.Num() == 0)
+	{
+		return false;
+	}
+	//Find the default asset stack for the first file of this "asset import data"
+	//Then we will generate all the pipeline with the correct context
+	//The final goal is to transfer the UInterchangePiplineBase::PropertiesStates to make sure property category are not hidden if user re-import the asset owning this AssetImportData
+	constexpr bool bImportSceneFalse = false;
+	const FInterchangeImportSettings& InterchangeImportSettings = FInterchangeProjectSettingsUtils::GetDefaultImportSettings(bImportSceneFalse);
+	const TMap<FName, FInterchangePipelineStack>& DefaultPipelineStacks = InterchangeImportSettings.PipelineStacks;
+	UE::Interchange::FScopedSourceData InterchangeSourceData(GetFirstFilename());
+	UE::Interchange::FScopedTranslator ScopedTranslator(InterchangeSourceData.GetSourceData());
+	FName DefaultStackName = FInterchangeProjectSettingsUtils::GetDefaultPipelineStackName(bImportSceneFalse, *InterchangeSourceData.GetSourceData());
+	if (DefaultPipelineStacks.Contains(DefaultStackName))
+	{
+		TArray<UInterchangePipelineBase*> GeneratedPipelines;
+		const FInterchangePipelineStack& PipelineStack = DefaultPipelineStacks.FindChecked(DefaultStackName);
+		const TArray<FSoftObjectPath>* SoftPathPipelines = &PipelineStack.Pipelines;
+		// If applicable, check to see if a specific pipeline stack is associated with this translator
+		for (const FInterchangeTranslatorPipelines& TranslatorPipelines : PipelineStack.PerTranslatorPipelines)
+		{
+			const UClass* TranslatorClass = TranslatorPipelines.Translator.LoadSynchronous();
+			if (ScopedTranslator.GetTranslator() && ScopedTranslator.GetTranslator()->IsA(TranslatorClass))
+			{
+				SoftPathPipelines = &TranslatorPipelines.Pipelines;
+				break;
+			}
+		}
+
+		for (int32 PipelineIndex = 0; PipelineIndex < SoftPathPipelines->Num(); ++PipelineIndex)
+		{
+			if (UInterchangePipelineBase* GeneratedPipeline = UE::Interchange::GeneratePipelineInstance((*SoftPathPipelines)[PipelineIndex]))
+			{
+				FInterchangePipelineContextParams ContextParams;
+				ContextParams.ContextType = EInterchangePipelineContext::AssetImport;
+				ContextParams.ReimportAsset = Owner;
+				ContextParams.BaseNodeContainer = GetNodeContainer();
+				GeneratedPipeline->AdjustSettingsForContext(ContextParams);
+				GeneratedPipelines.Add(GeneratedPipeline);
+			}
+		}
+
+		//We are now properly setup with Some generated pipelines stack that has the correct context value for the PropertiesStates
+		for (TObjectPtr<UObject>& PipelinePtr : TransientPipelines)
+		{
+			if (UInterchangePipelineBase* Pipeline = Cast<UInterchangePipelineBase>(PipelinePtr.Get()))
+			{
+				UClass* PipelineClass = Pipeline->GetClass();
+				for (UInterchangePipelineBase* GeneratedPipeline : GeneratedPipelines)
+				{
+					if (GeneratedPipeline->GetClass()->IsChildOf(PipelineClass))
+					{
+						//Push the properties states to the pipeline
+						Pipeline->TransferAdjustSettings(GeneratedPipeline);
+						Pipeline->AdjustSettingsFromCache();
+						break;
+					}
+				}
+			}
+		}
+	}
+	return true;
+}
+#endif
 
 UInterchangeBaseNodeContainer* UInterchangeAssetImportData::GetNodeContainer() const
 {
diff --git a/Engine/Source/Runtime/Interchange/Engine/Private/InterchangeManager.cpp b/Engine/Source/Runtime/Interchange/Engine/Private/InterchangeManager.cpp
--- a/Engine/Source/Runtime/Interchange/Engine/Private/InterchangeManager.cpp
+++ b/Engine/Source/Runtime/Interchange/Engine/Private/InterchangeManager.cpp
@@ -20,8 +20,9 @@
 #include "InterchangeAssetImportData.h"
 #include "InterchangeBlueprintPipelineBase.h"
 #include "InterchangeEditorUtilitiesBase.h"
-#include "InterchangeFactoryBase.h"
 #include "InterchangeEngineLogPrivate.h"
+#include "InterchangeFactoryBase.h"
+#include "InterchangeHelper.h"
 #include "InterchangeProjectSettings.h"
 #include "InterchangePythonPipelineBase.h"
 #include "InterchangeSourceData.h"
diff --git a/Engine/Source/Runtime/Interchange/Engine/Public/InterchangeAssetImportData.h b/Engine/Source/Runtime/Interchange/Engine/Public/InterchangeAssetImportData.h
--- a/Engine/Source/Runtime/Interchange/Engine/Public/InterchangeAssetImportData.h
+++ b/Engine/Source/Runtime/Interchange/Engine/Public/InterchangeAssetImportData.h
@@ -26,6 +26,9 @@ public:
 	virtual void Serialize(FArchive& Ar) override;
 	// End UObject interface
 
+#if WITH_EDITOR
+	INTERCHANGEENGINE_API  virtual bool ConvertAssetImportDataToNewOwner(UObject* Owner) override;
+#endif
 
 	/**
 	 * Return the first filename stored in this data. The resulting filename will be absolute (that is, not relative to the asset).
