Interchange: Expose translator settings Save/Load to Python.

diff --git a/Engine/Source/Runtime/Interchange/Core/Public/InterchangeTranslatorBase.h b/Engine/Source/Runtime/Interchange/Core/Public/InterchangeTranslatorBase.h
--- a/Engine/Source/Runtime/Interchange/Core/Public/InterchangeTranslatorBase.h
+++ b/Engine/Source/Runtime/Interchange/Core/Public/InterchangeTranslatorBase.h
@@ -29,8 +29,10 @@ class UInterchangeTranslatorSettings : public UObject
 	GENERATED_BODY()
 
 public:
+	UFUNCTION(BlueprintCallable, Category = "Interchange")
 	INTERCHANGECORE_API void LoadSettings();
 
+	UFUNCTION(BlueprintCallable, Category = "Interchange")
 	INTERCHANGECORE_API void SaveSettings();
 };
 
