[Interchange] Groups usage, default pipeline stack fix.

#rnx

diff --git a/Engine/Source/Runtime/Interchange/Engine/Private/InterchangeProjectSettings.cpp b/Engine/Source/Runtime/Interchange/Engine/Private/InterchangeProjectSettings.cpp
--- a/Engine/Source/Runtime/Interchange/Engine/Private/InterchangeProjectSettings.cpp
+++ b/Engine/Source/Runtime/Interchange/Engine/Private/InterchangeProjectSettings.cpp
@@ -64,7 +64,7 @@ FName FInterchangeProjectSettingsUtils::GetDefaultPipelineStackName(const bool b
 	const FInterchangeGroup& UsedInterchangeGroup = FInterchangeProjectSettingsUtils::GetUsedGroup(UsedGroupStatus);
 	bool bInterchangeGroupUsed = (UsedGroupStatus == FInterchangeGroup::EUsedGroupStatus::SetAndValid);
 
-	FName DefaultPipelineStack = bInterchangeGroupUsed ? ImportSettings.DefaultPipelineStack : UsedInterchangeGroup.DefaultPipelineStack;
+	FName DefaultPipelineStack = bInterchangeGroupUsed ? UsedInterchangeGroup.DefaultPipelineStack : ImportSettings.DefaultPipelineStack;
 
 	if (!bIsSceneImport)
 	{
