[BUGFIX] Wrong interchange log on successful import

diff --git a/Engine/Source/Runtime/Interchange/Engine/Private/Tasks/InterchangeTaskCompletion.cpp b/Engine/Source/Runtime/Interchange/Engine/Private/Tasks/InterchangeTaskCompletion.cpp
--- a/Engine/Source/Runtime/Interchange/Engine/Private/Tasks/InterchangeTaskCompletion.cpp
+++ b/Engine/Source/Runtime/Interchange/Engine/Private/Tasks/InterchangeTaskCompletion.cpp
@@ -268,14 +268,15 @@ void UE::Interchange::FTaskCompletion_GameThread::Execute()
 					}
 				}
 
-				UE_LOG(LogInterchangeEngine, Display, TEXT("Interchange import cancelled [%s]"), *AsyncHelper->SourceDatas[SourceIndex]->ToDisplayString());
+				UE_LOG(LogInterchangeEngine, Display, TEXT("Interchange import completed [%s]"), *AsyncHelper->SourceDatas[SourceIndex]->ToDisplayString());
 			});
 	}
 	else
 	{
 		//If task is canceled, delete all created assets by this task
-		AsyncHelper->IterateImportedAssetsPerSourceIndex([](int32 SourceIndex, const TArray<FImportAsyncHelper::FImportedObjectInfo>& AssetInfos)
+		AsyncHelper->IterateImportedAssetsPerSourceIndex([AsyncHelper](int32 SourceIndex, const TArray<FImportAsyncHelper::FImportedObjectInfo>& AssetInfos)
 			{
+				UE_LOG(LogInterchangeEngine, Display, TEXT("Interchange import cancelled [%s]"), *AsyncHelper->SourceDatas[SourceIndex]->ToDisplayString());
 				for (const FImportAsyncHelper::FImportedObjectInfo& AssetInfo : AssetInfos)
 				{
 					using namespace UE::Interchange::Private::ObjectDeletionUtils;
