[Interchange] Interchange Group duplication bug fixes (UUID related issue).

#rnx

diff --git a/Engine/Source/Runtime/Interchange/Engine/Public/InterchangeProjectSettings.h b/Engine/Source/Runtime/Interchange/Engine/Public/InterchangeProjectSettings.h
--- a/Engine/Source/Runtime/Interchange/Engine/Public/InterchangeProjectSettings.h
+++ b/Engine/Source/Runtime/Interchange/Engine/Public/InterchangeProjectSettings.h
@@ -181,7 +181,7 @@ public:
 	/**
 	* Groups that define PerTransalatorPipelines that user can select to use.
 	*/
-	UPROPERTY(EditAnywhere, config, Category = "Groups")
+	UPROPERTY(EditAnywhere, config, Category = "Groups", meta = (NoElementDuplicate))
 	TArray<FInterchangeGroup> InterchangeGroups;
 
 private:
