[Intercnhange] Removing Axis Conversion on Socket Transforms.

#rnx

diff --git a/Engine/Plugins/Interchange/Runtime/Source/Import/Private/Gltf/InterchangeGltfTranslator.cpp b/Engine/Plugins/Interchange/Runtime/Source/Import/Private/Gltf/InterchangeGltfTranslator.cpp
--- a/Engine/Plugins/Interchange/Runtime/Source/Import/Private/Gltf/InterchangeGltfTranslator.cpp
+++ b/Engine/Plugins/Interchange/Runtime/Source/Import/Private/Gltf/InterchangeGltfTranslator.cpp
@@ -438,6 +438,19 @@ bool UInterchangeGLTFTranslator::Translate( UInterchangeBaseNodeContainer& NodeC
 		SourceNode->SetExtraInformation(Extra.Name, Extra.Value);
 	}
 
+	{
+		//For sockets, crete the Axis System conversion inverse:
+
+		FMatrix glTFToUE5Matrix = FMatrix(
+			FPlane(1, 0, 0, 0),  // X-axis remains the same
+			FPlane(0, 0, 1, 0),  // Y-axis becomes Z-axis
+			FPlane(0, 1, 0, 0),  // Z-axis becomes Y-axis
+			FPlane(0, 0, 0, 1)   // W (homogeneous coordinate)
+		);
+		FTransform AxisConversionInverseTransform = FTransform(glTFToUE5Matrix.Inverse());
+		SourceNode->SetCustomAxisConversionInverseTransform(AxisConversionInverseTransform);
+	}
+
 	const FString FileName = GltfAsset.Name;
 
 	//Required Extension Check:
diff --git a/Engine/Plugins/Interchange/Runtime/Source/Import/Private/Mesh/InterchangeStaticMeshFactory.cpp b/Engine/Plugins/Interchange/Runtime/Source/Import/Private/Mesh/InterchangeStaticMeshFactory.cpp
--- a/Engine/Plugins/Interchange/Runtime/Source/Import/Private/Mesh/InterchangeStaticMeshFactory.cpp
+++ b/Engine/Plugins/Interchange/Runtime/Source/Import/Private/Mesh/InterchangeStaticMeshFactory.cpp
@@ -40,6 +40,7 @@
 #include "StaticMeshCompiler.h"
 #include "StaticMeshOperations.h"
 #include "InterchangeManager.h"
+#include "Nodes/InterchangeSourceNode.h"
 
 #include UE_INLINE_GENERATED_CPP_BY_NAME(InterchangeStaticMeshFactory)
 
@@ -1871,6 +1872,14 @@ bool UInterchangeStaticMeshFactory::ImportSockets(const FImportAssetObjectParams
 
 			UE::Interchange::Private::MeshHelper::AddSceneNodeGeometricAndPivotToGlobalTransform(Transform, SceneNode, bBakeMeshes, bBakePivotMeshes);
 
+			//Apply axis transformation inverse to get correct Socket Transform:
+			const UInterchangeSourceNode* SourceNode = UInterchangeSourceNode::GetUniqueInstance(Arguments.NodeContainer);
+			FTransform AxisConversionInverseTransform;
+			if (SourceNode->GetCustomAxisConversionInverseTransform(AxisConversionInverseTransform))
+			{
+				Transform = AxisConversionInverseTransform * Transform;
+			}
+
 			UStaticMeshSocket* Socket = StaticMesh->FindSocket(SocketName);
 			if (!Socket)
 			{
diff --git a/Engine/Plugins/Interchange/Runtime/Source/Parsers/Fbx/Private/FbxAPI.cpp b/Engine/Plugins/Interchange/Runtime/Source/Parsers/Fbx/Private/FbxAPI.cpp
--- a/Engine/Plugins/Interchange/Runtime/Source/Parsers/Fbx/Private/FbxAPI.cpp
+++ b/Engine/Plugins/Interchange/Runtime/Source/Parsers/Fbx/Private/FbxAPI.cpp
@@ -74,7 +74,7 @@ namespace UE
 				return FbxHelper;
 			}
 
-			bool FFbxParser::LoadFbxFile(const FString& Filename)
+			bool FFbxParser::LoadFbxFile(const FString& Filename, UInterchangeBaseNodeContainer& NodeContainer)
 			{
 				SourceFilename = Filename;
 				int32 SDKMajor, SDKMinor, SDKRevision;
@@ -140,7 +140,13 @@ namespace UE
 				EnsureNodeNameAreValid(FPaths::GetBaseFilename(Filename, bRemovePath));
 
 				//We always convert scene to UE axis and units
-				FFbxConvert::ConvertScene(SDKScene, bConvertScene, bForceFrontXAxis, bConvertSceneUnit, FileDetails.AxisDirection, FileDetails.UnitSystem);
+				FbxAMatrix AxisConversionInverseMatrix;
+				FFbxConvert::ConvertScene(SDKScene, bConvertScene, bForceFrontXAxis, bConvertSceneUnit, FileDetails.AxisDirection, FileDetails.UnitSystem, AxisConversionInverseMatrix);
+
+				//Save the AxisConversionInverseTransform into InterchangeSourceNode (so that socket transport can use it accordingly).
+				FTransform AxisConversionInverseTransform = FFbxConvert::ConvertTransform<FTransform, FVector, FQuat>(AxisConversionInverseMatrix);
+				UInterchangeSourceNode* SourceNode = UInterchangeSourceNode::FindOrCreateUniqueInstance(&NodeContainer);
+				SourceNode->SetCustomAxisConversionInverseTransform(AxisConversionInverseTransform);
 
 				FrameRate = FbxTime::GetFrameRate(SDKScene->GetGlobalSettings().GetTimeMode());
 				FileDetails.FrameRate = FString::Printf(TEXT("%.2f"), FrameRate);
diff --git a/Engine/Plugins/Interchange/Runtime/Source/Parsers/Fbx/Private/FbxAPI.h b/Engine/Plugins/Interchange/Runtime/Source/Parsers/Fbx/Private/FbxAPI.h
--- a/Engine/Plugins/Interchange/Runtime/Source/Parsers/Fbx/Private/FbxAPI.h
+++ b/Engine/Plugins/Interchange/Runtime/Source/Parsers/Fbx/Private/FbxAPI.h
@@ -63,7 +63,7 @@ namespace UE
 				const TSharedPtr<FFbxHelper> GetFbxHelper();
 
 				/* Load an fbx file into the fbx sdk, return false if the file could not be load. */
-				bool LoadFbxFile(const FString& Filename);
+				bool LoadFbxFile(const FString& Filename, UInterchangeBaseNodeContainer& NodeContainer);
 
 				/* Extract the fbx data from the sdk into our node container */
 				void FillContainerWithFbxScene(UInterchangeBaseNodeContainer& NodeContainer);
diff --git a/Engine/Plugins/Interchange/Runtime/Source/Parsers/Fbx/Private/FbxConvert.cpp b/Engine/Plugins/Interchange/Runtime/Source/Parsers/Fbx/Private/FbxConvert.cpp
--- a/Engine/Plugins/Interchange/Runtime/Source/Parsers/Fbx/Private/FbxConvert.cpp
+++ b/Engine/Plugins/Interchange/Runtime/Source/Parsers/Fbx/Private/FbxConvert.cpp
@@ -55,7 +55,7 @@ namespace UE
 				return AxisDirection;
 			}
 
-			void FFbxConvert::ConvertScene(FbxScene* SDKScene, const bool bConvertScene, const bool bForceFrontXAxis, const bool bConvertSceneUnit, FString& FileSystemDirection, FString& FileUnitSystem)
+			void FFbxConvert::ConvertScene(FbxScene* SDKScene, const bool bConvertScene, const bool bForceFrontXAxis, const bool bConvertSceneUnit, FString& FileSystemDirection, FString& FileUnitSystem, FbxAMatrix& AxisConversionInverseMatrix)
 			{
 				if (!ensure(SDKScene))
 				{
@@ -87,6 +87,15 @@ namespace UE
 					{
 						FbxRootNodeUtility::RemoveAllFbxRoots(SDKScene);
 						UnrealImportAxis.ConvertScene(SDKScene);
+
+						FbxAMatrix SourceMatrix;
+						FileAxisSystem.GetMatrix(SourceMatrix);
+						FbxAMatrix UnrealMatrix;
+						UnrealImportAxis.GetMatrix(UnrealMatrix);
+
+						FbxAMatrix AxisConversionMatrix;
+						AxisConversionMatrix = SourceMatrix.Inverse() * UnrealMatrix;
+						AxisConversionInverseMatrix = AxisConversionMatrix.Inverse();
 					}
 				}
 
diff --git a/Engine/Plugins/Interchange/Runtime/Source/Parsers/Fbx/Private/FbxConvert.h b/Engine/Plugins/Interchange/Runtime/Source/Parsers/Fbx/Private/FbxConvert.h
--- a/Engine/Plugins/Interchange/Runtime/Source/Parsers/Fbx/Private/FbxConvert.h
+++ b/Engine/Plugins/Interchange/Runtime/Source/Parsers/Fbx/Private/FbxConvert.h
@@ -162,7 +162,7 @@ namespace UE
 				/**
 				 * Convert a fbx scene 
 				 */
-				static void ConvertScene(FbxScene* SDKScene, const bool bConvertScene, const bool bForceFrontXAxis, const bool bConvertSceneUnit, FString& FileSystemDirection, FString& FileUnitSystem);
+				static void ConvertScene(FbxScene* SDKScene, const bool bConvertScene, const bool bForceFrontXAxis, const bool bConvertSceneUnit, FString& FileSystemDirection, FString& FileUnitSystem, FbxAMatrix& AxisConversionInverseMatrix);
 				
 				/** Scene Conversion API End */
 				//////////////////////////////////////////////////////////////////////////
diff --git a/Engine/Plugins/Interchange/Runtime/Source/Parsers/Fbx/Private/InterchangeFbxParser.cpp b/Engine/Plugins/Interchange/Runtime/Source/Parsers/Fbx/Private/InterchangeFbxParser.cpp
--- a/Engine/Plugins/Interchange/Runtime/Source/Parsers/Fbx/Private/InterchangeFbxParser.cpp
+++ b/Engine/Plugins/Interchange/Runtime/Source/Parsers/Fbx/Private/InterchangeFbxParser.cpp
@@ -59,7 +59,10 @@ namespace UE::Interchange
 		SourceFilename = Filename;
 		ResultsContainer->Empty();
 
-		if (!FbxParserPrivate->LoadFbxFile(Filename))
+		//Since we are not in main thread we cannot use TStrongPtr, so we will add the object to the root and remove it when we are done
+		UInterchangeBaseNodeContainer* Container = NewObject<UInterchangeBaseNodeContainer>(GetTransientPackage(), NAME_None);
+
+		if (!FbxParserPrivate->LoadFbxFile(Filename, *Container))
 		{
 			if (UInterchangeResultError_Generic* Error = AddMessage<UInterchangeResultError_Generic>())
 			{
@@ -70,8 +73,7 @@ namespace UE::Interchange
 		}
 
 		ResultFilepath = ResultFolder + TEXT("/SceneDescription.itc");
-		//Since we are not in main thread we cannot use TStrongPtr, so we will add the object to the root and remove it when we are done
-		UInterchangeBaseNodeContainer* Container = NewObject<UInterchangeBaseNodeContainer>(GetTransientPackage(), NAME_None);
+		
 		if (!ensure(Container != nullptr))
 		{
 			if (UInterchangeResultError_Generic* Error = AddMessage<UInterchangeResultError_Generic>())
@@ -101,7 +103,7 @@ namespace UE::Interchange
 			return;
 		}
 		
-		if (!FbxParserPrivate->LoadFbxFile(Filename))
+		if (!FbxParserPrivate->LoadFbxFile(Filename, BaseNodecontainer))
 		{
 			if (UInterchangeResultError_Generic* Error = AddMessage<UInterchangeResultError_Generic>())
 			{
diff --git a/Engine/Source/Runtime/Interchange/Core/Private/Nodes/InterchangeSourceNode.cpp b/Engine/Source/Runtime/Interchange/Core/Private/Nodes/InterchangeSourceNode.cpp
--- a/Engine/Source/Runtime/Interchange/Core/Private/Nodes/InterchangeSourceNode.cpp
+++ b/Engine/Source/Runtime/Interchange/Core/Private/Nodes/InterchangeSourceNode.cpp
@@ -147,4 +147,14 @@ bool UInterchangeSourceNode::RemoveExtraInformation(const FString& Name)
 void UInterchangeSourceNode::GetExtraInformation(TMap<FString, FString>& OutExtraInformation) const
 {
 	OutExtraInformation = ExtraInformation.ToMap();
-}
\ No newline at end of file
+}
+
+bool UInterchangeSourceNode::GetCustomAxisConversionInverseTransform(FTransform& AttributeValue) const
+{
+	IMPLEMENT_NODE_ATTRIBUTE_GETTER(AxisConversionInverseTransform, FTransform);
+}
+bool UInterchangeSourceNode::SetCustomAxisConversionInverseTransform(const FTransform& AttributeValue)
+{
+	IMPLEMENT_NODE_ATTRIBUTE_SETTER_NODELEGATE(AxisConversionInverseTransform, FTransform);
+}
+
diff --git a/Engine/Source/Runtime/Interchange/Core/Public/Nodes/InterchangeSourceNode.h b/Engine/Source/Runtime/Interchange/Core/Public/Nodes/InterchangeSourceNode.h
--- a/Engine/Source/Runtime/Interchange/Core/Public/Nodes/InterchangeSourceNode.h
+++ b/Engine/Source/Runtime/Interchange/Core/Public/Nodes/InterchangeSourceNode.h
@@ -119,6 +119,15 @@ public:
 	UFUNCTION(BlueprintCallable, Category = "Interchange | Node | Source")
 	INTERCHANGECORE_API void GetExtraInformation(TMap<FString, FString>& OutExtraInformation) const;
 
+
+	/** Query Axis Conversion Inverse Transform (Primarily used for Socket transform calculations.). */
+	UFUNCTION(BlueprintCallable, Category = "Interchange | Node | Source")
+	INTERCHANGECORE_API bool GetCustomAxisConversionInverseTransform(FTransform& AxisConversionInverseTransform) const;
+
+	/** Set the Axis Conversion Inverse Transform (Primarily used for Socket transform calculations.). */
+	UFUNCTION(BlueprintCallable, Category = "Interchange | Node | Source")
+	INTERCHANGECORE_API bool SetCustomAxisConversionInverseTransform(const FTransform& AxisConversionInverseTransform);
+
 private:
 	const UE::Interchange::FAttributeKey Macro_CustomSourceFrameRateNumeratorKey = UE::Interchange::FAttributeKey(TEXT("SourceFrameRateNumerator"));
 	const UE::Interchange::FAttributeKey Macro_CustomSourceFrameRateDenominatorKey = UE::Interchange::FAttributeKey(TEXT("SourceFrameRateDenominator"));
@@ -127,6 +136,7 @@ private:
 	const UE::Interchange::FAttributeKey Macro_CustomAnimatedTimeStartKey = UE::Interchange::FAttributeKey(TEXT("AnimatedTimeStart"));
 	const UE::Interchange::FAttributeKey Macro_CustomAnimatedTimeEndKey = UE::Interchange::FAttributeKey(TEXT("AnimatedTimeEnd"));
 	const UE::Interchange::FAttributeKey Macro_CustomImportUnusedMaterialKey = UE::Interchange::FAttributeKey(TEXT("ImportUnusedMaterial"));
+	const UE::Interchange::FAttributeKey Macro_CustomAxisConversionInverseTransformKey = UE::Interchange::FAttributeKey(TEXT("AxisConversionInverseTransform"));
 
 	// Extra InformationTo show in the Config Panel.
 	UE::Interchange::TMapAttributeHelper<FString, FString> ExtraInformation; 
