Interchange: Make sure scene import will do a re-import if the source data filename match a scene import data in the specified ContentPath.

#rnx

diff --git a/Engine/Plugins/Interchange/Editor/Source/InterchangeEditor/Private/AssetDefinition_InterchangeSceneImportAsset.cpp b/Engine/Plugins/Interchange/Editor/Source/InterchangeEditor/Private/AssetDefinition_InterchangeSceneImportAsset.cpp
--- a/Engine/Plugins/Interchange/Editor/Source/InterchangeEditor/Private/AssetDefinition_InterchangeSceneImportAsset.cpp
+++ b/Engine/Plugins/Interchange/Editor/Source/InterchangeEditor/Private/AssetDefinition_InterchangeSceneImportAsset.cpp
@@ -96,22 +96,6 @@ namespace MenuExtension_InterchangeSceneImportAsset
 #endif
 	}
 
-	ULevel* GetOriginalLevel(UInterchangeSceneImportAsset* Asset)
-	{
-		TArray<FSoftObjectPath> SceneSoftObjectPaths;
-		Asset->GetSceneSoftObjectPaths(SceneSoftObjectPaths);
-		for (const FSoftObjectPath& SoftObjectPath : SceneSoftObjectPaths)
-		{
-			if (AActor* Actor = Cast<AActor>(SoftObjectPath.TryLoad()))
-			{
-				if (!Actor->IsA<ALevelInstance>())
-				{
-					return Actor->GetLevel();
-				}
-			}
-		}
-		return nullptr;
-	}
 
 	void ExecuteResetOneAsset(UInterchangeSceneImportAsset* Asset, const FString& FilePath)
 	{
@@ -247,7 +231,7 @@ namespace MenuExtension_InterchangeSceneImportAsset
 			ImportAssetParameters.bIsAutomated = GIsAutomationTesting || FApp::IsUnattended() || IsRunningCommandlet() || GIsRunningUnattendedScript;
 			ImportAssetParameters.ReimportAsset = Asset;
 			ImportAssetParameters.ReimportSourceIndex = INDEX_NONE;
-			ImportAssetParameters.ImportLevel = GetOriginalLevel(Asset);
+			ImportAssetParameters.ImportLevel = Asset->GetOriginalLevel();
 
 			TTuple<FAssetImportResultRef, FSceneImportResultRef> ImportResult = InterchangeManager.ImportSceneAsync(FString(), SourceData, ImportAssetParameters);
 
diff --git a/Engine/Plugins/Tests/InterchangeTests/Source/InterchangeTests/Private/InterchangeImportTestStepImport.cpp b/Engine/Plugins/Tests/InterchangeTests/Source/InterchangeTests/Private/InterchangeImportTestStepImport.cpp
--- a/Engine/Plugins/Tests/InterchangeTests/Source/InterchangeTests/Private/InterchangeImportTestStepImport.cpp
+++ b/Engine/Plugins/Tests/InterchangeTests/Source/InterchangeTests/Private/InterchangeImportTestStepImport.cpp
@@ -19,6 +19,30 @@ TTuple<UE::Interchange::FAssetImportResultPtr, UE::Interchange::FSceneImportResu
 	// Empty the destination folder here if requested
 	if (bEmptyDestinationFolderPriorToImport)
 	{
+		for (UObject* AssetObject : Data.ResultObjects)
+		{
+			UPackage* PackageObject = AssetObject->GetPackage();
+			if (!ensure(PackageObject))
+			{
+				continue;
+			}
+			// Mark all objects in the package as garbage, and remove the standalone flag, so that GC can remove the temporary asset later
+			//Also rename them so we dont found them if we re-import the same file at the same place
+			TArray<UObject*> ObjectsInPackage;
+			GetObjectsWithPackage(PackageObject, ObjectsInPackage, true);
+			for (UObject* ObjectInPackage : ObjectsInPackage)
+			{
+				//Do not rename actors
+				if (!ObjectInPackage->IsA<AActor>())
+				{
+					const ERenameFlags RenameFlags = REN_DontCreateRedirectors | REN_NonTransactional | REN_DoNotDirty;
+					ObjectInPackage->Rename(nullptr, nullptr, RenameFlags);
+				}
+				ObjectInPackage->ClearFlags(RF_Standalone | RF_Public);
+				ObjectInPackage->MarkAsGarbage();
+			}
+		}
+
 		Data.ResultObjects.Empty();
 		Data.ImportedAssets.Empty();
 
diff --git a/Engine/Source/Runtime/Interchange/Engine/Private/InterchangeManager.cpp b/Engine/Source/Runtime/Interchange/Engine/Private/InterchangeManager.cpp
--- a/Engine/Source/Runtime/Interchange/Engine/Private/InterchangeManager.cpp
+++ b/Engine/Source/Runtime/Interchange/Engine/Private/InterchangeManager.cpp
@@ -25,6 +25,7 @@
 #include "InterchangeHelper.h"
 #include "InterchangeProjectSettings.h"
 #include "InterchangePythonPipelineBase.h"
+#include "InterchangeSceneImportAsset.h"
 #include "InterchangeSourceData.h"
 #include "InterchangeTaskSystem.h"
 #include "InterchangeTranslatorBase.h"
@@ -1781,17 +1782,27 @@ UInterchangeManager::ImportInternal(const FString& ContentPath, const UInterchan
 	}
 
 	const bool bIsUnattended = FApp::IsUnattended() || FApp::IsGame() || GIsAutomationTesting || ImportAssetParameters.bIsAutomated;
-
 	//If we detect the import will do a re-import of an existing asset, the ReimportAsset parameter should be filled
 	//And we will also convert the import data. With this code this import will be considered a re-import.
-	if (ImportType == UE::Interchange::EImportType::ImportType_Asset && !ImportAssetParameters.ReimportAsset)
+	if (!ImportAssetParameters.ReimportAsset)
 	{
-		FString AssetpackageName = ImportAssetParameters.DestinationName.IsEmpty() ? FPaths::GetCleanFilename(SourceData->GetFilename()) : ImportAssetParameters.DestinationName;
+		auto GetAssetNameFromFilename = [&bImportScene, &SourceData]()
+			{
+				FString AssetpackageName = FPaths::GetBaseFilename(SourceData->GetFilename());
+				if (bImportScene)
+				{
+					AssetpackageName = TEXT("SceneImport_") + AssetpackageName;
+				}
+				return AssetpackageName;
+			};
+		FString AssetpackageName = ImportAssetParameters.DestinationName.IsEmpty() ? GetAssetNameFromFilename() : ImportAssetParameters.DestinationName;
+		
 		FString PackageBasePath = ContentPath;
 		SanitizeNameInline(AssetpackageName, ESanitizeNameTypeFlags::ObjectName | ESanitizeNameTypeFlags::ObjectPath | ESanitizeNameTypeFlags::LongPackage);
 		SanitizeNameInline(PackageBasePath, ESanitizeNameTypeFlags::ObjectPath | ESanitizeNameTypeFlags::LongPackage);
 		FString FullPackagePath = FPaths::Combine(*PackageBasePath, *AssetpackageName);
-		if (!UE::Interchange::FPackageUtils::IsMapPackageAsset(FullPackagePath))
+		const bool bCanConvertAsset = bImportScene || !UE::Interchange::FPackageUtils::IsMapPackageAsset(FullPackagePath);
+		if (bCanConvertAsset)
 		{
 			UPackage* Pkg = FindPackage(nullptr, *FullPackagePath);
 			if (!Pkg)
@@ -1806,14 +1817,19 @@ UInterchangeManager::ImportInternal(const FString& ContentPath, const UInterchan
 				//If we found an asset we can re-import, convert set the param to re-import instead of override
 				if (ExistingAsset && GetRegisteredFactoryClass(ExistingAsset->GetClass()))
 				{
-					bool bConvertToReimport = ImportAssetParameters.bReplaceExisting;
-					if (!bIsUnattended && !ImportAssetParameters.bReplaceExisting)
+					//Scene conversion only support UInterchangeSceneImportAsset
+					//Asset conversion support all registered factory class
+					UInterchangeSceneImportAsset* SceneImportAsset = Cast<UInterchangeSceneImportAsset>(ExistingAsset);
+					if (!bImportScene || SceneImportAsset)
 					{
-						FText OverrideDialogMessage = FText::Format(NSLOCTEXT("InterchangeManager", "OverrideAssetMessage", "This import will override asset '{0}', Do you want to convert this import to a re-import?")
-							, FText::FromString(ExistingAsset->GetFullName()));
-						EAppReturnType::Type DialogResult = FMessageDialog::Open(EAppMsgType::YesNoCancel, OverrideDialogMessage);
-						switch (DialogResult)
+						bool bConvertToReimport = ImportAssetParameters.bReplaceExisting;
+						if (!bIsUnattended && !ImportAssetParameters.bReplaceExisting)
 						{
+							FText OverrideDialogMessage = FText::Format(NSLOCTEXT("InterchangeManager", "OverrideAssetMessage", "This import will override asset '{0}', Do you want to convert this import to a re-import?")
+								, FText::FromString(ExistingAsset->GetFullName()));
+							EAppReturnType::Type DialogResult = FMessageDialog::Open(EAppMsgType::YesNoCancel, OverrideDialogMessage);
+							switch (DialogResult)
+							{
 							case EAppReturnType::Yes:
 							{
 								bConvertToReimport = true;
@@ -1829,20 +1845,32 @@ UInterchangeManager::ImportInternal(const FString& ContentPath, const UInterchan
 								UE_LOG(LogInterchangeEngine, Display, TEXT("User cancel import of file [%s]."), *(SourceData->GetFilename()));
 								return EarlyExit();
 							}
+							}
 						}
-					}
 
-					if (bConvertToReimport)
-					{
-						//Make sure we have the correct flag
-						ExistingAsset->SetFlags(RF_Standalone | RF_Public);
-						FImportAssetParameters* MutableImportAssetParameters = const_cast<FImportAssetParameters*>(&ImportAssetParameters);
-						MutableImportAssetParameters->ReimportAsset = ExistingAsset;
-						UInterchangeAssetImportData* OriginalAssetImportData = UInterchangeAssetImportData::GetFromObject(ImportAssetParameters.ReimportAsset);
-						if (!OriginalAssetImportData)
+						if (bConvertToReimport)
 						{
-							//Convert the object asset import data to interchange
-							ConvertImportData(MutableImportAssetParameters->ReimportAsset, FPaths::GetExtension(SourceData->GetFilename()));
+							//Make sure we have the correct flag
+							ExistingAsset->SetFlags(RF_Standalone | RF_Public);
+							if (FImportAssetParameters* MutableImportAssetParameters = const_cast<FImportAssetParameters*>(&ImportAssetParameters))
+							{
+								MutableImportAssetParameters->ReimportAsset = ExistingAsset;
+								MutableImportAssetParameters->ReimportSourceIndex = INDEX_NONE;
+								MutableImportAssetParameters->ImportLevel = (bImportScene && SceneImportAsset) ? SceneImportAsset->GetOriginalLevel() : nullptr;
+								MutableImportAssetParameters->bReplaceExisting |= bImportScene;
+							}
+							UInterchangeAssetImportData* MutableOriginalAssetImportData = UInterchangeAssetImportData::GetFromObject(ImportAssetParameters.ReimportAsset);
+							if (!MutableOriginalAssetImportData)
+							{
+								//Convert the object asset import data to interchange
+								ConvertImportData(ImportAssetParameters.ReimportAsset, FPaths::GetExtension(SourceData->GetFilename()));
+							}
+
+							//Add an analytic so we can know if a not re-import was turn into a re-import
+							if (FEngineAnalytics::IsAvailable())
+							{
+								Attribs.Add(FAnalyticsEventAttribute(TEXT("Parameters.ConvertToReimport"), true));
+							}
 						}
 					}
 				}
diff --git a/Engine/Source/Runtime/Interchange/Engine/Private/InterchangeSceneImportAsset.cpp b/Engine/Source/Runtime/Interchange/Engine/Private/InterchangeSceneImportAsset.cpp
--- a/Engine/Source/Runtime/Interchange/Engine/Private/InterchangeSceneImportAsset.cpp
+++ b/Engine/Source/Runtime/Interchange/Engine/Private/InterchangeSceneImportAsset.cpp
@@ -9,8 +9,10 @@
 
 #include "Engine/AssetUserData.h"
 #include "Engine/Blueprint.h"
+#include "Engine/Level.h"
 #include "Engine/World.h"
 #include "EngineUtils.h"
+#include "LevelInstance/LevelInstanceActor.h"
 #include "UObject/AssetRegistryTagsContext.h"
 #include "UObject/Package.h"
 
@@ -179,6 +181,23 @@ const TArray<UAssetUserData*>* UInterchangeSceneImportAsset::GetAssetUserDataArr
 #endif // #if WITH_EDITORONLY_DATA
 }
 
+ULevel* UInterchangeSceneImportAsset::GetOriginalLevel() const
+{
+	TArray<FSoftObjectPath> SceneSoftObjectPaths;
+	GetSceneSoftObjectPaths(SceneSoftObjectPaths);
+	for (const FSoftObjectPath& SoftObjectPath : SceneSoftObjectPaths)
+	{
+		if (AActor* Actor = Cast<AActor>(SoftObjectPath.TryLoad()))
+		{
+			if (!Actor->IsA<ALevelInstance>())
+			{
+				return Actor->GetLevel();
+			}
+		}
+	}
+	return nullptr;
+}
+
 void UInterchangeSceneImportAsset::UpdateSceneObjects()
 {
 #if WITH_EDITORONLY_DATA
@@ -302,12 +321,12 @@ const UInterchangeFactoryBaseNode* UInterchangeSceneImportAsset::GetFactoryNode(
 	return nullptr;
 }
 
-void UInterchangeSceneImportAsset::GetSceneSoftObjectPaths(TArray<FSoftObjectPath>& SoftObjectPaths)
+void UInterchangeSceneImportAsset::GetSceneSoftObjectPaths(TArray<FSoftObjectPath>& SoftObjectPaths) const
 {
 #if WITH_EDITORONLY_DATA
 	SoftObjectPaths.Reserve(SceneObjects.Num());
 
-	for (TPair< FSoftObjectPath, FString >& Entry : SceneObjects)
+	for (const TPair< FSoftObjectPath, FString >& Entry : SceneObjects)
 	{
 		SoftObjectPaths.Add(Entry.Key);
 	}
diff --git a/Engine/Source/Runtime/Interchange/Engine/Public/InterchangeSceneImportAsset.h b/Engine/Source/Runtime/Interchange/Engine/Public/InterchangeSceneImportAsset.h
--- a/Engine/Source/Runtime/Interchange/Engine/Public/InterchangeSceneImportAsset.h
+++ b/Engine/Source/Runtime/Interchange/Engine/Public/InterchangeSceneImportAsset.h
@@ -10,6 +10,7 @@
 
 class UInterchangeAssetImportData;
 class UInterchangeFactoryBaseNode;
+class ULevel;
 
 /*
  * Class to hold all the data required to properly re-import a level
@@ -50,6 +51,9 @@ public:
 	INTERCHANGEENGINE_API void RegisterWorldRenameCallbacks();
 #endif
 
+	/** Return the original level we import actor into. Return nullptr if no level was found */
+	INTERCHANGEENGINE_API ULevel* GetOriginalLevel() const;
+
 	/** Updates the SceneObjects cache based on the node container stored in AssetImportData */
 	INTERCHANGEENGINE_API void UpdateSceneObjects();
 
@@ -75,7 +79,7 @@ public:
 	
 	INTERCHANGEENGINE_API const UInterchangeFactoryBaseNode* GetFactoryNode(const FSoftObjectPath& ObjectPath) const;
 
-	INTERCHANGEENGINE_API void GetSceneSoftObjectPaths(TArray<FSoftObjectPath>& SoftObjectPaths);
+	INTERCHANGEENGINE_API void GetSceneSoftObjectPaths(TArray<FSoftObjectPath>& SoftObjectPaths) const;
 
 	INTERCHANGEENGINE_API bool ContainsObject(const FSoftObjectPath& SoftObjectPath) const;
 
