[Interchange] Fix ASAN issue in task parsing

#rnx

diff --git a/Engine/Source/Runtime/Interchange/Engine/Private/Tasks/InterchangeTaskParsing.cpp b/Engine/Source/Runtime/Interchange/Engine/Private/Tasks/InterchangeTaskParsing.cpp
--- a/Engine/Source/Runtime/Interchange/Engine/Private/Tasks/InterchangeTaskParsing.cpp
+++ b/Engine/Source/Runtime/Interchange/Engine/Private/Tasks/InterchangeTaskParsing.cpp
@@ -25,6 +25,19 @@
 #include "UObject/UObjectGlobals.h"
 #include "UObject/WeakObjectPtrTemplates.h"
 
+struct FTaskData
+{
+	FString UniqueID;
+	int32 SourceIndex = INDEX_NONE;
+	bool bIsSceneNode = false;
+	TArray<FString> Dependencies;
+	uint64 GraphEventRef;
+	TArray<uint64> Prerequisites;
+	const UClass* FactoryClass;
+
+	TArray<UInterchangeFactoryBaseNode*, TInlineAllocator<1>> Nodes; // For scenes, we can group multiple nodes into a single task as they are usually very light
+};
+
 /**
  * For the Dependency sort to work the predicate must be transitive ( A > B > C implying A > C).
  * That means we must take into account the whole dependency chain, not just the immediate dependencies.
@@ -33,9 +46,20 @@
  */
 struct FNodeDependencyCache
 {
+	void CacheAccumulatedDependencies(const UInterchangeBaseNodeContainer* NodeContainer, const TArray<FTaskData>& SourceTaskDatas)
+	{
+		CachedDependencies.Reset();
+		for (const FTaskData& TaskData : SourceTaskDatas)
+		{
+			TSet<FString> NodeStack;
+			GetAccumulatedDependencies(NodeContainer, TaskData.UniqueID, NodeStack);
+		}
+	}
+
 	const TSet<FString>& GetAccumulatedDependencies(const UInterchangeBaseNodeContainer* NodeContainer, const FString& NodeID)
 	{
 		TRACE_CPUPROFILER_EVENT_SCOPE(FNodeDependencyCache::GetAccumulatedDependencies)
+		ensure(CachedDependencies.Find(NodeID));
 		TSet<FString> NodeStack;
 		return GetAccumulatedDependencies(NodeContainer, NodeID, NodeStack);
 	}
@@ -105,19 +129,6 @@ void UE::Interchange::FTaskParsing::Execute()
 	check(AsyncHelper.IsValid());
 
 	//Parse each graph and prepare import task data, we will then be able to create all the task with the correct dependencies
-	struct FTaskData
-	{
-		FString UniqueID;
-		int32 SourceIndex = INDEX_NONE;
-		bool bIsSceneNode = false;
-		TArray<FString> Dependencies;
-		uint64 GraphEventRef;
-		TArray<uint64> Prerequisites;
-		const UClass* FactoryClass;
-
-		TArray<UInterchangeFactoryBaseNode*, TInlineAllocator<1>> Nodes; // For scenes, we can group multiple nodes into a single task as they are usually very light
-	};
-
 	TArray<FTaskData> TaskDatas;
 
 	//Avoid creating asset if the asynchronous import is canceled, just create the completion task
@@ -169,6 +180,8 @@ void UE::Interchange::FTaskParsing::Execute()
 
 			{
 				FNodeDependencyCache DependencyCache;
+				//Reserving the cache will allow us to use reference when sorting the dependencies.
+				DependencyCache.CacheAccumulatedDependencies(BaseNodeContainer, SourceTaskDatas);
 
 				//Sort per dependencies
 				auto SortByDependencies =
