[Interchange] Add reimport API support c++, scripting and blueprint.

#rnx

diff --git a/Engine/Source/Runtime/Interchange/Engine/Private/InterchangeManager.cpp b/Engine/Source/Runtime/Interchange/Engine/Private/InterchangeManager.cpp
--- a/Engine/Source/Runtime/Interchange/Engine/Private/InterchangeManager.cpp
+++ b/Engine/Source/Runtime/Interchange/Engine/Private/InterchangeManager.cpp
@@ -1664,6 +1664,63 @@ bool UInterchangeManager::ScriptedImportAssetAsync(const FString& ContentPath, c
 	return InterchangeResult->IsValid();
 }
 
+TOptional<FString> UInterchangeManager::ValidateReimportParameter(UObject* ObjectToReimport, const FImportAssetParameters& ImportAssetParameters, bool bRunSynchronous)
+{
+
+	TArray<FString> Filenames;
+	if (!ObjectToReimport || !CanReimport(ObjectToReimport, Filenames))
+	{
+		return TOptional<FString>();
+	}
+	int32 ReimportIndex = ImportAssetParameters.ReimportSourceIndex == INDEX_NONE ? 0 : ImportAssetParameters.ReimportSourceIndex;
+	if (!Filenames.IsValidIndex(ReimportIndex))
+	{
+		return TOptional<FString>();
+	}
+	if (FImportAssetParameters* MutableImportAssetParameters = const_cast<FImportAssetParameters*>(&ImportAssetParameters))
+	{
+		MutableImportAssetParameters->bRunSynchronous = bRunSynchronous;
+		MutableImportAssetParameters->ReimportAsset = ObjectToReimport;
+	}
+	return Filenames[ReimportIndex];
+}
+
+bool UInterchangeManager::ReimportAsset(UObject* ObjectToReimport, const FImportAssetParameters& ImportAssetParameters, TArray<UObject*>& OutImportedObjects)
+{
+	constexpr bool bRunSynchronousTrue = true;
+	TOptional<FString> Filename = ValidateReimportParameter(ObjectToReimport, ImportAssetParameters, bRunSynchronousTrue);
+	if (!Filename.IsSet())
+	{
+		return false;
+	}
+	UE::Interchange::FScopedSourceData SourceData(Filename.GetValue());
+	UE::Interchange::FAssetImportResultRef InterchangeResult = ImportInternal(FString(), SourceData.GetSourceData(), ImportAssetParameters, UE::Interchange::EImportType::ImportType_Asset).Get<0>();
+	InterchangeResult->WaitUntilDone(ImportAssetParameters.bRunSynchronous);
+	OutImportedObjects = InterchangeResult->GetImportedObjects();
+	return InterchangeResult->IsValid();
+}
+
+UE::Interchange::FAssetImportResultRef UInterchangeManager::ReimportAssetAsync(UObject* ObjectToReimport, const FImportAssetParameters& ImportAssetParameters)
+{
+	UE::Interchange::FAssetImportResultRef ErrorImportResult = UE::Interchange::FAssetImportResultRef { MakeShared< UE::Interchange::FImportResult, ESPMode::ThreadSafe >() };
+	ErrorImportResult->SetDone();
+	constexpr bool bRunSynchronousFalse = false;
+	TOptional<FString> Filename = ValidateReimportParameter(ObjectToReimport, ImportAssetParameters, bRunSynchronousFalse);
+	if (!Filename.IsSet())
+	{
+		return ErrorImportResult;
+	}
+	UE::Interchange::FScopedSourceData SourceData(Filename.GetValue());
+	return ImportInternal(FString(), SourceData.GetSourceData(), ImportAssetParameters, UE::Interchange::EImportType::ImportType_Asset).Get<0>();
+}
+
+bool UInterchangeManager::ScriptedReimportAssetAsync(UObject* ObjectToReimport, const FImportAssetParameters& ImportAssetParameters)
+{
+	UE::Interchange::FAssetImportResultRef InterchangeResult = ReimportAssetAsync(ObjectToReimport, ImportAssetParameters);
+	return InterchangeResult->IsValid();
+}
+
+
 bool UInterchangeManager::ImportScene(const FString& ContentPath, const UInterchangeSourceData* SourceData, const FImportAssetParameters& ImportAssetParameters)
 {
 	ImportAssetParameters.bRunSynchronous = true;
diff --git a/Engine/Source/Runtime/Interchange/Engine/Public/InterchangeManager.h b/Engine/Source/Runtime/Interchange/Engine/Public/InterchangeManager.h
--- a/Engine/Source/Runtime/Interchange/Engine/Public/InterchangeManager.h
+++ b/Engine/Source/Runtime/Interchange/Engine/Public/InterchangeManager.h
@@ -636,6 +636,51 @@ public:
 	UFUNCTION(BlueprintCallable, Category = "Interchange | Import Manager")
 	INTERCHANGEENGINE_API bool ScriptedImportAssetAsync(const FString& ContentPath, const UInterchangeSourceData* SourceData, const FImportAssetParameters& ImportAssetParameters);
 
+	/**
+	 * Call this to start a synchronous asset re-import process.
+	 * This process can re-import many different assets into the game content.
+	 *
+	 * @Param ObjectToReimport - The object to re-import.
+	 * @param ImportAssetParameters - All parameters that need to be passed to the import asset function.
+	 * @return true if the import succeeds, or false otherwise.
+	 * 
+	 * @Note - The interchange manager will by default use the last file use to re-import the ObjectToReimport. If the file doesn't exist
+	 *         and the bAutomated flag is:
+	 *         - true, the function will return false and log a warning
+	 *         - false, a dialog will ask 
+	 * @Note - In blueprint depending on the event you use to start the import its possible to have a deadlock, use the async function if its what you are experimenting
+	 */
+ 	UFUNCTION(BlueprintCallable, Category = "Interchange | Import Manager")
+ 	INTERCHANGEENGINE_API bool ReimportAsset(UObject* ObjectToReimport, const FImportAssetParameters& ImportAssetParameters, TArray<UObject*>& OutImportedObjects);
+
+	/**
+	 * Call this to start an asynchronous asset import process.
+	 * This process can import many different assets into the game content.
+	 *
+	 * @Param ContentPath - The path where the imported assets will be created.
+	 * @Param SourceData - The source data input to translate.
+	 * @param ImportAssetParameters - All parameters that need to be passed to the import asset function.
+	 * @return return an import result which can be use to know when the asynchronous import is terminate.
+	 */
+	INTERCHANGEENGINE_API UE::Interchange::FAssetImportResultRef ReimportAssetAsync(UObject* ObjectToReimport, const FImportAssetParameters& ImportAssetParameters);
+
+	/**
+	 * Call this from blueprint or python to start an asynchronous asset import process.
+	 * This process can import many different assets into the game content.
+	 *
+	 * @Param ObjectToReimport - The object to re-import.
+	 * @param ImportAssetParameters - All parameters that need to be passed to the import asset function.
+	 * @return true if the import was started, or false otherwise.
+	 * 
+	 * @Note - The interchange manager will by default use the last file use to re-import the ObjectToReimport. If the file doesn't exist
+	 *         and the bAutomated flag is:
+	 *         - true, the function will return false and log a warning
+	 *         - false, a dialog will ask 
+	 * @Note - If you need an event for when the import is done use the ImportAssetParameters events.
+	 */
+	UFUNCTION(BlueprintCallable, Category = "Interchange | Import Manager")
+	INTERCHANGEENGINE_API bool ScriptedReimportAssetAsync(UObject* ObjectToReimport, const FImportAssetParameters& ImportAssetParameters);
+
 	/**
 	 * Call this to start a synchronous scene import process.
 	 * This process can import many different assets and their transforms (USceneComponent).
@@ -835,6 +880,7 @@ protected:
 	INTERCHANGEENGINE_API TTuple<UE::Interchange::FAssetImportResultRef, UE::Interchange::FSceneImportResultRef>
 	ImportInternal(const FString& ContentPath, const UInterchangeSourceData* SourceData, const FImportAssetParameters& ImportAssetParameters, const UE::Interchange::EImportType ImportType);
 
+	TOptional<FString> ValidateReimportParameter(UObject* ObjectToReimport, const FImportAssetParameters& ImportAssetParameters, bool bRunSynchronous);
 private:
 
 	static bool bIsCreatingSingleton;
