[Interchange] Make sure all the API user need is expose to script and blueprint.

#rnx

diff --git a/Engine/Plugins/Interchange/Runtime/Source/FactoryNodes/Public/InterchangeAnimSequenceFactoryNode.h b/Engine/Plugins/Interchange/Runtime/Source/FactoryNodes/Public/InterchangeAnimSequenceFactoryNode.h
--- a/Engine/Plugins/Interchange/Runtime/Source/FactoryNodes/Public/InterchangeAnimSequenceFactoryNode.h
+++ b/Engine/Plugins/Interchange/Runtime/Source/FactoryNodes/Public/InterchangeAnimSequenceFactoryNode.h
@@ -75,7 +75,6 @@ public:
 #endif //WITH_EDITOR
 
 	/** Get the class this node creates. */
-	UFUNCTION(BlueprintCallable, Category = "Interchange | Node | AnimSequence")
 	virtual class UClass* GetObjectClass() const override;
 
 public:
diff --git a/Engine/Plugins/Interchange/Runtime/Source/FactoryNodes/Public/InterchangeDecalActorFactoryNode.h b/Engine/Plugins/Interchange/Runtime/Source/FactoryNodes/Public/InterchangeDecalActorFactoryNode.h
--- a/Engine/Plugins/Interchange/Runtime/Source/FactoryNodes/Public/InterchangeDecalActorFactoryNode.h
+++ b/Engine/Plugins/Interchange/Runtime/Source/FactoryNodes/Public/InterchangeDecalActorFactoryNode.h
@@ -12,7 +12,6 @@ class INTERCHANGEFACTORYNODES_API UInterchangeDecalActorFactoryNode : public UIn
 	GENERATED_BODY()
 
 public:
-	UFUNCTION(BlueprintCallable, Category = "Interchange | Node | Decal")
 	UClass* GetObjectClass()const override;
 
 	UFUNCTION(BlueprintCallable, Category = "Interchange | Node | Decal")
diff --git a/Engine/Plugins/Interchange/Runtime/Source/FactoryNodes/Public/InterchangeDecalMaterialFactoryNode.h b/Engine/Plugins/Interchange/Runtime/Source/FactoryNodes/Public/InterchangeDecalMaterialFactoryNode.h
--- a/Engine/Plugins/Interchange/Runtime/Source/FactoryNodes/Public/InterchangeDecalMaterialFactoryNode.h
+++ b/Engine/Plugins/Interchange/Runtime/Source/FactoryNodes/Public/InterchangeDecalMaterialFactoryNode.h
@@ -13,7 +13,6 @@ class INTERCHANGEFACTORYNODES_API UInterchangeDecalMaterialFactoryNode : public
 	GENERATED_BODY()
 
 public:
-	UFUNCTION(BlueprintCallable, Category = "Interchange | Node | Decal")
 	UClass* GetObjectClass()const override;
 
 	UFUNCTION(BlueprintCallable, Category = "Interchange | Node | Decal")
diff --git a/Engine/Plugins/Interchange/Runtime/Source/FactoryNodes/Public/InterchangeLevelFactoryNode.h b/Engine/Plugins/Interchange/Runtime/Source/FactoryNodes/Public/InterchangeLevelFactoryNode.h
--- a/Engine/Plugins/Interchange/Runtime/Source/FactoryNodes/Public/InterchangeLevelFactoryNode.h
+++ b/Engine/Plugins/Interchange/Runtime/Source/FactoryNodes/Public/InterchangeLevelFactoryNode.h
@@ -30,7 +30,6 @@ public:
 	}
 
 	/** Get the class this node want to create */
-	UFUNCTION(BlueprintCallable, Category = "Interchange | Node | LevelFactory")
 	virtual class UClass* GetObjectClass() const override;
 
 	/**
diff --git a/Engine/Plugins/Interchange/Runtime/Source/FactoryNodes/Public/InterchangeLevelSequenceFactoryNode.h b/Engine/Plugins/Interchange/Runtime/Source/FactoryNodes/Public/InterchangeLevelSequenceFactoryNode.h
--- a/Engine/Plugins/Interchange/Runtime/Source/FactoryNodes/Public/InterchangeLevelSequenceFactoryNode.h
+++ b/Engine/Plugins/Interchange/Runtime/Source/FactoryNodes/Public/InterchangeLevelSequenceFactoryNode.h
@@ -30,7 +30,6 @@ public:
 	}
 
 	/** Get the class this node want to create */
-	UFUNCTION(BlueprintCallable, Category = "Interchange | Node | LevelSequenceFactory")
 	virtual class UClass* GetObjectClass() const override;
 
 	/**
diff --git a/Engine/Plugins/Interchange/Runtime/Source/FactoryNodes/Public/InterchangeMaterialFactoryNode.h b/Engine/Plugins/Interchange/Runtime/Source/FactoryNodes/Public/InterchangeMaterialFactoryNode.h
--- a/Engine/Plugins/Interchange/Runtime/Source/FactoryNodes/Public/InterchangeMaterialFactoryNode.h
+++ b/Engine/Plugins/Interchange/Runtime/Source/FactoryNodes/Public/InterchangeMaterialFactoryNode.h
@@ -41,7 +41,6 @@ public:
 
 	virtual FString GetTypeName() const override;
 
-	UFUNCTION(BlueprintCallable, Category = "Interchange | Node | MaterialFactory")
 	virtual class UClass* GetObjectClass() const override;
 
 // Material Inputs
@@ -364,7 +363,6 @@ public:
 
 	virtual FString GetTypeName() const override;
 
-	UFUNCTION(BlueprintCallable, Category = "Interchange | Node | MaterialFactory")
 	virtual class UClass* GetObjectClass() const override;
 
 	UFUNCTION(BlueprintCallable, Category = "Interchange | Node | MaterialFactory")
diff --git a/Engine/Plugins/Interchange/Runtime/Source/FactoryNodes/Public/InterchangePhysicsAssetFactoryNode.h b/Engine/Plugins/Interchange/Runtime/Source/FactoryNodes/Public/InterchangePhysicsAssetFactoryNode.h
--- a/Engine/Plugins/Interchange/Runtime/Source/FactoryNodes/Public/InterchangePhysicsAssetFactoryNode.h
+++ b/Engine/Plugins/Interchange/Runtime/Source/FactoryNodes/Public/InterchangePhysicsAssetFactoryNode.h
@@ -57,7 +57,6 @@ public:
 	}
 
 	/** Get the class this node creates. */
-	UFUNCTION(BlueprintCallable, Category = "Interchange | Node | Skeleton")
 	virtual class UClass* GetObjectClass() const override
 	{
 		ensure(bIsNodeClassInitialized);
diff --git a/Engine/Plugins/Interchange/Runtime/Source/FactoryNodes/Public/InterchangeSceneVariantSetsFactoryNode.h b/Engine/Plugins/Interchange/Runtime/Source/FactoryNodes/Public/InterchangeSceneVariantSetsFactoryNode.h
--- a/Engine/Plugins/Interchange/Runtime/Source/FactoryNodes/Public/InterchangeSceneVariantSetsFactoryNode.h
+++ b/Engine/Plugins/Interchange/Runtime/Source/FactoryNodes/Public/InterchangeSceneVariantSetsFactoryNode.h
@@ -30,7 +30,6 @@ public:
 	}
 
 	/** Get the class this node creates. */
-	UFUNCTION(BlueprintCallable, Category = "Interchange | Node | SceneVariantSetFactory")
 	virtual class UClass* GetObjectClass() const override;
 
 	/**
diff --git a/Engine/Plugins/Interchange/Runtime/Source/FactoryNodes/Public/InterchangeSkeletalMeshFactoryNode.h b/Engine/Plugins/Interchange/Runtime/Source/FactoryNodes/Public/InterchangeSkeletalMeshFactoryNode.h
--- a/Engine/Plugins/Interchange/Runtime/Source/FactoryNodes/Public/InterchangeSkeletalMeshFactoryNode.h
+++ b/Engine/Plugins/Interchange/Runtime/Source/FactoryNodes/Public/InterchangeSkeletalMeshFactoryNode.h
@@ -44,7 +44,6 @@ public:
 	virtual FString GetTypeName() const override;
 
 	/** Get the class this node creates. */
-	UFUNCTION(BlueprintCallable, Category = "Interchange | Node | SkeletalMesh")
 	virtual class UClass* GetObjectClass() const override;
 
 public:
diff --git a/Engine/Plugins/Interchange/Runtime/Source/FactoryNodes/Public/InterchangeSkeletonFactoryNode.h b/Engine/Plugins/Interchange/Runtime/Source/FactoryNodes/Public/InterchangeSkeletonFactoryNode.h
--- a/Engine/Plugins/Interchange/Runtime/Source/FactoryNodes/Public/InterchangeSkeletonFactoryNode.h
+++ b/Engine/Plugins/Interchange/Runtime/Source/FactoryNodes/Public/InterchangeSkeletonFactoryNode.h
@@ -56,7 +56,6 @@ public:
 	}
 
 	/** Get the class this node creates. */
-	UFUNCTION(BlueprintCallable, Category = "Interchange | Node | Skeleton")
 	virtual class UClass* GetObjectClass() const override
 	{
 		ensure(bIsNodeClassInitialized);
diff --git a/Engine/Plugins/Interchange/Runtime/Source/FactoryNodes/Public/InterchangeStaticMeshFactoryNode.h b/Engine/Plugins/Interchange/Runtime/Source/FactoryNodes/Public/InterchangeStaticMeshFactoryNode.h
--- a/Engine/Plugins/Interchange/Runtime/Source/FactoryNodes/Public/InterchangeStaticMeshFactoryNode.h
+++ b/Engine/Plugins/Interchange/Runtime/Source/FactoryNodes/Public/InterchangeStaticMeshFactoryNode.h
@@ -49,7 +49,6 @@ public:
 	virtual FString GetTypeName() const override;
 
 	/** Get the class this node creates. */
-	UFUNCTION(BlueprintCallable, Category = "Interchange | Node | StaticMesh")
 	virtual class UClass* GetObjectClass() const override;
 
 #if WITH_EDITOR
diff --git a/Engine/Plugins/Interchange/Runtime/Source/FactoryNodes/Public/InterchangeTextureFactoryNode.h b/Engine/Plugins/Interchange/Runtime/Source/FactoryNodes/Public/InterchangeTextureFactoryNode.h
--- a/Engine/Plugins/Interchange/Runtime/Source/FactoryNodes/Public/InterchangeTextureFactoryNode.h
+++ b/Engine/Plugins/Interchange/Runtime/Source/FactoryNodes/Public/InterchangeTextureFactoryNode.h
@@ -43,7 +43,6 @@ public:
 	}
 
 	/** Get the class this node creates. */
-	UFUNCTION(BlueprintCallable, Category = "Interchange | Node | Texture")
 	virtual UClass* GetObjectClass() const override
 	{
 		// This function should always be overridden by the sub texture factory node
diff --git a/Engine/Source/Runtime/Interchange/Core/Public/InterchangeFactoryBase.h b/Engine/Source/Runtime/Interchange/Core/Public/InterchangeFactoryBase.h
--- a/Engine/Source/Runtime/Interchange/Core/Public/InterchangeFactoryBase.h
+++ b/Engine/Source/Runtime/Interchange/Core/Public/InterchangeFactoryBase.h
@@ -79,6 +79,7 @@ public:
 	/**
 	 * Return the asset type this factory can create.
 	 */
+	UFUNCTION(BlueprintCallable, Category = "Interchange | Factory")
 	virtual EInterchangeFactoryAssetType GetFactoryAssetType() { return EInterchangeFactoryAssetType::None; }
 
 	/**
@@ -310,6 +311,7 @@ public:
 	 * Fills the OutSourceFilenames array with the list of source files contained in the asset source data.
 	 * Returns true if the operation was successful.
 	 */
+	UFUNCTION(BlueprintCallable, Category = "Interchange | Factory")
 	virtual bool GetSourceFilenames(const UObject* Object, TArray<FString>& OutSourceFilenames) const
 	{
 		return false;
@@ -318,6 +320,7 @@ public:
 	/**
 	 * Sets the object's source at the specified index to the given SourceFileName.
 	 */
+	UFUNCTION(BlueprintCallable, Category = "Interchange | Factory")
 	virtual bool SetSourceFilename(const UObject* Object, const FString& SourceFilename, int32 SourceIndex) const
 	{
 		return false;
diff --git a/Engine/Source/Runtime/Interchange/Core/Public/InterchangePipelineBase.h b/Engine/Source/Runtime/Interchange/Core/Public/InterchangePipelineBase.h
--- a/Engine/Source/Runtime/Interchange/Core/Public/InterchangePipelineBase.h
+++ b/Engine/Source/Runtime/Interchange/Core/Public/InterchangePipelineBase.h
@@ -413,6 +413,7 @@ public:
 	}
 
 	/** Fill the list of all asset this pipeline can create */
+	UFUNCTION(BlueprintCallable, Category = "Interchange | Pipeline")
 	INTERCHANGECORE_API virtual void GetSupportAssetClasses(TArray<UClass*>& PipelineSupportAssetClasses) const
 	{
 		return;
@@ -488,7 +489,10 @@ public:
 	static INTERCHANGECORE_API FName GetResultsPropertyName();
 
 	bool CanEditPropertiesStates() { return bAllowPropertyStatesEdition; }
+
+	UFUNCTION(BlueprintCallable, Category = "Interchange | Pipeline")
 	bool IsReimportContext() { return bIsReimportContext; }
+
 	bool IsShowEssentials() { return bIsShowEssentials; }
 
 #if WITH_EDITOR
@@ -529,6 +533,7 @@ public:
 	 * We will reuse this pipeline when reimporting the asset.
 	 * If false, it's probably a debug helper pipeline that we do not want to save into assets.
 	 */
+	UFUNCTION(BlueprintCallable, Category = "Interchange | Pipeline")
 	INTERCHANGECORE_API virtual bool SupportReimport() const { return true; }
 
 	/*
diff --git a/Engine/Source/Runtime/Interchange/Core/Public/InterchangeTranslatorBase.h b/Engine/Source/Runtime/Interchange/Core/Public/InterchangeTranslatorBase.h
--- a/Engine/Source/Runtime/Interchange/Core/Public/InterchangeTranslatorBase.h
+++ b/Engine/Source/Runtime/Interchange/Core/Public/InterchangeTranslatorBase.h
@@ -63,6 +63,7 @@ class UInterchangeTranslatorBase : public UObject
 public:
 
 	/** return true if the translator can translate the given source data. */
+	UFUNCTION(BlueprintCallable, Category = "Interchange | Translator")
 	INTERCHANGECORE_API virtual bool CanImportSourceData(const UInterchangeSourceData* InSourceData) const;
 
 	/**
@@ -73,12 +74,16 @@ public:
 	INTERCHANGECORE_API virtual bool IsThreadSafe() const { return true; }
 
 	/** Specifies the capabilities of a translator. */
-	virtual EInterchangeTranslatorType GetTranslatorType() const { return EInterchangeTranslatorType::Assets; }
+	UFUNCTION(BlueprintCallable, Category = "Interchange | Translator")
+	INTERCHANGECORE_API virtual EInterchangeTranslatorType GetTranslatorType() const { return EInterchangeTranslatorType::Assets; }
 
 	bool DoesSupportAssetType(EInterchangeTranslatorAssetType AssetType) const { return EnumHasAllFlags(GetSupportedAssetTypes(), AssetType); }
+
+	UFUNCTION(BlueprintCallable, Category = "Interchange | Translator")
 	INTERCHANGECORE_API virtual EInterchangeTranslatorAssetType GetSupportedAssetTypes() const PURE_VIRTUAL(UInterchangeTranslatorBase::GetSupportedAssetTypes, return EInterchangeTranslatorAssetType::None;);
 
 	/** List of formats supported by the translator. Each entry is of the form "ext;Description" where ext is the file extension. */
+	UFUNCTION(BlueprintCallable, Category = "Interchange | Translator")
 	INTERCHANGECORE_API virtual TArray<FString> GetSupportedFormats() const PURE_VIRTUAL(UInterchangeTranslatorBase::GetSupportedExtensions, return TArray<FString>(););
 
 	/**
@@ -127,7 +132,6 @@ public:
 		return nullptr;
 	}
 
-
 	void AddMessage(UInterchangeResult* Item) const
 	{
 		if (ensure(Results) && ensure(SourceData))
@@ -137,7 +141,6 @@ public:
 		}
 	}
 	
-
 	void SetResultsContainer(UInterchangeResultsContainer* InResults)
 	{
 		Results = InResults;
@@ -146,6 +149,7 @@ public:
 	/**
 	 * Get the associated source data for this translator.
 	 */
+	UFUNCTION(BlueprintCallable, Category = "Interchange | Translator")
 	const UInterchangeSourceData* GetSourceData() const
 	{
 		return SourceData;
@@ -154,11 +158,13 @@ public:
 	/**
 	 * Return this translator settings
 	 */
+	UFUNCTION(BlueprintCallable, Category = "Interchange | Translator")
 	INTERCHANGECORE_API virtual UInterchangeTranslatorSettings* GetSettings() const { return nullptr; }
 
 	/**
 	 * Reimport will set the settings if they exist in UInterchangeAssetImportData
 	 */
+	UFUNCTION(BlueprintCallable, Category = "Interchange | Translator")
 	INTERCHANGECORE_API virtual void SetSettings(const UInterchangeTranslatorSettings* InterchangeTranslatorSettings) { }
 
 	UPROPERTY()
diff --git a/Engine/Source/Runtime/Interchange/Core/Public/Nodes/InterchangeBaseNode.h b/Engine/Source/Runtime/Interchange/Core/Public/Nodes/InterchangeBaseNode.h
--- a/Engine/Source/Runtime/Interchange/Core/Public/Nodes/InterchangeBaseNode.h
+++ b/Engine/Source/Runtime/Interchange/Core/Public/Nodes/InterchangeBaseNode.h
@@ -209,11 +209,13 @@ public:
 	/**
 	 * Return the node type name of the class. This is used when reporting errors.
 	 */
+	UFUNCTION(BlueprintCallable, Category = "Interchange | Node")
 	INTERCHANGECORE_API virtual FString GetTypeName() const;
 
 	/**
 	 * Icon name, used to retrieve the brush when we display the node in any UI.
 	 */
+	UFUNCTION(BlueprintCallable, Category = "Interchange | Node")
 	INTERCHANGECORE_API virtual FName GetIconName() const;
 
 #if WITH_EDITOR
@@ -444,7 +446,6 @@ public:
 	 * Set the parent unique ID.
 	 */
 	UE_DEPRECATED(5.1, "This function should not be use - Use UInterchangeBaseNodeContainer::SetNodeParentUid function to create a parent hierarchy. This will ensure the container cache is always in sync")
-	UFUNCTION(BlueprintCallable, Category = "Interchange | Node")
 	INTERCHANGECORE_API bool SetParentUid(const FString& ParentUid);
 
 	/**
diff --git a/Engine/Source/Runtime/Interchange/Core/Public/Nodes/InterchangeFactoryBaseNode.h b/Engine/Source/Runtime/Interchange/Core/Public/Nodes/InterchangeFactoryBaseNode.h
--- a/Engine/Source/Runtime/Interchange/Core/Public/Nodes/InterchangeFactoryBaseNode.h
+++ b/Engine/Source/Runtime/Interchange/Core/Public/Nodes/InterchangeFactoryBaseNode.h
@@ -224,6 +224,7 @@ public:
 	/**
 	 * Return the UClass of the object we represent, so we can find the factory/writer.
 	 */
+	UFUNCTION(BlueprintCallable, Category = "Interchange | Node")
 	INTERCHANGECORE_API virtual class UClass* GetObjectClass() const;
 
 	/**
diff --git a/Engine/Source/Runtime/Interchange/Core/Public/Nodes/InterchangeSourceNode.h b/Engine/Source/Runtime/Interchange/Core/Public/Nodes/InterchangeSourceNode.h
--- a/Engine/Source/Runtime/Interchange/Core/Public/Nodes/InterchangeSourceNode.h
+++ b/Engine/Source/Runtime/Interchange/Core/Public/Nodes/InterchangeSourceNode.h
@@ -48,6 +48,7 @@ public:
 	static INTERCHANGECORE_API UInterchangeSourceNode* FindOrCreateUniqueInstance(UInterchangeBaseNodeContainer* NodeContainer);
 
 	/* This function should be use by pipelines to avoid creating a node. If the unique instance doesn't exist, returns nullptr. */
+	UFUNCTION(BlueprintCallable, Category = "Interchange | Node | source")
 	static INTERCHANGECORE_API const UInterchangeSourceNode* GetUniqueInstance(const UInterchangeBaseNodeContainer* NodeContainer);
 
 	/** Query the source frame rate numerator. */
diff --git a/Engine/Source/Runtime/Interchange/Engine/Public/InterchangeManager.h b/Engine/Source/Runtime/Interchange/Engine/Public/InterchangeManager.h
--- a/Engine/Source/Runtime/Interchange/Engine/Public/InterchangeManager.h
+++ b/Engine/Source/Runtime/Interchange/Engine/Public/InterchangeManager.h
@@ -546,17 +546,20 @@ public:
 	/**
 	 * Returns the list of supported formats for a given translator type.
 	 */
+	UFUNCTION(BlueprintCallable, Category = "Interchange | Import Manager")
 	INTERCHANGEENGINE_API TArray<FString> GetSupportedFormats(const EInterchangeTranslatorType ForTranslatorType) const;
 
 	/**
 	 * Returns the list of formats supporting the specified translator asset type.
 	 */
+	UFUNCTION(BlueprintCallable, Category = "Interchange | Import Manager")
 	INTERCHANGEENGINE_API TArray<FString> GetSupportedAssetTypeFormats(const EInterchangeTranslatorAssetType ForTranslatorAssetType, const EInterchangeTranslatorType ForTranslatorType = EInterchangeTranslatorType::Invalid, bool bStrictMatchTranslatorType = false) const;
 
 	/**
 	 * Returns the list of supported formats for a given Object.
 	 * 
 	 */
+	UFUNCTION(BlueprintCallable, Category = "Interchange | Import Manager")
 	INTERCHANGEENGINE_API TArray<FString> GetSupportedFormatsForObject(const UObject* Object, int32 SourceFileIndex) const;
 
 	/**
@@ -567,6 +570,7 @@ public:
 	 * Otherwise, returns true only if the translator supports level import
 	 * @Note: Temporary until FBX level import is production ready
 	 */
+	UFUNCTION(BlueprintCallable, Category = "Interchange | Import Manager")
 	INTERCHANGEENGINE_API bool CanTranslateSourceData(const UInterchangeSourceData* SourceData, bool bSceneImportOnly = false) const;
 
 	/**
@@ -574,6 +578,7 @@ public:
 	 * @Param Object - The object we want to reimport.
 	 * @Param OutFilenames - An array that is filled with the object's source filenames if the operation is successful.
 	 */
+	UFUNCTION(BlueprintCallable, Category = "Interchange | Import Manager")
 	INTERCHANGEENGINE_API bool CanReimport(const UObject* Object, TArray<FString>& OutFilenames) const;
 
 	/**
@@ -768,11 +773,13 @@ public:
 	 * @Param SourceData - The source data that you want a translator for.
 	 * @return return a matching translator, or nullptr if no translators exist for the source data.
 	 */
+	UFUNCTION(BlueprintCallable, Category = "Interchange | Import Manager")
 	INTERCHANGEENGINE_API UInterchangeTranslatorBase* GetTranslatorForSourceData(const UInterchangeSourceData* SourceData) const;
 
 	/**
 	 * Return true if Interchange is actively importing or exporting, or false otherwise.
 	 */
+	UFUNCTION(BlueprintCallable, Category = "Interchange | Import Manager")
 	INTERCHANGEENGINE_API bool IsInterchangeActive();
 
 	/**
@@ -780,6 +787,7 @@ public:
 	 * If Interchange is active, it will display a notification to let the user know they can cancel the asynchronous import/export
 	 * to be able to complete the operation they requested. (This is called by the exit editor operation.)
 	 */
+	UFUNCTION(BlueprintCallable, Category = "Interchange | Import Manager")
 	INTERCHANGEENGINE_API bool WarnIfInterchangeIsActive();
 
 	/**
@@ -801,6 +809,7 @@ public:
 	 * Return true if the object is being imported, or false otherwise. If the user imports multiple file in the same folder, it's possible to
 	 * have the same asset name in two different files.
 	 */
+	UFUNCTION(BlueprintCallable, Category = "Interchange | Import Manager")
 	INTERCHANGEENGINE_API bool IsObjectBeingImported(UObject* Object) const;
 
 	/**
@@ -854,11 +863,13 @@ protected:
 	 * We use this if the user cancels the work or if the editor exits.
 	 * @note - This is a asynchronous call. Tasks will be completed (canceled) soon.
 	 */
+	UFUNCTION(BlueprintCallable, Category = "Interchange | Import Manager")
 	INTERCHANGEENGINE_API void CancelAllTasks();
 
 	/**
 	 * Wait synchronously until all tasks are done.
 	 */
+	UFUNCTION(BlueprintCallable, Category = "Interchange | Import Manager")
 	INTERCHANGEENGINE_API void WaitUntilAllTasksDone(bool bCancel);
 
 
