[Interchange] fix the regression where we replace bone name space with underscore instead of replacing it with dash like in the legacy fbx importer.

diff --git a/Engine/Source/Runtime/Interchange/Core/Private/InterchangeHelper.cpp b/Engine/Source/Runtime/Interchange/Core/Private/InterchangeHelper.cpp
--- a/Engine/Source/Runtime/Interchange/Core/Private/InterchangeHelper.cpp
+++ b/Engine/Source/Runtime/Interchange/Core/Private/InterchangeHelper.cpp
@@ -10,6 +10,11 @@ namespace UE::Interchange
 
 		while (*InvalidChar)
 		{
+			if (bIsJoint)
+			{
+				//For joint we want to replace any space by a dash
+				OutName.ReplaceCharInline(TCHAR(' '), TCHAR('-'), ESearchCase::CaseSensitive);
+			}
 			OutName.ReplaceCharInline(*InvalidChar, TCHAR('_'), ESearchCase::CaseSensitive);
 			++InvalidChar;
 		}
