Interchange: Change a check to be a log when the asynchelper cannot be release right after we end the import.

#rnx

diff --git a/Engine/Source/Runtime/Interchange/Engine/Private/InterchangeManager.cpp b/Engine/Source/Runtime/Interchange/Engine/Private/InterchangeManager.cpp
--- a/Engine/Source/Runtime/Interchange/Engine/Private/InterchangeManager.cpp
+++ b/Engine/Source/Runtime/Interchange/Engine/Private/InterchangeManager.cpp
@@ -2692,7 +2692,10 @@ void UInterchangeManager::ReleaseAsyncHelper(TWeakPtr<UE::Interchange::FImportAs
 	//Make sure the async helper is destroy, if not destroy its because we are canceling the import and we still have a shared ptr on it
 	{
 		TSharedPtr<UE::Interchange::FImportAsyncHelper, ESPMode::ThreadSafe> AsyncHelperSharedPtr = AsyncHelper.Pin();
-		check(!AsyncHelperSharedPtr.IsValid() || AsyncHelperSharedPtr->bCancel);
+		if (AsyncHelperSharedPtr.IsValid() && !AsyncHelperSharedPtr->bCancel)
+		{
+			UE_LOG(LogInterchangeEngine, Warning, TEXT("Interchange AsyncHelper resource not released properly."));
+		}
 	}
 
 	int32 ImportTaskNumber = ImportTasks.Num() + QueueTaskCount;
