Interchange: Allow removing the ParentNodeUid of a given node;

Fix how it was possible to have a single node count as children from different parents.

diff --git a/Engine/Plugins/Interchange/Runtime/Source/Import/Private/MaterialX/MaterialXUtils/MaterialXSurfaceMaterial.cpp b/Engine/Plugins/Interchange/Runtime/Source/Import/Private/MaterialX/MaterialXUtils/MaterialXSurfaceMaterial.cpp
--- a/Engine/Plugins/Interchange/Runtime/Source/Import/Private/MaterialX/MaterialXUtils/MaterialXSurfaceMaterial.cpp
+++ b/Engine/Plugins/Interchange/Runtime/Source/Import/Private/MaterialX/MaterialXUtils/MaterialXSurfaceMaterial.cpp
@@ -34,7 +34,7 @@ void FMaterialXSurfaceMaterial::Translate(MaterialX::NodePtr SurfaceMaterialNode
 	{
 		ShaderGraphNode = NewObject<UInterchangeShaderGraphNode>(&NodeContainer);
 		ShaderGraphNode->InitializeNode(ShaderGraphNodeUID, SurfaceMaterialNode->getName().c_str(), EInterchangeNodeContainerType::TranslatedAsset);
-		NodeContainer.SetNodeParentUid(ShaderGraphNodeUID, FString{});
+		NodeContainer.ClearNodeParentUid(ShaderGraphNodeUID);
 		bIsShaderGraphInContainer = false;
 	}
 
diff --git a/Engine/Source/Runtime/Interchange/Core/Private/Nodes/InterchangeBaseNodeContainer.cpp b/Engine/Source/Runtime/Interchange/Core/Private/Nodes/InterchangeBaseNodeContainer.cpp
--- a/Engine/Source/Runtime/Interchange/Core/Private/Nodes/InterchangeBaseNodeContainer.cpp
+++ b/Engine/Source/Runtime/Interchange/Core/Private/Nodes/InterchangeBaseNodeContainer.cpp
@@ -31,7 +31,7 @@ FString UInterchangeBaseNodeContainer::AddNode(UInterchangeBaseNode* Node)
 	{
 		return UInterchangeBaseNode::InvalidNodeUid();
 	}
-		
+
 	//Cannot add an node with the same IDs
 	if (Nodes.Contains(NodeUniqueID))
 	{
@@ -141,27 +141,52 @@ UInterchangeFactoryBaseNode* UInterchangeBaseNodeContainer::GetFactoryNode(const
 
 bool UInterchangeBaseNodeContainer::SetNodeParentUid(const FString& NodeUniqueID, const FString& NewParentNodeUid)
 {
+	const bool bClearParent = NewParentNodeUid == UInterchangeBaseNode::InvalidNodeUid();
+
 	if (!Nodes.Contains(NodeUniqueID))
 	{
 		return false;
 	}
-	if (!Nodes.Contains(NewParentNodeUid))
+	if (!bClearParent && !Nodes.Contains(NewParentNodeUid))
 	{
 		return false;
 	}
 	UInterchangeBaseNode* Node = Nodes.FindChecked(NodeUniqueID);
-	
-	PRAGMA_DISABLE_DEPRECATION_WARNINGS
-	Node->SetParentUid(NewParentNodeUid);
-	PRAGMA_ENABLE_DEPRECATION_WARNINGS
 
-	//Update the children cache
-	TArray<FString>& Children = ChildrenCache.FindOrAdd(NewParentNodeUid);
-	Children.Add(NodeUniqueID);
+	// Remove from previous parent
+	const FString PreviousParent = Node->GetParentUid();
+	if (PreviousParent != UInterchangeBaseNode::InvalidNodeUid())
+	{
+		if (TArray<FString>* FoundPreviousChildren = ChildrenCache.Find(PreviousParent))
+		{
+			FoundPreviousChildren->Remove(NodeUniqueID);
+		}
+	}
+
+	// Set new parent
+	if (bClearParent)
+	{
+		Node->RemoveAttribute(UE::Interchange::FBaseNodeStaticData::ParentIDKey().Key);
+	}
+	else
+	{
+		PRAGMA_DISABLE_DEPRECATION_WARNINGS
+		Node->SetParentUid(NewParentNodeUid);
+		PRAGMA_ENABLE_DEPRECATION_WARNINGS
+
+		//Update the children cache
+		TArray<FString>& Children = ChildrenCache.FindOrAdd(NewParentNodeUid);
+		Children.Add(NodeUniqueID);
+	}
 
 	return true;
 }
 
+bool UInterchangeBaseNodeContainer::ClearNodeParentUid(const FString& NodeUniqueID)
+{
+	return SetNodeParentUid(NodeUniqueID, UInterchangeBaseNode::InvalidNodeUid());
+}
+
 int32 UInterchangeBaseNodeContainer::GetNodeChildrenCount(const FString& NodeUniqueID) const
 {
 	if (TArray<FString>* CacheChildrenPtr = GetCachedNodeChildrenUids(NodeUniqueID))
@@ -342,7 +367,7 @@ UInterchangeBaseNode* UInterchangeBaseNodeContainer::GetNodeChildrenInternal(con
 bool UInterchangeBaseNodeContainer::GetIsAncestor(const FString& NodeUniqueID, const FString& AncestorUID) const
 {
 	FString CurrentNodeUID = NodeUniqueID;
-	
+
 	while (CurrentNodeUID != UInterchangeBaseNode::InvalidNodeUid() && Nodes.Contains(CurrentNodeUID))
 	{
 		UInterchangeBaseNode* Node = Nodes.FindChecked(CurrentNodeUID);
diff --git a/Engine/Source/Runtime/Interchange/Core/Public/Nodes/InterchangeBaseNodeContainer.h b/Engine/Source/Runtime/Interchange/Core/Public/Nodes/InterchangeBaseNodeContainer.h
--- a/Engine/Source/Runtime/Interchange/Core/Public/Nodes/InterchangeBaseNodeContainer.h
+++ b/Engine/Source/Runtime/Interchange/Core/Public/Nodes/InterchangeBaseNodeContainer.h
@@ -30,7 +30,7 @@ template <typename T> struct TObjectPtr;
 /**
  * The Interchange UInterchangeBaseNode graph is a format used to feed factories and writers when they import, reimport, and
  * export an asset or scene.
- * 
+ *
  * This container holds a flat list of all nodes that have been translated from the source data.
  * Translators fill this container, and the import/export managers read it to execute the import/export process.
  */
@@ -104,7 +104,7 @@ public:
 	/**
 	 * Recursively traverse the hierarchy starting with the specified node unique ID.
 	 * If the iteration lambda returns true, the iteration will stop. If the iteration lambda returns false, the iteration will continue.
-	 * 
+	 *
 	 * @return - Return true if the iteration was broken, or false otherwise.
 	 */
 	bool BreakableIterateNodeChildren(const FString& NodeUniqueID, TFunctionRef<bool(const UInterchangeBaseNode*)> IterationLambda) const
@@ -165,6 +165,10 @@ public:
 	UFUNCTION(BlueprintCallable, Category = "Interchange | Node Container")
 	INTERCHANGECORE_API bool SetNodeParentUid(const FString& NodeUniqueID, const FString& NewParentNodeUid);
 
+	/* Remove the node's ParentUid, making it into a top-level node */
+	UFUNCTION(BlueprintCallable, Category = "Interchange | Node Container")
+	INTERCHANGECORE_API bool ClearNodeParentUid(const FString& NodeUniqueID);
+
 	/** Get the number of children the node has. */
 	UFUNCTION(BlueprintCallable, Category = "Interchange | Node Container")
 	INTERCHANGECORE_API int32 GetNodeChildrenCount(const FString& NodeUniqueID) const;
