[Interchange] Create API for all the temporary code inject in the 5.5.2 release.

#rnx

#virtualized

diff --git a/Engine/Plugins/Interchange/Runtime/Source/Import/Private/Mesh/InterchangeSkeletalMeshFactory.cpp b/Engine/Plugins/Interchange/Runtime/Source/Import/Private/Mesh/InterchangeSkeletalMeshFactory.cpp
--- a/Engine/Plugins/Interchange/Runtime/Source/Import/Private/Mesh/InterchangeSkeletalMeshFactory.cpp
+++ b/Engine/Plugins/Interchange/Runtime/Source/Import/Private/Mesh/InterchangeSkeletalMeshFactory.cpp
@@ -1096,12 +1096,12 @@ namespace UE
 				FTransform BakeFromRootJointTransfromModifier = RootJointNodeLocalTransform.Inverse() * RootJointNodeGlobalTransform * GlobalOffsetTransform.Inverse(); //GlobalOffseTransform will be added by the BindPoseGlobalTransform when bBakeMeshes && !bRootAncestorOfMeshDependency is used
 
 				{
-					// Temporary 5.5.2: fbx legacy has a special way to bake the skeletal mesh that do not fit the interchange standard.
+					// fbx legacy has a special way to bake the skeletal mesh that do not fit the interchange standard.
 					// This fix the issue with blender armature bone skip if the armature has transform not identity.
 					if(const UInterchangeSourceNode* SourceNode = UInterchangeSourceNode::GetUniqueInstance(Arguments.NodeContainer))
 					{
 						bool bUseSceneNodeGlobalTransform = false;
-						SourceNode->GetBooleanAttribute(TEXT("SkeletalMesh-Bake-UseSceneNodeGlobalTransform"), bUseSceneNodeGlobalTransform);
+						SourceNode->GetCustomUseLegacySkeletalMeshBakeTransform(bUseSceneNodeGlobalTransform);
 						if (bUseSceneNodeGlobalTransform)
 						{
 							BakeFromRootJointTransfromModifier.SetIdentity();
diff --git a/Engine/Plugins/Interchange/Runtime/Source/Nodes/Private/InterchangeAnimationTrackSetNode.cpp b/Engine/Plugins/Interchange/Runtime/Source/Nodes/Private/InterchangeAnimationTrackSetNode.cpp
--- a/Engine/Plugins/Interchange/Runtime/Source/Nodes/Private/InterchangeAnimationTrackSetNode.cpp
+++ b/Engine/Plugins/Interchange/Runtime/Source/Nodes/Private/InterchangeAnimationTrackSetNode.cpp
@@ -305,6 +305,26 @@ bool UInterchangeSkeletalAnimationTrackNode::SetCustomAnimationStopTime(const do
 	IMPLEMENT_NODE_ATTRIBUTE_SETTER_NODELEGATE(AnimationStopTime, double);
 }
 
+bool UInterchangeSkeletalAnimationTrackNode::SetCustomSourceTimelineAnimationStartTime(const double& AttributeValue)
+{
+	IMPLEMENT_NODE_ATTRIBUTE_SETTER_NODELEGATE(SourceTimelineAnimationStartTime, double);
+}
+
+bool UInterchangeSkeletalAnimationTrackNode::GetCustomSourceTimelineAnimationStartTime(double& AttributeValue) const
+{
+	IMPLEMENT_NODE_ATTRIBUTE_GETTER(SourceTimelineAnimationStartTime, double);
+}
+
+bool UInterchangeSkeletalAnimationTrackNode::SetCustomSourceTimelineAnimationStopTime(const double& AttributeValue)
+{
+	IMPLEMENT_NODE_ATTRIBUTE_SETTER_NODELEGATE(SourceTimelineAnimationStopTime, double);
+}
+
+bool UInterchangeSkeletalAnimationTrackNode::GetCustomSourceTimelineAnimationStopTime(double& AttributeValue) const
+{
+	IMPLEMENT_NODE_ATTRIBUTE_GETTER(SourceTimelineAnimationStopTime, double);
+}
+
 void UInterchangeSkeletalAnimationTrackNode::GetSceneNodeAnimationPayloadKeys(TMap<FString, FString>& OutSceneNodeAnimationPayloadKeyUids, TMap<FString, uint8>& OutSceneNodeAnimationPayloadKeyTypes) const
 {
 	OutSceneNodeAnimationPayloadKeyUids = SceneNodeAnimationPayloadKeyUidMap.ToMap();
diff --git a/Engine/Plugins/Interchange/Runtime/Source/Nodes/Public/InterchangeAnimationTrackSetNode.h b/Engine/Plugins/Interchange/Runtime/Source/Nodes/Public/InterchangeAnimationTrackSetNode.h
--- a/Engine/Plugins/Interchange/Runtime/Source/Nodes/Public/InterchangeAnimationTrackSetNode.h
+++ b/Engine/Plugins/Interchange/Runtime/Source/Nodes/Public/InterchangeAnimationTrackSetNode.h
@@ -517,6 +517,39 @@ public:
 	UFUNCTION(BlueprintCallable, Category = "Interchange | Node | SkeletalAnimationTrack")
 	bool GetCustomAnimationStopTime(double& StopTime) const;
 
+	/**
+	 * Set the source timeline animation start time. Return false if the attribute could not be set.
+	 * Note: This represent the DCC setup of the time for the animation.
+	 * Optional: if not set CustomAnimationStartTime will be used
+	 */
+	UFUNCTION(BlueprintCallable, Category = "Interchange | Node | SkeletalAnimationTrack")
+	bool SetCustomSourceTimelineAnimationStartTime(const double& StartTime);
+
+	/**
+	 * Get the source timeline animation start time. Return false if the attribute is not set.
+	 * Note: This represent the DCC setup of the time for the animation.
+	 * Optional: if not set CustomAnimationStartTime will be used
+	 */
+	UFUNCTION(BlueprintCallable, Category = "Interchange | Node | SkeletalAnimationTrack")
+	bool GetCustomSourceTimelineAnimationStartTime(double& StartTime) const;
+	
+	/**
+	 * Set the source timeline animation stop time. Return false if the attribute could not be set.
+	 * Note: This represent the DCC setup of the time for the animation.
+	 * Optional: if not set CustomAnimationStopTime will be used
+	 */
+	UFUNCTION(BlueprintCallable, Category = "Interchange | Node | SkeletalAnimationTrack")
+	bool SetCustomSourceTimelineAnimationStopTime(const double& StopTime);
+
+	/**
+	 * Get the source timeline animation stop time. Return false if the attribute is not set.
+	 * Note: This represent the DCC setup of the time for the animation.
+	 * Optional: if not set CustomAnimationStopTime will be used
+	 */
+	UFUNCTION(BlueprintCallable, Category = "Interchange | Node | SkeletalAnimationTrack")
+	bool GetCustomSourceTimelineAnimationStopTime(double& StopTime) const;
+
+
 	UFUNCTION(BlueprintCallable, Category = "Interchange | Node | SkeletalAnimationTrack")
 	void GetSceneNodeAnimationPayloadKeys(TMap<FString, FString>& OutSceneNodeAnimationPayloadKeyUids, TMap<FString, uint8>& OutSceneNodeAnimationPayloadKeyTypes) const;
 
@@ -537,6 +570,8 @@ private:
 	const UE::Interchange::FAttributeKey Macro_CustomAnimationSampleRateKey = UE::Interchange::FAttributeKey(TEXT("AnimationSampleRate"));
 	const UE::Interchange::FAttributeKey Macro_CustomAnimationStartTimeKey = UE::Interchange::FAttributeKey(TEXT("AnimationStartTime"));
 	const UE::Interchange::FAttributeKey Macro_CustomAnimationStopTimeKey = UE::Interchange::FAttributeKey(TEXT("AnimationStopTime"));
+	const UE::Interchange::FAttributeKey Macro_CustomSourceTimelineAnimationStartTimeKey = UE::Interchange::FAttributeKey(TEXT("SourceTimelineAnimationStartTime"));
+	const UE::Interchange::FAttributeKey Macro_CustomSourceTimelineAnimationStopTimeKey = UE::Interchange::FAttributeKey(TEXT("SourceTimelineAnimationStopTime"));
 	
 	UE::Interchange::TMapAttributeHelper<FString, FString>	SceneNodeAnimationPayloadKeyUidMap;
 	UE::Interchange::TMapAttributeHelper<FString, uint8>	SceneNodeAnimationPayloadKeyTypeMap;
diff --git a/Engine/Plugins/Interchange/Runtime/Source/Parsers/Fbx/Private/FbxAPI.cpp b/Engine/Plugins/Interchange/Runtime/Source/Parsers/Fbx/Private/FbxAPI.cpp
--- a/Engine/Plugins/Interchange/Runtime/Source/Parsers/Fbx/Private/FbxAPI.cpp
+++ b/Engine/Plugins/Interchange/Runtime/Source/Parsers/Fbx/Private/FbxAPI.cpp
@@ -157,10 +157,10 @@ namespace UE
 					SourceNode->SetCustomSourceFrameRateDenominator(Denominator);
 				}
 
-				// Temporary 5.5.2: fbx legacy has a special way to bake the skeletal mesh that do not fit the interchange standard
+				// Fbx legacy has a special way to bake the skeletal mesh that do not fit the interchange standard
 				// The interchange skeletal mesh factory will read this to use the proper bake transform so it match legacy behavior.
 				// This fix the issue with blender armature bone skip
-				SourceNode->AddBooleanAttribute(TEXT("SkeletalMesh-Bake-UseSceneNodeGlobalTransform"), true);
+				SourceNode->SetCustomUseLegacySkeletalMeshBakeTransform(true);
 
 				// Get the version number of the FBX file format.
 				int32 FileMajor, FileMinor, FileRevision;
diff --git a/Engine/Plugins/Interchange/Runtime/Source/Parsers/Fbx/Private/FbxScene.cpp b/Engine/Plugins/Interchange/Runtime/Source/Parsers/Fbx/Private/FbxScene.cpp
--- a/Engine/Plugins/Interchange/Runtime/Source/Parsers/Fbx/Private/FbxScene.cpp
+++ b/Engine/Plugins/Interchange/Runtime/Source/Parsers/Fbx/Private/FbxScene.cpp
@@ -901,9 +901,9 @@ namespace UE
 							FbxTimeSpan AnimStackInterval(FBXSDK_TIME_INFINITE, FBXSDK_TIME_MINUS_INFINITE);
 							AnimStackInterval = CurrentAnimationStack->GetLocalTimeSpan();
 
-							//Temporary 5.5.2: This should be replace by an API in UInterchangeSkeletalAnimationTrackNode after the hotfix is done
-							SkeletalAnimationTrackNode->AddDoubleAttribute(TEXT("SourceTimelineAnimationStartTime"), AnimStackInterval.GetStart().GetSecondDouble());
-							SkeletalAnimationTrackNode->AddDoubleAttribute(TEXT("SourceTimelineAnimationStopTime"), AnimStackInterval.GetStop().GetSecondDouble());
+							//AnimStackInterval represent the animation time line setup in the DCC
+							SkeletalAnimationTrackNode->SetCustomSourceTimelineAnimationStartTime(AnimStackInterval.GetStart().GetSecondDouble());
+							SkeletalAnimationTrackNode->SetCustomSourceTimelineAnimationStopTime(AnimStackInterval.GetStop().GetSecondDouble());
 							return true;
 						}
 
diff --git a/Engine/Plugins/Interchange/Runtime/Source/Pipelines/Private/InterchangeGenericAnimationPipeline.cpp b/Engine/Plugins/Interchange/Runtime/Source/Pipelines/Private/InterchangeGenericAnimationPipeline.cpp
--- a/Engine/Plugins/Interchange/Runtime/Source/Pipelines/Private/InterchangeGenericAnimationPipeline.cpp
+++ b/Engine/Plugins/Interchange/Runtime/Source/Pipelines/Private/InterchangeGenericAnimationPipeline.cpp
@@ -664,8 +664,7 @@ void UInterchangeGenericAnimationPipeline::CreateAnimSequenceFactoryNode(UInterc
 		//Some format don't fill the animation range data we fall back on the trackNode
 		if (AnimationRange == EInterchangeAnimationRange::Timeline)
 		{
-			//Temporary 5.5.2: This should be replace by an API in UInterchangeSkeletalAnimationTrackNode after the hotfix is done
-			bTimeRangeIsValid = TrackNode.GetDoubleAttribute(TEXT("SourceTimelineAnimationStartTime"), StartTime) && TrackNode.GetDoubleAttribute(TEXT("SourceTimelineAnimationStopTime"), StopTime);
+			bTimeRangeIsValid = TrackNode.GetCustomSourceTimelineAnimationStartTime(StartTime) && TrackNode.GetCustomSourceTimelineAnimationStopTime(StopTime);
 		}
 		else if (AnimationRange == EInterchangeAnimationRange::SetRange)
 		{
diff --git a/Engine/Source/Runtime/Interchange/Core/Private/Nodes/InterchangeSourceNode.cpp b/Engine/Source/Runtime/Interchange/Core/Private/Nodes/InterchangeSourceNode.cpp
--- a/Engine/Source/Runtime/Interchange/Core/Private/Nodes/InterchangeSourceNode.cpp
+++ b/Engine/Source/Runtime/Interchange/Core/Private/Nodes/InterchangeSourceNode.cpp
@@ -158,3 +158,12 @@ bool UInterchangeSourceNode::SetCustomAxisConversionInverseTransform(const FTran
 	IMPLEMENT_NODE_ATTRIBUTE_SETTER_NODELEGATE(AxisConversionInverseTransform, FTransform);
 }
 
+bool UInterchangeSourceNode::GetCustomUseLegacySkeletalMeshBakeTransform(bool& AttributeValue) const
+{
+	IMPLEMENT_NODE_ATTRIBUTE_GETTER(UseLegacySkeletalMeshBakeTransform, bool);
+}
+bool UInterchangeSourceNode::SetCustomUseLegacySkeletalMeshBakeTransform(const bool& AttributeValue)
+{
+	IMPLEMENT_NODE_ATTRIBUTE_SETTER_NODELEGATE(UseLegacySkeletalMeshBakeTransform, bool);
+}
+
diff --git a/Engine/Source/Runtime/Interchange/Core/Public/Nodes/InterchangeSourceNode.h b/Engine/Source/Runtime/Interchange/Core/Public/Nodes/InterchangeSourceNode.h
--- a/Engine/Source/Runtime/Interchange/Core/Public/Nodes/InterchangeSourceNode.h
+++ b/Engine/Source/Runtime/Interchange/Core/Public/Nodes/InterchangeSourceNode.h
@@ -69,34 +69,42 @@ public:
 
 	/** Query the start of the source timeline. */
 	UFUNCTION(BlueprintCallable, Category = "Interchange | Node | Source")
+	UE_DEPRECATED(5.6, "Please, use UInterchangeSkeletalAnimationTrackNode::GetCustomSourceTimelineAnimationStartTime")
 	INTERCHANGECORE_API bool GetCustomSourceTimelineStart(double& AttributeValue) const;
 
 	/** Set the start of the source timeline. */
 	UFUNCTION(BlueprintCallable, Category = "Interchange | Node | Source")
+	UE_DEPRECATED(5.6, "Please, use UInterchangeSkeletalAnimationTrackNode::SetCustomSourceTimelineAnimationStartTime")
 	INTERCHANGECORE_API bool SetCustomSourceTimelineStart(const double& AttributeValue);
 
 	/** Query the end of the source timeline. */
 	UFUNCTION(BlueprintCallable, Category = "Interchange | Node | Source")
+	UE_DEPRECATED(5.6, "Please, use UInterchangeSkeletalAnimationTrackNode::GetCustomSourceTimelineAnimationStopTime")
 	INTERCHANGECORE_API bool GetCustomSourceTimelineEnd(double& AttributeValue) const;
 
 	/** Set the end of the source timeline. */
 	UFUNCTION(BlueprintCallable, Category = "Interchange | Node | Source")
+	UE_DEPRECATED(5.6, "Please, use UInterchangeSkeletalAnimationTrackNode::SetCustomSourceTimelineAnimationStopTime")
 	INTERCHANGECORE_API bool SetCustomSourceTimelineEnd(const double& AttributeValue);
 
 	/** Query the start of the source animated time. */
 	UFUNCTION(BlueprintCallable, Category = "Interchange | Node | Source")
+	UE_DEPRECATED(5.6, "Please, use UInterchangeSkeletalAnimationTrackNode::GetCustomAnimationStartTime")
 	INTERCHANGECORE_API bool GetCustomAnimatedTimeStart(double& AttributeValue) const;
 
 	/** Set the start of the source animated time. */
 	UFUNCTION(BlueprintCallable, Category = "Interchange | Node | Source")
+	UE_DEPRECATED(5.6, "Please, use UInterchangeSkeletalAnimationTrackNode::SetCustomAnimationStartTime")
 	INTERCHANGECORE_API bool SetCustomAnimatedTimeStart(const double& AttributeValue);
 
 	/** Query the end of the source animated time. */
 	UFUNCTION(BlueprintCallable, Category = "Interchange | Node | Source")
+	UE_DEPRECATED(5.6, "Please, use UInterchangeSkeletalAnimationTrackNode::GetCustomAnimationStopTime")
 	INTERCHANGECORE_API bool GetCustomAnimatedTimeEnd(double& AttributeValue) const;
 
 	/** Set the end of the source animated time. */
 	UFUNCTION(BlueprintCallable, Category = "Interchange | Node | Source")
+	UE_DEPRECATED(5.6, "Please, use UInterchangeSkeletalAnimationTrackNode::SetCustomAnimationStopTime")
 	INTERCHANGECORE_API bool SetCustomAnimatedTimeEnd(const double& AttributeValue);
 
 	/** Query whether to import materials that aren't used. */
@@ -129,6 +137,15 @@ public:
 	UFUNCTION(BlueprintCallable, Category = "Interchange | Node | Source")
 	INTERCHANGECORE_API bool SetCustomAxisConversionInverseTransform(const FTransform& AxisConversionInverseTransform);
 
+
+	/** Does skeletalMesh factory should uses legacy bake transform behavior to create the skeletal mesh. */
+	UFUNCTION(BlueprintCallable, Category = "Interchange | Node | Source")
+	INTERCHANGECORE_API bool GetCustomUseLegacySkeletalMeshBakeTransform(bool& AttributeValue) const;
+
+	/** Set the SkeletalMesh factory to uses legacy bake transform behavior to create the skeletalmesh. */
+	UFUNCTION(BlueprintCallable, Category = "Interchange | Node | Source")
+	INTERCHANGECORE_API bool SetCustomUseLegacySkeletalMeshBakeTransform(const bool& AttributeValue);
+
 private:
 	const UE::Interchange::FAttributeKey Macro_CustomSourceFrameRateNumeratorKey = UE::Interchange::FAttributeKey(TEXT("SourceFrameRateNumerator"));
 	const UE::Interchange::FAttributeKey Macro_CustomSourceFrameRateDenominatorKey = UE::Interchange::FAttributeKey(TEXT("SourceFrameRateDenominator"));
@@ -138,6 +155,7 @@ private:
 	const UE::Interchange::FAttributeKey Macro_CustomAnimatedTimeEndKey = UE::Interchange::FAttributeKey(TEXT("AnimatedTimeEnd"));
 	const UE::Interchange::FAttributeKey Macro_CustomImportUnusedMaterialKey = UE::Interchange::FAttributeKey(TEXT("ImportUnusedMaterial"));
 	const UE::Interchange::FAttributeKey Macro_CustomAxisConversionInverseTransformKey = UE::Interchange::FAttributeKey(TEXT("AxisConversionInverseTransform"));
+	const UE::Interchange::FAttributeKey Macro_CustomUseLegacySkeletalMeshBakeTransformKey = UE::Interchange::FAttributeKey(TEXT("UseLegacySkeletalMeshBakeTransform"));
 
 	// Extra InformationTo show in the Config Panel.
 	UE::Interchange::TMapAttributeHelper<FString, FString> ExtraInformation; 
