[Interchange] Node container basic features added.
- Add Node Namespace attribute to be able to instantiate many time the same container into an existing one. This is needed for external reference feature.
- Add children ordering capability via a desired index.

- Add system test for those two features

#rnx

diff --git a/Engine/Source/Runtime/Interchange/Core/Private/Nodes/InterchangeBaseNode.cpp b/Engine/Source/Runtime/Interchange/Core/Private/Nodes/InterchangeBaseNode.cpp
--- a/Engine/Source/Runtime/Interchange/Core/Private/Nodes/InterchangeBaseNode.cpp
+++ b/Engine/Source/Runtime/Interchange/Core/Private/Nodes/InterchangeBaseNode.cpp
@@ -26,6 +26,12 @@ namespace UE
 			static FAttributeKey AttributeKey(TEXT("__DSPL_LBL_"));
 			return AttributeKey;
 		}
+		
+		const FAttributeKey& FBaseNodeStaticData::NamespaceKey()
+		{
+			static FAttributeKey AttributeKey(TEXT("__Namespace_"));
+			return AttributeKey;
+		}
 
 		const FAttributeKey& FBaseNodeStaticData::ParentIDKey()
 		{
@@ -33,6 +39,12 @@ namespace UE
 			return AttributeKey;
 		}
 
+		const FAttributeKey& FBaseNodeStaticData::DesiredChildIndexKey()
+		{
+			static FAttributeKey AttributeKey(TEXT("__DESIRED_CHILD_INDEX_"));
+			return AttributeKey;
+		}
+
 		const FAttributeKey& FBaseNodeStaticData::IsEnabledKey()
 		{
 			static FAttributeKey AttributeKey(TEXT("__IS_NBLD_"));
@@ -314,6 +326,11 @@ FString UInterchangeBaseNode::GetUniqueID() const
 	ensure(bIsInitialized);
 	FString UniqueID;
 	Attributes->GetAttributeHandle<FString>(UE::Interchange::FBaseNodeStaticData::UniqueIDKey()).Get(UniqueID);
+	FString Namespace;
+	if (GetNamespace(Namespace))
+	{
+		UniqueID = Namespace + UniqueID;
+	}
 	return UniqueID;
 }
 
@@ -337,6 +354,40 @@ bool UInterchangeBaseNode::SetDisplayLabel(const FString& DisplayLabel)
 	return false;
 }
 
+bool UInterchangeBaseNode::GetNamespace(FString& OutNamespace) const
+{
+	if (!Attributes->ContainAttribute(UE::Interchange::FBaseNodeStaticData::NamespaceKey()))
+	{
+		return false;
+	}
+	UE::Interchange::FAttributeStorage::TAttributeHandle<FString> Handle = Attributes->GetAttributeHandle<FString>(UE::Interchange::FBaseNodeStaticData::NamespaceKey());
+	if (Handle.IsValid())
+	{
+		UE::Interchange::EAttributeStorageResult Result = Handle.Get(OutNamespace);
+		return IsAttributeStorageResultSuccess(Result);
+	}
+	return false;
+}
+
+bool UInterchangeBaseNode::SetNamespace(const FString& InNamespace)
+{
+	if(InNamespace.IsEmpty())
+	{
+		if(Attributes->ContainAttribute(UE::Interchange::FBaseNodeStaticData::NamespaceKey()))
+		{
+			Attributes->UnregisterAttribute(UE::Interchange::FBaseNodeStaticData::NamespaceKey());
+		}
+		return !Attributes->ContainAttribute(UE::Interchange::FBaseNodeStaticData::NamespaceKey());
+	}
+	UE::Interchange::EAttributeStorageResult Result = Attributes->RegisterAttribute(UE::Interchange::FBaseNodeStaticData::NamespaceKey(), InNamespace);
+	if (IsAttributeStorageResultSuccess(Result))
+	{
+		UE::Interchange::FAttributeStorage::TAttributeHandle<FString> Handle = Attributes->GetAttributeHandle<FString>(UE::Interchange::FBaseNodeStaticData::NamespaceKey());
+		return Handle.IsValid();
+	}
+	return false;
+}
+
 FString UInterchangeBaseNode::GetParentUid() const
 {
 	if (!Attributes->ContainAttribute(UE::Interchange::FBaseNodeStaticData::ParentIDKey()))
@@ -365,6 +416,34 @@ bool UInterchangeBaseNode::SetParentUid(const FString& ParentUid)
 	return false;
 }
 
+int32 UInterchangeBaseNode::GetDesiredChildIndex() const
+{
+	if (!Attributes->ContainAttribute(UE::Interchange::FBaseNodeStaticData::DesiredChildIndexKey()))
+	{
+		return INDEX_NONE;
+	}
+
+	int32 ChildIndex;
+	UE::Interchange::FAttributeStorage::TAttributeHandle<int32> Handle = Attributes->GetAttributeHandle<int32>(UE::Interchange::FBaseNodeStaticData::DesiredChildIndexKey());
+	if (Handle.IsValid())
+	{
+		Handle.Get(ChildIndex);
+		return ChildIndex;
+	}
+	return INDEX_NONE;
+}
+
+bool UInterchangeBaseNode::SetDesiredChildIndex(const int32& ChildIndex)
+{
+	UE::Interchange::EAttributeStorageResult Result = Attributes->RegisterAttribute(UE::Interchange::FBaseNodeStaticData::DesiredChildIndexKey(), ChildIndex);
+	if (IsAttributeStorageResultSuccess(Result))
+	{
+		UE::Interchange::FAttributeStorage::TAttributeHandle<int32> Handle = Attributes->GetAttributeHandle<int32>(UE::Interchange::FBaseNodeStaticData::DesiredChildIndexKey());
+		return Handle.IsValid();
+	}
+	return false;
+}
+
 bool UInterchangeBaseNode::IsEnabled() const
 {
 	if (!Attributes->ContainAttribute(UE::Interchange::FBaseNodeStaticData::IsEnabledKey()))
diff --git a/Engine/Source/Runtime/Interchange/Core/Private/Nodes/InterchangeBaseNodeContainer.cpp b/Engine/Source/Runtime/Interchange/Core/Private/Nodes/InterchangeBaseNodeContainer.cpp
--- a/Engine/Source/Runtime/Interchange/Core/Private/Nodes/InterchangeBaseNodeContainer.cpp
+++ b/Engine/Source/Runtime/Interchange/Core/Private/Nodes/InterchangeBaseNodeContainer.cpp
@@ -22,6 +22,7 @@ UInterchangeBaseNodeContainer::UInterchangeBaseNodeContainer()
 
 FString UInterchangeBaseNodeContainer::AddNode(UInterchangeBaseNode* Node)
 {
+	TRACE_CPUPROFILER_EVENT_SCOPE(UInterchangeBaseNodeContainer::AddNode)
 	if (!Node)
 	{
 		return UInterchangeBaseNode::InvalidNodeUid();
@@ -53,6 +54,7 @@ FString UInterchangeBaseNodeContainer::AddNode(UInterchangeBaseNode* Node)
 
 void UInterchangeBaseNodeContainer::ReplaceNode(const FString& NodeUniqueID, UInterchangeFactoryBaseNode* NewNode)
 {
+	TRACE_CPUPROFILER_EVENT_SCOPE(UInterchangeBaseNodeContainer::ReplaceNode)
 	if (GetFactoryNode(NodeUniqueID)) //Check existance and confirm it is FactoryNode
 	{
 		Nodes.Remove(NodeUniqueID);
@@ -69,6 +71,76 @@ bool UInterchangeBaseNodeContainer::IsNodeUidValid(const FString& NodeUniqueID)
 	return Nodes.Contains(NodeUniqueID);
 }
 
+void UInterchangeBaseNodeContainer::SetNamespace(const FString& Namespace, UClass* TargetClass)
+{
+	TRACE_CPUPROFILER_EVENT_SCOPE(UInterchangeBaseNodeContainer::SetNamespace)
+	TMap<FString, FString> UniqueIdSwap;
+	UniqueIdSwap.Reserve(Nodes.Num());
+	//Change the asset unique ID
+	for (auto& NodeKeyValue : Nodes)
+	{
+		UInterchangeBaseNode* Node = NodeKeyValue.Value;
+		if(TargetClass && !Node->GetClass()->IsChildOf(TargetClass))
+		{
+			continue;
+		}
+		const FString NodeUniqueId = Node->GetUniqueID();
+		Node->SetNamespace(Namespace);
+		UniqueIdSwap.Add(NodeUniqueId, Node->GetUniqueID());
+	}
+
+	//Update all Attributes
+	for (auto& NodeKeyValue : Nodes)
+	{
+		UInterchangeBaseNode* Node = NodeKeyValue.Value;
+		const FString NodeUniqueId = Node->GetUniqueID();
+		TArray<UE::Interchange::FAttributeKey> AttributeKeys;
+		Node->GetAttributeKeys(AttributeKeys);
+		for (const UE::Interchange::FAttributeKey& AttributeKey : AttributeKeys)
+		{
+			if (Node->GetAttributeType(AttributeKey) != UE::Interchange::EAttributeTypes::String)
+			{
+				continue;
+			}
+			if (AttributeKey == UE::Interchange::FBaseNodeStaticData::UniqueIDKey())
+			{
+				continue;
+			}
+			FString AttributeValue;
+			if (!Node->GetStringAttribute(AttributeKey.Key, AttributeValue))
+			{
+				continue;
+			}
+			for(const TPair<FString, FString>& OldAndNewUid : UniqueIdSwap)
+			{
+				//Replace any reference to the new unique ids.
+				if (AttributeValue.Equals(OldAndNewUid.Key))
+				{
+					Node->AddStringAttribute(AttributeKey.Key, OldAndNewUid.Value);
+					break;
+				}
+			}
+		}
+	}
+
+	//Update the container keys. Remove and re-add each node we remap.
+	for (const TPair<FString, FString>& OldUidToNewUid : UniqueIdSwap)
+	{
+		if (OldUidToNewUid.Key == UInterchangeBaseNode::InvalidNodeUid())
+		{
+			continue;
+		}
+		UInterchangeBaseNode* ToReplaceNode = Nodes.FindRef(OldUidToNewUid.Key);
+		if (!ToReplaceNode)
+		{
+			continue;
+		}
+		
+		Nodes.Remove(OldUidToNewUid.Key);
+		AddNode(ToReplaceNode);
+	}
+}
+
 void UInterchangeBaseNodeContainer::IterateNodes(TFunctionRef<void(const FString&, UInterchangeBaseNode*)> IterationLambda) const
 {
 	for (auto& NodeKeyValue : Nodes)
@@ -101,6 +173,7 @@ void UInterchangeBaseNodeContainer::GetRoots(TArray<FString>& RootNodes) const
 
 void UInterchangeBaseNodeContainer::GetNodes(const UClass* ClassNode, TArray<FString>& OutNodes) const
 {
+	TRACE_CPUPROFILER_EVENT_SCOPE(UInterchangeBaseNodeContainer::GetNodes)
 	OutNodes.Empty();
 	IterateNodes([&ClassNode, &OutNodes](const FString& NodeUid, UInterchangeBaseNode* Node)
 	{
@@ -113,45 +186,41 @@ void UInterchangeBaseNodeContainer::GetNodes(const UClass* ClassNode, TArray<FSt
 
 const UInterchangeBaseNode* UInterchangeBaseNodeContainer::GetNode(const FString& NodeUniqueID) const
 {
+	TRACE_CPUPROFILER_EVENT_SCOPE(UInterchangeBaseNodeContainer::GetNode)
 	if (NodeUniqueID == UInterchangeBaseNode::InvalidNodeUid())
 	{
 		return nullptr;
 	}
-	if (!Nodes.Contains(NodeUniqueID))
-	{
-		return nullptr;
-	}
-	UInterchangeBaseNode* Node = Nodes.FindChecked(NodeUniqueID);
+	UInterchangeBaseNode* Node = Nodes.FindRef(NodeUniqueID);
 	return Node;
 }
 
 UInterchangeFactoryBaseNode* UInterchangeBaseNodeContainer::GetFactoryNode(const FString& NodeUniqueID) const
 {
+	TRACE_CPUPROFILER_EVENT_SCOPE(UInterchangeBaseNodeContainer::GetFactoryNode)
 	if (NodeUniqueID == UInterchangeBaseNode::InvalidNodeUid())
 	{
 		return nullptr;
 	}
-	if (!Nodes.Contains(NodeUniqueID))
-	{
-		return nullptr;
-	}
-	UInterchangeFactoryBaseNode* FactoryNode = Cast< UInterchangeFactoryBaseNode>(Nodes.FindChecked(NodeUniqueID));
+	UInterchangeFactoryBaseNode* FactoryNode = Cast< UInterchangeFactoryBaseNode>(Nodes.FindRef(NodeUniqueID));
 	return FactoryNode;
 }
 
 bool UInterchangeBaseNodeContainer::SetNodeParentUid(const FString& NodeUniqueID, const FString& NewParentNodeUid)
 {
+	TRACE_CPUPROFILER_EVENT_SCOPE(UInterchangeBaseNodeContainer::SetNodeParentUid)
 	const bool bClearParent = NewParentNodeUid == UInterchangeBaseNode::InvalidNodeUid();
 
-	if (!Nodes.Contains(NodeUniqueID))
+	UInterchangeBaseNode* Node = Nodes.FindRef(NodeUniqueID);
+	if (!Node)
 	{
 		return false;
 	}
+
 	if (!bClearParent && !Nodes.Contains(NewParentNodeUid))
 	{
 		return false;
 	}
-	UInterchangeBaseNode* Node = Nodes.FindChecked(NodeUniqueID);
 
 	// Remove from previous parent
 	const FString PreviousParent = Node->GetParentUid();
@@ -170,9 +239,7 @@ bool UInterchangeBaseNodeContainer::SetNodeParentUid(const FString& NodeUniqueID
 	}
 	else
 	{
-		PRAGMA_DISABLE_DEPRECATION_WARNINGS
 		Node->SetParentUid(NewParentNodeUid);
-		PRAGMA_ENABLE_DEPRECATION_WARNINGS
 
 		//Update the children cache
 		TArray<FString>& Children = ChildrenCache.FindOrAdd(NewParentNodeUid);
@@ -187,6 +254,30 @@ bool UInterchangeBaseNodeContainer::ClearNodeParentUid(const FString& NodeUnique
 	return SetNodeParentUid(NodeUniqueID, UInterchangeBaseNode::InvalidNodeUid());
 }
 
+bool UInterchangeBaseNodeContainer::SetNodeDesiredChildIndex(const FString& NodeUniqueID, const int32& NewNodeChildIndex)
+{
+	TRACE_CPUPROFILER_EVENT_SCOPE(UInterchangeBaseNodeContainer::SetNodeDesiredChildIndex)
+	UInterchangeBaseNode* Node = Nodes.FindRef(NodeUniqueID);
+	if (!Node)
+	{
+		return false;
+	}
+	int32 OldNodeChildIndex = Node->GetDesiredChildIndex();
+	if (OldNodeChildIndex != NewNodeChildIndex)
+	{
+		Node->SetDesiredChildIndex(NewNodeChildIndex);
+		FString ParentNodeUid = Node->GetParentUid();
+		if (ParentNodeUid != UInterchangeBaseNode::InvalidNodeUid())
+		{
+			if (TArray<FString>* ParentChildrenCache = ChildrenCache.Find(ParentNodeUid))
+			{
+				InternalReorderChildren(*ParentChildrenCache);
+			}
+		}
+	}
+	return true;
+}
+
 int32 UInterchangeBaseNodeContainer::GetNodeChildrenCount(const FString& NodeUniqueID) const
 {
 	if (TArray<FString>* CacheChildrenPtr = GetCachedNodeChildrenUids(NodeUniqueID))
@@ -200,6 +291,7 @@ int32 UInterchangeBaseNodeContainer::GetNodeChildrenCount(const FString& NodeUni
 
 TArray<FString> UInterchangeBaseNodeContainer::GetNodeChildrenUids(const FString& NodeUniqueID) const
 {
+	TRACE_CPUPROFILER_EVENT_SCOPE(UInterchangeBaseNodeContainer::GetNodeChildrenUids)
 	if (TArray<FString>* CacheChildrenPtr = GetCachedNodeChildrenUids(NodeUniqueID))
 	{
 		return *CacheChildrenPtr;
@@ -214,6 +306,8 @@ TArray<FString> UInterchangeBaseNodeContainer::GetNodeChildrenUids(const FString
 			CacheChildren.Add(NodeKeyValue.Key);
 		}
 	}
+	//Reorder the children cache
+	InternalReorderChildren(CacheChildren);
 	return CacheChildren;
 }
 
@@ -303,6 +397,7 @@ void UInterchangeBaseNodeContainer::SerializeNodeContainerData(FArchive& Ar)
 
 void UInterchangeBaseNodeContainer::SaveToFile(const FString& Filename)
 {
+	TRACE_CPUPROFILER_EVENT_SCOPE(UInterchangeBaseNodeContainer::SaveToFile)
 	FLargeMemoryWriter Ar;
 	SerializeNodeContainerData(Ar);
 	uint8* ArchiveData = Ar.GetData();
@@ -313,6 +408,7 @@ void UInterchangeBaseNodeContainer::SaveToFile(const FString& Filename)
 
 void UInterchangeBaseNodeContainer::LoadFromFile(const FString& Filename)
 {
+	TRACE_CPUPROFILER_EVENT_SCOPE(UInterchangeBaseNodeContainer::LoadFromFile)
 	//All sub object should be gone with the reset
 	Nodes.Reset();
 	TArray64<uint8> Buffer;
@@ -331,6 +427,7 @@ void UInterchangeBaseNodeContainer::LoadFromFile(const FString& Filename)
 
 void UInterchangeBaseNodeContainer::ComputeChildrenCache()
 {
+	TRACE_CPUPROFILER_EVENT_SCOPE(UInterchangeBaseNodeContainer::ComputeChildrenCache)
 	ResetChildrenCache();
 	for (const auto& NodeKeyValue : Nodes)
 	{
@@ -345,6 +442,53 @@ void UInterchangeBaseNodeContainer::ComputeChildrenCache()
 			Children.Add(NodeKeyValue.Key);
 		}
 	}
+	for (TPair<FString, TArray<FString>>& ParentChildrenCache : ChildrenCache)
+	{
+		InternalReorderChildren(ParentChildrenCache.Value);
+	}
+}
+
+void UInterchangeBaseNodeContainer::InternalReorderChildren(TArray<FString>& Children) const
+{
+	TRACE_CPUPROFILER_EVENT_SCOPE(UInterchangeBaseNodeContainer::InternalReorderChildren)
+	bool bSortChildren = false;
+	const int32 ChildrenCount = Children.Num();
+	TMap<int32, FString> NodeChildIndexCache;
+	NodeChildIndexCache.Reserve(ChildrenCount);
+	for(int32 ChildIndex = 0; ChildIndex < ChildrenCount; ++ChildIndex)
+	{
+		const FString& ChildUid = Children[ChildIndex];
+		const UInterchangeBaseNode* Node = GetNode(ChildUid);
+		int32 ChildDesiredIndex = Node ? Node->GetDesiredChildIndex() : INDEX_NONE;
+		if (ChildDesiredIndex != INDEX_NONE)
+		{
+			bSortChildren = true;
+		}
+		else
+		{
+			//Keep the order for the INDEX_NONE item
+			ChildDesiredIndex = ChildrenCount + ChildIndex;
+		}
+		NodeChildIndexCache.FindOrAdd(ChildDesiredIndex) = ChildUid;
+	}
+
+	//If no indexes was specified, do not sort the children
+	if (!bSortChildren)
+	{
+		return;
+	}
+
+	NodeChildIndexCache.KeySort([](const int32& ChildA, const int32& ChildB)
+	{
+		return ChildA < ChildB;
+	});
+
+	Children.Reset();
+
+	for (const TPair<int32, FString>& ChildData : NodeChildIndexCache)
+	{
+		Children.Add(ChildData.Value);
+	}
 }
 
 UInterchangeBaseNode* UInterchangeBaseNodeContainer::GetNodeChildrenInternal(const FString& NodeUniqueID, int32 ChildIndex)
@@ -355,22 +499,20 @@ UInterchangeBaseNode* UInterchangeBaseNodeContainer::GetNodeChildrenInternal(con
 		return nullptr;
 	}
 
-	if (Nodes.Contains(ChildrenUids[ChildIndex]))
-	{
-		UInterchangeBaseNode* Node = Nodes.FindChecked(ChildrenUids[ChildIndex]);
-		return Node;
-	}
-
-	return nullptr;
+	return Nodes.FindRef(ChildrenUids[ChildIndex]);
 }
 
 bool UInterchangeBaseNodeContainer::GetIsAncestor(const FString& NodeUniqueID, const FString& AncestorUID) const
 {
 	FString CurrentNodeUID = NodeUniqueID;
 
-	while (CurrentNodeUID != UInterchangeBaseNode::InvalidNodeUid() && Nodes.Contains(CurrentNodeUID))
+	while (CurrentNodeUID != UInterchangeBaseNode::InvalidNodeUid())
 	{
-		UInterchangeBaseNode* Node = Nodes.FindChecked(CurrentNodeUID);
+		UInterchangeBaseNode* Node = Nodes.FindRef(CurrentNodeUID);
+		if (!Node)
+		{
+			break;
+		}
 		const FString& ParentUID = Node->GetParentUid();
 
 		if (AncestorUID == ParentUID)
diff --git a/Engine/Source/Runtime/Interchange/Core/Private/Tests/NodeContainerTests.cpp b/Engine/Source/Runtime/Interchange/Core/Private/Tests/NodeContainerTests.cpp
new file mode 100644
--- /dev/null
+++ b/Engine/Source/Runtime/Interchange/Core/Private/Tests/NodeContainerTests.cpp
@@ -0,0 +1,262 @@
+// Copyright Epic Games, Inc. All Rights Reserved.
+
+#include "CoreMinimal.h"
+#include "Async/ParallelFor.h"
+#include "Misc/AutomationTest.h"
+#include "Nodes/InterchangeBaseNode.h"
+#include "Nodes/InterchangeBaseNodeContainer.h"
+#include "Templates/Atomic.h"
+#include "Types/AttributeStorage.h"
+#include "UObject/UObjectGlobals.h"
+
+#if WITH_DEV_AUTOMATION_TESTS
+
+IMPLEMENT_SIMPLE_AUTOMATION_TEST(FNodeContainerTest, "System.Runtime.Interchange.NodeContainer", EAutomationTestFlags::EditorContext | EAutomationTestFlags::EngineFilter)
+
+bool FNodeContainerTest::RunTest(const FString& Parameters)
+{
+	using namespace UE::Interchange;
+
+	//Create a node container
+	UInterchangeBaseNodeContainer* NodeContainer = NewObject<UInterchangeBaseNodeContainer>();
+	if (!NodeContainer)
+	{
+		AddError(TEXT("Cannot create a UInterchangeNodeContainer object."));
+		return false;
+	}
+	//Add a couple of translated asset Nodes
+	const FString TranslatedAssetNodePrefix = TEXT("TranslatedAssetNode_");
+	const int32 TranslatedAssetNodeCount = 10;
+	for (int32 NodeIndex = 0; NodeIndex < TranslatedAssetNodeCount; ++NodeIndex)
+	{
+		UInterchangeBaseNode* Node = NewObject<UInterchangeBaseNode>();
+		FString NodeUniqueID = TranslatedAssetNodePrefix + FString::FromInt(NodeIndex);
+		Node->InitializeNode(NodeUniqueID, NodeUniqueID, EInterchangeNodeContainerType::TranslatedAsset);
+		NodeContainer->AddNode(Node);
+	}
+
+	//Add a couple of translated scene Nodes
+	const FString TranslatedSceneNodePrefix = TEXT("TranslatedSceneNode_");
+	const FString TranslatedRootNodeUid = TranslatedSceneNodePrefix + FString::FromInt(0);
+	const int32 TranslatedSceneNodeCount = 100;
+	const int32 ChildCount = 9;
+	{
+		int32 CurrentParentIndex = 0;
+		FString CurrentParentUid = TranslatedSceneNodePrefix + FString::FromInt(0);
+		for(int32 NodeIndex = 0; NodeIndex < TranslatedSceneNodeCount; ++NodeIndex)
+		{
+			UInterchangeBaseNode* Node = NewObject<UInterchangeBaseNode>();
+			FString NodeUniqueID = TranslatedSceneNodePrefix + FString::FromInt(NodeIndex);
+			Node->InitializeNode(NodeUniqueID, NodeUniqueID, EInterchangeNodeContainerType::TranslatedScene);
+			NodeContainer->AddNode(Node);
+			if (CurrentParentIndex != NodeIndex)
+			{
+				NodeContainer->SetNodeParentUid(NodeUniqueID, CurrentParentUid);
+			}
+			int32 ComputeParentIndex = (NodeIndex % ChildCount == 0) ? NodeIndex : CurrentParentIndex;
+			if (ComputeParentIndex != CurrentParentIndex)
+			{
+				//We have a new parent
+				CurrentParentIndex = ComputeParentIndex;
+				CurrentParentUid = NodeUniqueID;
+			}
+		}
+	}
+
+	//Add some factory node
+	const FString FactoryNodePrefix = TEXT("FactoryNode_");
+	const int32 FactoryNodeCount = 3;
+	{
+		for (int32 NodeIndex = 0; NodeIndex < FactoryNodeCount; ++NodeIndex)
+		{
+			UInterchangeBaseNode* Node = NewObject<UInterchangeBaseNode>();
+			FString NodeUniqueID = FactoryNodePrefix + FString::FromInt(NodeIndex);
+			Node->InitializeNode(NodeUniqueID, NodeUniqueID, EInterchangeNodeContainerType::FactoryData);
+			NodeContainer->AddNode(Node);
+			Node->AddTargetNodeUid(TranslatedRootNodeUid);
+		}
+	}
+
+	TArray<FString> TranslatedAssetNodes;
+	TArray<FString> TranslatedSceneNodes;
+	TArray<FString> FactoryNodes;
+	TArray<UInterchangeBaseNode*> FactoryNodesPtr;
+	{
+		bool bUnknownNodeType = false;
+		NodeContainer->IterateNodes([&TranslatedAssetNodes, &TranslatedSceneNodes, &FactoryNodes, &FactoryNodesPtr, &bUnknownNodeType](const FString& NodeUid, UInterchangeBaseNode* Node)
+		{
+			switch(Node->GetNodeContainerType())
+			{
+				case EInterchangeNodeContainerType::TranslatedAsset:
+				{
+					TranslatedAssetNodes.Add(NodeUid);
+				}
+				break;
+			
+				case EInterchangeNodeContainerType::TranslatedScene:
+				{
+					TranslatedSceneNodes.Add(NodeUid);
+				}
+				break;
+
+				case EInterchangeNodeContainerType::FactoryData:
+				{
+					FactoryNodes.Add(NodeUid);
+					FactoryNodesPtr.Add(Node);
+				}
+				break;
+
+				default:
+				{
+					bUnknownNodeType = true;
+				}
+			}
+		});
+
+		TestEqual(TEXT("Node container translated asset node count"), TranslatedAssetNodes.Num(), TranslatedAssetNodeCount);
+		TestEqual(TEXT("Node container translated scene node count"), TranslatedSceneNodes.Num(), TranslatedSceneNodeCount);
+		TestEqual(TEXT("Node container translated asset node count"), FactoryNodes.Num(), FactoryNodeCount);
+		TestFalse(TEXT("Node container contains unknown node"), bUnknownNodeType);
+	}
+	
+	//Test root node
+	{
+		TArray<FString> RootNodes;
+		NodeContainer->GetRoots(RootNodes);
+		TestEqual(TEXT("Node container root count"), RootNodes.Num(), 14);
+	}
+
+	//Children caches tests
+	{
+		TArray<FString>* RootChildrenUidsPtr = NodeContainer->GetCachedNodeChildrenUids(TranslatedRootNodeUid);
+		if(RootChildrenUidsPtr)
+		{
+			//Copy the array
+			TArray<FString> RootChildrenUids = *RootChildrenUidsPtr;
+			//Test children cache feature
+			{
+				TestEqual(TEXT("Node container root node children count"), RootChildrenUids.Num(), ChildCount);
+				for (int32 ChildIndex = 0; ChildIndex < RootChildrenUids.Num(); ++ChildIndex)
+				{
+					FString ExpectedChildName = TranslatedSceneNodePrefix + FString::FromInt(ChildIndex + 1);
+					if (!TestEqual(TEXT("Node container child unique id doesn't match"), RootChildrenUids[ChildIndex], ExpectedChildName))
+					{
+						//Log only one error
+						break;
+					}
+				}
+			}
+
+			//Test children index feature
+			{
+				TArray<int32> RemapChildren;
+				const int32 RootChildCount = RootChildrenUids.Num();
+				RemapChildren.AddZeroed(RootChildCount);
+				for (int32 ChildIndex = 0; ChildIndex < RootChildCount; ++ChildIndex)
+				{
+					RemapChildren[ChildIndex] = RootChildCount - ChildIndex - 1;
+					NodeContainer->SetNodeDesiredChildIndex(RootChildrenUids[ChildIndex], RemapChildren[ChildIndex]);
+				}
+				auto VerifyCacheChildrenReorder = [this, &NodeContainer, &RootChildrenUids, &RemapChildren, &TranslatedRootNodeUid]()
+				{
+					TArray<FString>* ReorderedChildrenPtr = NodeContainer->GetCachedNodeChildrenUids(TranslatedRootNodeUid);
+					if(ReorderedChildrenPtr)
+					{
+						TArray<FString> ReorderedChildren = *ReorderedChildrenPtr;
+						if(RootChildrenUids.Num() == ReorderedChildren.Num())
+						{
+							for (int32 ChildIndex = 0; ChildIndex < ReorderedChildren.Num(); ++ChildIndex)
+							{
+								if (!TestEqual(TEXT("Node container child index reorder fail"), ReorderedChildren[ChildIndex], RootChildrenUids[RemapChildren[ChildIndex]]))
+								{
+									//Log only one error
+									break;
+								}
+							}
+						}
+						else
+						{
+							AddError(TEXT("Node container child index feature fail, the number of child change when indexes are specified."));
+						}
+					}
+					else
+					{
+						AddError(TEXT("Node container child index feature fail, Cannot get the cache node childrenUids."));
+					}
+				};
+
+				//Verify before re-computing the cache
+				VerifyCacheChildrenReorder();
+
+				NodeContainer->ComputeChildrenCache();
+
+				//Verify after the cache was reset and recompute
+				VerifyCacheChildrenReorder();
+			}
+		}
+		else
+		{
+			AddError(TEXT("Node container child index feature fail, Cannot get the cache node childrenUids."));
+		}
+	}
+
+	//Test namespace feature
+	{
+		const FString Namespace = TEXT("Foo");
+		const FString NamespaceAndUniqueId = Namespace + TranslatedRootNodeUid;
+		const UInterchangeBaseNode* TranslatedNode = nullptr;
+
+		auto TestFactoryNodeTargetUid = [this, &FactoryNodesPtr, &NamespaceAndUniqueId, &TranslatedRootNodeUid](bool bWithNamespace)
+		{
+			for (UInterchangeBaseNode* Node : FactoryNodesPtr)
+			{
+				if (Node)
+				{
+					TArray<FString> TargetNodeUids;
+					Node->GetTargetNodeUids(TargetNodeUids);
+					if (TargetNodeUids.IsValidIndex(0))
+					{
+						const FString TestUid = bWithNamespace ? NamespaceAndUniqueId : TranslatedRootNodeUid;
+						if(!TestEqual(TEXT("Namespace should update node reference into all node attributes"), TestUid, TargetNodeUids[0]))
+						{
+							break;
+						}
+					}
+				}
+			}
+		};
+
+		//Test without the namespace
+		TranslatedNode = NodeContainer->GetNode(TranslatedRootNodeUid);
+		TestTrue(TEXT("UInterchangeNodeContainer::GetNode() should return the node if the unique id is pass without the namesapce, if the container dont have any namespace"), TranslatedNode != nullptr);
+		TranslatedNode = nullptr;
+		TranslatedNode = NodeContainer->GetNode(NamespaceAndUniqueId);
+		TestTrue(TEXT("UInterchangeNodeContainer::GetNode() should return null node if the namespace and the unique id are combine, if the container dont have any namespace"), TranslatedNode == nullptr);
+		TestFactoryNodeTargetUid(false);
+
+		//Test with the namespace
+		NodeContainer->SetNamespace(Namespace, nullptr);
+		TranslatedNode = nullptr;
+		TranslatedNode = NodeContainer->GetNode(TranslatedRootNodeUid);
+		TestTrue(TEXT("UInterchangeNodeContainer::GetNode() should return null node if the unique id is pass without the namesapce, if the container have a namespace"), TranslatedNode == nullptr);
+		TranslatedNode = nullptr;
+		TranslatedNode = NodeContainer->GetNode(NamespaceAndUniqueId);
+		TestTrue(TEXT("UInterchangeNodeContainer::GetNode() should return the node if the namespace and the unique id are combine, if the container have a namespace"), TranslatedNode != nullptr);
+		TestFactoryNodeTargetUid(true);
+		
+		//Test with the namespace removed
+		NodeContainer->SetNamespace(FString(), nullptr);
+		TranslatedNode = nullptr;
+		TranslatedNode = NodeContainer->GetNode(TranslatedRootNodeUid);
+		TestTrue(TEXT("UInterchangeNodeContainer::GetNode() should return the node if the unique id is pass without the namesapce, if the container dont have any namespace"), TranslatedNode != nullptr);
+		TranslatedNode = nullptr;
+		TranslatedNode = NodeContainer->GetNode(NamespaceAndUniqueId);
+		TestTrue(TEXT("UInterchangeNodeContainer::GetNode() should return null node if the namespace and the unique id are combine, if the container dont have any namespace"), TranslatedNode == nullptr);
+		TestFactoryNodeTargetUid(false);
+	}
+
+	return true;
+}
+
+
+#endif //WITH_DEV_AUTOMATION_TESTS
diff --git a/Engine/Source/Runtime/Interchange/Core/Public/Nodes/InterchangeBaseNode.h b/Engine/Source/Runtime/Interchange/Core/Public/Nodes/InterchangeBaseNode.h
--- a/Engine/Source/Runtime/Interchange/Core/Public/Nodes/InterchangeBaseNode.h
+++ b/Engine/Source/Runtime/Interchange/Core/Public/Nodes/InterchangeBaseNode.h
@@ -32,6 +32,7 @@
 class FArchive;
 class FAssetRegistryTagsContext;
 class UInterchangeBaseNode;
+class UInterchangeBaseNodeContainer;
 struct FFrame;
 
 /**
@@ -156,7 +157,9 @@ namespace UE::Interchange
 	{
 		static INTERCHANGECORE_API const FAttributeKey& UniqueIDKey();
 		static INTERCHANGECORE_API const FAttributeKey& DisplayLabelKey();
+		static INTERCHANGECORE_API const FAttributeKey& NamespaceKey();
 		static INTERCHANGECORE_API const FAttributeKey& ParentIDKey();
+		static INTERCHANGECORE_API const FAttributeKey& DesiredChildIndexKey();
 		static INTERCHANGECORE_API const FAttributeKey& IsEnabledKey();
 		static INTERCHANGECORE_API const FAttributeKey& TargetAssetIDsKey();
 		static INTERCHANGECORE_API const FAttributeKey& ClassTypeAttributeKey();
@@ -436,17 +439,30 @@ public:
 	UFUNCTION(BlueprintCallable, Category = "Interchange | Node")
 	INTERCHANGECORE_API bool SetDisplayLabel(const FString& DisplayName);
 
+	/**
+	 * Namespace is a prefix added when calling GetUniqueId().
+	 * If the attribute exist it will fill Namespace and return true.
+	 * If the attribute do not exist, it will not touch OutNamespace and return false.
+	 */
+	UFUNCTION(BlueprintCallable, Category = "Interchange | Node")
+	INTERCHANGECORE_API bool GetNamespace(FString& Namespace) const;
+
 	/**
 	 * Return the parent unique ID. If the attribute does not exist, returns InvalidNodeUid().
+	 * 
+	 * Use UInterchangeBaseNodeContainer::SetNodeParentUid function to create a parent hierarchy. This will ensure the container cache is always in sync
 	 */
 	UFUNCTION(BlueprintCallable, Category = "Interchange | Node")
 	INTERCHANGECORE_API FString GetParentUid() const;
 
 	/**
-	 * Set the parent unique ID.
+	 * Return the desired child index. If the attribute do not exist it will return INDEX_NONE.
+	 * The interchange node container will use this to re-order children if attribute are specified.
+	 * 
+	 * Use UInterchangeBaseNodeContainer::SetNodeDesiredChildIndex function to set this node child index. This will ensure the container cache is always in sync
 	 */
-	UE_DEPRECATED(5.1, "This function should not be use - Use UInterchangeBaseNodeContainer::SetNodeParentUid function to create a parent hierarchy. This will ensure the container cache is always in sync")
-	INTERCHANGECORE_API bool SetParentUid(const FString& ParentUid);
+	UFUNCTION(BlueprintCallable, Category = "Interchange | Node")
+	INTERCHANGECORE_API int32 GetDesiredChildIndex() const;
 
 	/**
 	 * Get the number of target assets relating to this object.
@@ -541,6 +557,23 @@ public:
 	}
 
 protected:
+
+	/**
+	 * Return true if the Namespace attribute is added, false otherwise.
+	 */
+	bool SetNamespace(const FString& InNamespace);
+
+	/**
+	 * Set the desired child index. Only the container should set the desired child index
+	 */
+	bool SetDesiredChildIndex(const int32& ChildIndex);
+
+	/**
+	 * Set the parent unique ID.
+	 * Use UInterchangeBaseNodeContainer::SetNodeParentUid function to create a parent hierarchy. This will ensure the container cache is always in sync
+	 */
+	bool SetParentUid(const FString& ParentUid);
+
 	/** The storage use to store the Key value attribute for this node. */
 	TSharedPtr<UE::Interchange::FAttributeStorage, ESPMode::ThreadSafe> Attributes;
 
@@ -550,4 +583,6 @@ protected:
 	 * This tracks the IDs of asset nodes which are the target of factories
 	 */
 	mutable UE::Interchange::TArrayAttributeHelper<FString> TargetNodes;
+
+	friend UInterchangeBaseNodeContainer;
 };
diff --git a/Engine/Source/Runtime/Interchange/Core/Public/Nodes/InterchangeBaseNodeContainer.h b/Engine/Source/Runtime/Interchange/Core/Public/Nodes/InterchangeBaseNodeContainer.h
--- a/Engine/Source/Runtime/Interchange/Core/Public/Nodes/InterchangeBaseNodeContainer.h
+++ b/Engine/Source/Runtime/Interchange/Core/Public/Nodes/InterchangeBaseNodeContainer.h
@@ -68,6 +68,20 @@ public:
 	UFUNCTION(BlueprintCallable, Category = "Interchange | Node Container")
 	INTERCHANGECORE_API bool IsNodeUidValid(const FString& NodeUniqueID) const;
 
+	/**
+	 * Set a namespace to all node of the target class in this container. A valid node namespace is prefix to the unique ID in UInterchangeBaseNode::GetUniqueId().
+	 * After adding the namespace this function will find any string attribute in all the node that reference the node unique ID and replace the attribute value with the new unique id.
+	 * The last step is to remap the node container with the new Ids.
+	 * 
+	 * @Param Namespace - Is the new namespace you want to set. Pass an empty string to remove an existing namespace.
+	 * @Param TargetClass - Optional, this parameter represent the node class we want to apply the namespace on. If null all node will be tagged with the namespace
+	 * 
+	 * @Note - Changing all node namespace wont work since some node class use a combinaison of the unique ID in there attributes and we cannot change those attribute
+	 *         in a generic way.
+	 */
+	UFUNCTION(BlueprintCallable, Category = "Interchange | Node Container")
+	INTERCHANGECORE_API void SetNamespace(const FString& Namespace, UClass* TargetClass);
+
 	/** Unordered iteration of the all nodes. */
 	INTERCHANGECORE_API void IterateNodes(TFunctionRef<void(const FString&, UInterchangeBaseNode*)> IterationLambda) const;
 
@@ -169,6 +183,10 @@ public:
 	UFUNCTION(BlueprintCallable, Category = "Interchange | Node Container")
 	INTERCHANGECORE_API bool ClearNodeParentUid(const FString& NodeUniqueID);
 
+	/** Set the desired child index of the node. */
+	UFUNCTION(BlueprintCallable, Category = "Interchange | Node Container")
+	INTERCHANGECORE_API bool SetNodeDesiredChildIndex(const FString& NodeUniqueID, const int32& NewNodeDesiredChildIndex);
+
 	/** Get the number of children the node has. */
 	UFUNCTION(BlueprintCallable, Category = "Interchange | Node Container")
 	INTERCHANGECORE_API int32 GetNodeChildrenCount(const FString& NodeUniqueID) const;
@@ -220,7 +238,9 @@ public:
 
 private:
 
-	INTERCHANGECORE_API UInterchangeBaseNode* GetNodeChildrenInternal(const FString& NodeUniqueID, int32 ChildIndex);
+	void InternalReorderChildren(TArray<FString>& Children) const;
+
+	UInterchangeBaseNode* GetNodeChildrenInternal(const FString& NodeUniqueID, int32 ChildIndex);
 
 	/** Flat List of the nodes. Since the nodes are variable size, we store a pointer. */
 	UPROPERTY(VisibleAnywhere, Category = "Interchange | Node Container")
