Iterchange: Make sure the task system is destroy after the interchange manager.

#rnx

diff --git a/Engine/Source/Runtime/Interchange/Core/Private/InterchangeTaskSystem.cpp b/Engine/Source/Runtime/Interchange/Core/Private/InterchangeTaskSystem.cpp
--- a/Engine/Source/Runtime/Interchange/Core/Private/InterchangeTaskSystem.cpp
+++ b/Engine/Source/Runtime/Interchange/Core/Private/InterchangeTaskSystem.cpp
@@ -81,7 +81,7 @@ UE::Interchange::FInterchangeTaskSystem& UE::Interchange::FInterchangeTaskSystem
 				return true; //Return true will keep the tick active
 			});
 			
-		FCoreDelegates::OnEnginePreExit.AddLambda([]()
+		FCoreDelegates::OnPreExit.AddLambda([]()
 			{
 				if (!ensure(TaskSystemPtr.IsValid()))
 				{
