FBX Interchange Analytics Improvements
- Application Name, Version and Vendor are now passed in the Analytics

#rnx

diff --git a/Engine/Plugins/Interchange/Runtime/Source/Parsers/Fbx/Private/FbxAPI.cpp b/Engine/Plugins/Interchange/Runtime/Source/Parsers/Fbx/Private/FbxAPI.cpp
--- a/Engine/Plugins/Interchange/Runtime/Source/Parsers/Fbx/Private/FbxAPI.cpp
+++ b/Engine/Plugins/Interchange/Runtime/Source/Parsers/Fbx/Private/FbxAPI.cpp
@@ -184,11 +184,15 @@ namespace UE
 					FString LastSavedAppName(UTF8_TO_TCHAR(DocInfo->LastSaved_ApplicationName.Get().Buffer()));
 					FString LastSavedAppVersion(UTF8_TO_TCHAR(DocInfo->LastSaved_ApplicationVersion.Get().Buffer()));
 
-					FileDetails.FbxFileCreatorApplication = LastSavedVendor + TEXT(" ") + LastSavedAppName + TEXT(" ") + LastSavedAppVersion;
+					FileDetails.ApplicationVendor = LastSavedVendor;
+					FileDetails.ApplicationName = LastSavedAppName;
+					FileDetails.ApplicationVersion = LastSavedAppVersion;
 				}
 				else
 				{
-					FileDetails.FbxFileCreatorApplication = TEXT("");
+					FileDetails.ApplicationVendor = TEXT("");
+					FileDetails.ApplicationName = TEXT("");
+					FileDetails.ApplicationVersion = TEXT("");
 				}
 
 				return true;
@@ -450,10 +454,28 @@ namespace UE
 
 				SourceNode->SetExtraInformation(TEXT("File Version"), FileDetails.FbxFileVersion);
 				SourceNode->SetExtraInformation(TEXT("File Creator"), FileDetails.FbxFileCreator);
-				SourceNode->SetExtraInformation(TEXT("File Creator Application"), FileDetails.FbxFileCreatorApplication);
 				SourceNode->SetExtraInformation(TEXT("File Units"), FileDetails.UnitSystem);
 				SourceNode->SetExtraInformation(TEXT("File Axis Direction"), FileDetails.AxisDirection);
 				SourceNode->SetExtraInformation(TEXT("File Frame Rate"), FileDetails.FrameRate);
+
+				// Analytics Data
+				{
+					using namespace UE::Interchange;
+					if (!FileDetails.ApplicationVendor.IsEmpty())
+					{
+						SourceNode->SetExtraInformation(FSourceNodeExtraInfoStaticData::GetApplicationVendorExtraInfoKey(), FileDetails.ApplicationVendor);
+					}
+
+					if (!FileDetails.ApplicationName.IsEmpty())
+					{
+						SourceNode->SetExtraInformation(FSourceNodeExtraInfoStaticData::GetApplicationNameExtraInfoKey(), FileDetails.ApplicationName);
+					}
+
+					if (!FileDetails.ApplicationVersion.IsEmpty())
+					{
+						SourceNode->SetExtraInformation(FSourceNodeExtraInfoStaticData::GetApplicationVersionExtraInfoKey(), FileDetails.ApplicationVersion);
+					}
+				}
 			}
 		} //ns Private
 	} //ns Interchange
diff --git a/Engine/Plugins/Interchange/Runtime/Source/Parsers/Fbx/Private/FbxAPI.h b/Engine/Plugins/Interchange/Runtime/Source/Parsers/Fbx/Private/FbxAPI.h
--- a/Engine/Plugins/Interchange/Runtime/Source/Parsers/Fbx/Private/FbxAPI.h
+++ b/Engine/Plugins/Interchange/Runtime/Source/Parsers/Fbx/Private/FbxAPI.h
@@ -143,7 +143,9 @@ namespace UE
 				{
 					FString FbxFileVersion;
 					FString FbxFileCreator;
-					FString FbxFileCreatorApplication;
+					FString ApplicationName;
+					FString ApplicationVersion;
+					FString ApplicationVendor;
 					FString UnitSystem;
 					FString AxisDirection;
 					FString FrameRate;
diff --git a/Engine/Source/Runtime/Interchange/Core/Private/Nodes/InterchangeSourceNode.cpp b/Engine/Source/Runtime/Interchange/Core/Private/Nodes/InterchangeSourceNode.cpp
--- a/Engine/Source/Runtime/Interchange/Core/Private/Nodes/InterchangeSourceNode.cpp
+++ b/Engine/Source/Runtime/Interchange/Core/Private/Nodes/InterchangeSourceNode.cpp
@@ -12,18 +12,39 @@
 #include UE_INLINE_GENERATED_CPP_BY_NAME(InterchangeSourceNode)
 
 
-namespace UE::Interchange::SourceNode
+namespace UE::Interchange
 {
-	FString GetSourceNodeUniqueID()
+	const FString& FSourceNodeExtraInfoStaticData::GetApplicationVendorExtraInfoKey()
 	{
-		static FString StaticUid = TEXT("__SourceNode__");
-		return StaticUid;
+		static FString ApplicationVendorExtraInfoKey = TEXT("Application Vendor");
+		return ApplicationVendorExtraInfoKey;
 	}
 
-	const FString& GetExtraInformationKey()
+	const FString& FSourceNodeExtraInfoStaticData::GetApplicationNameExtraInfoKey()
 	{
-		static FString ExtraInformationKey(TEXT("__ExtraInformation__Key"));
-		return ExtraInformationKey;
+		static FString ApplicationNameExtraInfoKey = TEXT("Application Name");
+		return ApplicationNameExtraInfoKey;
+	}
+
+	const FString& FSourceNodeExtraInfoStaticData::GetApplicationVersionExtraInfoKey()
+	{
+		static FString ApplicationVersionExtraInfoKey = TEXT("Application Version");
+		return ApplicationVersionExtraInfoKey;
+	}
+
+	namespace SourceNode
+	{
+		FString GetSourceNodeUniqueID()
+		{
+			static FString StaticUid = TEXT("__SourceNode__");
+			return StaticUid;
+		}
+
+		const FString& GetExtraInformationKey()
+		{
+			static FString ExtraInformationKey(TEXT("__ExtraInformation__Key"));
+			return ExtraInformationKey;
+		}
 	}
 }
 
diff --git a/Engine/Source/Runtime/Interchange/Core/Public/Nodes/InterchangeSourceNode.h b/Engine/Source/Runtime/Interchange/Core/Public/Nodes/InterchangeSourceNode.h
--- a/Engine/Source/Runtime/Interchange/Core/Public/Nodes/InterchangeSourceNode.h
+++ b/Engine/Source/Runtime/Interchange/Core/Public/Nodes/InterchangeSourceNode.h
@@ -19,6 +19,17 @@ class UInterchangeBaseNodeContainer;
 class UObject;
 struct FFrame;
 
+namespace UE::Interchange
+{
+	
+	struct FSourceNodeExtraInfoStaticData 
+	{
+		static INTERCHANGECORE_API const FString& GetApplicationVendorExtraInfoKey();
+		static INTERCHANGECORE_API const FString& GetApplicationNameExtraInfoKey();
+		static INTERCHANGECORE_API const FString& GetApplicationVersionExtraInfoKey();
+	};
+}
+
 /**
  * This class allows a translator to add general source data that describes the whole source. Pipelines can use this information.
  */
diff --git a/Engine/Source/Runtime/Interchange/Engine/Private/InterchangeManager.cpp b/Engine/Source/Runtime/Interchange/Engine/Private/InterchangeManager.cpp
--- a/Engine/Source/Runtime/Interchange/Engine/Private/InterchangeManager.cpp
+++ b/Engine/Source/Runtime/Interchange/Engine/Private/InterchangeManager.cpp
@@ -39,6 +39,7 @@
 #include "Misc/ScopedSlowTask.h"
 #include "Misc/DateTime.h"
 #include "Nodes/InterchangeBaseNodeContainer.h"
+#include "Nodes/InterchangeSourceNode.h"
 #include "PackageUtils/PackageUtils.h"
 #include "Policies/CondensedJsonPrintPolicy.h"
 #include "Serialization/JsonSerializer.h"
@@ -213,6 +214,34 @@ namespace UE::Interchange::Private
 		return FJsonFragment(MoveTemp(ReturnValue));
 	}
 
+	FJsonFragment AnalyticsPrimarySourceDataMetadataToJsonFragment(const UInterchangeSourceData* InSourceData, const UInterchangeBaseNodeContainer* InBaseNodeContainer)
+	{
+		using namespace UE::Interchange;
+		FString ReturnValue;
+		FAnalyticsJsonWriter JsonWriter(&ReturnValue);
+		
+		if (const UInterchangeSourceNode* SourceNode = UInterchangeSourceNode::GetUniqueInstance(InBaseNodeContainer))
+		{
+			TMap<FString, FString> ExtraInfoMap;
+			SourceNode->GetExtraInformation(ExtraInfoMap);
+
+			const FString DefaultValue(TEXT("N/A"));
+				
+			const FString ApplicationVendor = ExtraInfoMap.FindRef(FSourceNodeExtraInfoStaticData::GetApplicationVendorExtraInfoKey(), DefaultValue);
+			const FString ApplicationName = ExtraInfoMap.FindRef(FSourceNodeExtraInfoStaticData::GetApplicationNameExtraInfoKey(), DefaultValue);
+			const FString ApplicationVersion = ExtraInfoMap.FindRef(FSourceNodeExtraInfoStaticData::GetApplicationVersionExtraInfoKey(), DefaultValue);
+				
+			JsonWriter.WriteObjectStart();
+			JsonWriter.WriteValue(TEXT("SourceExtension"), FPaths::GetExtension(InSourceData->GetFilename()));
+			JsonWriter.WriteValue(TEXT("ApplicationVendor"), ApplicationVendor);
+			JsonWriter.WriteValue(TEXT("ApplicationName"), ApplicationName);
+			JsonWriter.WriteValue(TEXT("ApplicationVersion"), ApplicationVersion);
+			JsonWriter.WriteObjectEnd();
+		}
+		JsonWriter.Close();
+		return FJsonFragment(MoveTemp(ReturnValue));
+	}
+
 	bool ExtractNamespace(const FText& Text, FString& OutTextNamespaceId)
 	{
 		FText TextToUse = Text;
@@ -536,6 +565,11 @@ void UE::Interchange::FImportAsyncHelper::SendAnalyticImportEndData()
 		return;
 	}
 
+	{
+		using namespace UE::Interchange::Private;
+		Attribs.Add(FAnalyticsEventAttribute(TEXT("PrimarySourceDataMetadata"), AnalyticsPrimarySourceDataMetadataToJsonFragment(SourceDatas[0], BaseNodeContainers[0].Get())));
+	}
+
 	int32 ImportedObjectCount = 0;
 	for (const TPair<int32, TArray<FImportedObjectInfo>>& SourceIndexAndImportedAssets : ImportedAssetsPerSourceIndex)
 	{
