Hiding the cards view for Interchange Scene Import Asset on reimport

#rnx

diff --git a/Engine/Plugins/Interchange/Editor/Source/Pipelines/Private/InterchangePipelineConfigurationGeneric.cpp b/Engine/Plugins/Interchange/Editor/Source/Pipelines/Private/InterchangePipelineConfigurationGeneric.cpp
--- a/Engine/Plugins/Interchange/Editor/Source/Pipelines/Private/InterchangePipelineConfigurationGeneric.cpp
+++ b/Engine/Plugins/Interchange/Editor/Source/Pipelines/Private/InterchangePipelineConfigurationGeneric.cpp
@@ -175,7 +175,8 @@ EInterchangePipelineConfigurationDialogResult UInterchangePipelineConfigurationG
 	, TWeakObjectPtr<UInterchangeSourceData> SourceData
 	, TWeakObjectPtr <UInterchangeTranslatorBase> Translator
 	, TWeakObjectPtr<UInterchangeBaseNodeContainer> BaseNodeContainer
-	, TWeakObjectPtr <UObject> ReimportAsset)
+	, TWeakObjectPtr <UObject> ReimportAsset
+	, bool bSceneImport)
 {
 	using namespace UE::Interchange::Private;
 	//Create and show the graph inspector UI dialog
@@ -201,7 +202,7 @@ EInterchangePipelineConfigurationDialogResult UInterchangePipelineConfigurationG
 		SAssignNew(InterchangePipelineConfigurationDialog, SInterchangePipelineConfigurationDialog)
 		.OwnerWindow(Window)
 		.SourceData(SourceData)
-		.bSceneImport(false)
+		.bSceneImport(bSceneImport)
 		.bReimport(true)
 		.PipelineStacks(PipelineStacks)
 		.OutPipelines(&OutPipelines)
diff --git a/Engine/Plugins/Interchange/Editor/Source/Pipelines/Public/InterchangePipelineConfigurationGeneric.h b/Engine/Plugins/Interchange/Editor/Source/Pipelines/Public/InterchangePipelineConfigurationGeneric.h
--- a/Engine/Plugins/Interchange/Editor/Source/Pipelines/Public/InterchangePipelineConfigurationGeneric.h
+++ b/Engine/Plugins/Interchange/Editor/Source/Pipelines/Public/InterchangePipelineConfigurationGeneric.h
@@ -33,5 +33,6 @@ protected:
 		, TWeakObjectPtr<UInterchangeSourceData> SourceData
 		, TWeakObjectPtr <UInterchangeTranslatorBase> Translator
 		, TWeakObjectPtr<UInterchangeBaseNodeContainer> BaseNodeContainer
-		, TWeakObjectPtr <UObject> ReimportAsset) override;
+		, TWeakObjectPtr <UObject> ReimportAsset
+		, bool bSceneImport) override;
 };
diff --git a/Engine/Source/Runtime/Interchange/Engine/Private/InterchangeManager.cpp b/Engine/Source/Runtime/Interchange/Engine/Private/InterchangeManager.cpp
--- a/Engine/Source/Runtime/Interchange/Engine/Private/InterchangeManager.cpp
+++ b/Engine/Source/Runtime/Interchange/Engine/Private/InterchangeManager.cpp
@@ -2298,7 +2298,7 @@ UInterchangeManager::ImportInternal(const FString& ContentPath, const UInterchan
 			if (RegisteredPipelineConfiguration && bShowPipelineStacksConfigurationDialog && !bIsUnattended)
 			{
 				//Show the dialog, a plugin should have registered this dialog. We use a plugin to be able to use editor code when doing UI
-				EInterchangePipelineConfigurationDialogResult DialogResult = RegisteredPipelineConfiguration->ScriptedShowReimportPipelineConfigurationDialog(PipelineStacks, OutPipelines, DuplicateSourceData, AsyncTranslator, BaseNodeContainer, ImportAssetParameters.ReimportAsset);
+				EInterchangePipelineConfigurationDialogResult DialogResult = RegisteredPipelineConfiguration->ScriptedShowReimportPipelineConfigurationDialog(PipelineStacks, OutPipelines, DuplicateSourceData, AsyncTranslator, BaseNodeContainer, ImportAssetParameters.ReimportAsset, bImportScene);
 				if (DialogResult == EInterchangePipelineConfigurationDialogResult::Cancel)
 				{
 					bImportCanceled = true;
diff --git a/Engine/Source/Runtime/Interchange/Engine/Public/InterchangePipelineConfigurationBase.h b/Engine/Source/Runtime/Interchange/Engine/Public/InterchangePipelineConfigurationBase.h
--- a/Engine/Source/Runtime/Interchange/Engine/Public/InterchangePipelineConfigurationBase.h
+++ b/Engine/Source/Runtime/Interchange/Engine/Public/InterchangePipelineConfigurationBase.h
@@ -89,7 +89,8 @@ public:
 		, UInterchangeSourceData* SourceData
 		, UInterchangeTranslatorBase* Translator
 		, UInterchangeBaseNodeContainer* BaseNodeContainer
-		, UObject* ReimportAsset);
+		, UObject* ReimportAsset
+		, bool bSceneImport = false);
 
 	/** The default implementation, which is called if the Blueprint does not have any implementation, calls the virtual ShowReimportPipelineConfigurationDialog(). */
 	EInterchangePipelineConfigurationDialogResult ScriptedShowReimportPipelineConfigurationDialog_Implementation(TArray<FInterchangeStackInfo>& PipelineStacks
@@ -97,10 +98,11 @@ public:
 		, UInterchangeSourceData* SourceData
 		, UInterchangeTranslatorBase* Translator
 		, UInterchangeBaseNodeContainer* BaseNodeContainer
-		, UObject* ReimportAsset)
+		, UObject* ReimportAsset
+		, bool bSceneImport = false)
 	{
 		//By default we call the virtual import pipeline execution
-		return ShowReimportPipelineConfigurationDialog(PipelineStacks, OutPipelines, SourceData, Translator, BaseNodeContainer, ReimportAsset);
+		return ShowReimportPipelineConfigurationDialog(PipelineStacks, OutPipelines, SourceData, Translator, BaseNodeContainer, ReimportAsset, bSceneImport);
 	}
 
 protected:
@@ -139,7 +141,8 @@ protected:
 		, TWeakObjectPtr<UInterchangeSourceData> SourceData
 		, TWeakObjectPtr <UInterchangeTranslatorBase> Translator
 		, TWeakObjectPtr <UInterchangeBaseNodeContainer> BaseNodeContainer
-		, TWeakObjectPtr <UObject> ReimportAsset)
+		, TWeakObjectPtr <UObject> ReimportAsset
+		, bool bSceneImport)
 	{
 		//Not implemented
 		return EInterchangePipelineConfigurationDialogResult::Cancel;
