[Interchange] Made Interchange Pipelines' not-editable properties to not save/load to and from ini files.

#rnx

diff --git a/Engine/Source/Runtime/Interchange/Core/Private/InterchangePipelineBase.cpp b/Engine/Source/Runtime/Interchange/Core/Private/InterchangePipelineBase.cpp
--- a/Engine/Source/Runtime/Interchange/Core/Private/InterchangePipelineBase.cpp
+++ b/Engine/Source/Runtime/Interchange/Core/Private/InterchangePipelineBase.cpp
@@ -109,7 +109,8 @@ void UInterchangePipelineBase::LoadSettingsInternal(const FName PipelineStackNam
 	for (FProperty* Property = Class->PropertyLink; Property; Property = Property->PropertyLinkNext)
 	{
 		//Do not load a transient property
-		if (Property->HasAnyPropertyFlags(CPF_Transient))
+		//Properties that are not editable shouldn't be loaded either (As they shouldn't have been saved to begin with).
+		if (Property->HasAnyPropertyFlags(CPF_Transient) || !Property->HasAnyPropertyFlags(CPF_Edit))
 		{
 			continue;
 		}
@@ -242,7 +243,8 @@ void UInterchangePipelineBase::SaveSettingsInternal(const FName PipelineStackNam
 	for (FProperty* Property = Class->PropertyLink; Property; Property = Property->PropertyLinkNext)
 	{
 		//Do not save a transient property
-		if (Property->HasAnyPropertyFlags(CPF_Transient))
+		//Properties that are not editable shouldn't be saved either.
+		if (Property->HasAnyPropertyFlags(CPF_Transient) || !Property->HasAnyPropertyFlags(CPF_Edit))
 		{
 			continue;
 		}
