Interchange: Fix incorrect item count after removing one entry from an arrayhelper with duplicate entries;

Fix MapAttributeHelper not compiling RemoveBySwap for key types other than FString.

diff --git a/Engine/Source/Runtime/Interchange/Core/Public/Nodes/InterchangeBaseNodeUtilities.h b/Engine/Source/Runtime/Interchange/Core/Public/Nodes/InterchangeBaseNodeUtilities.h
--- a/Engine/Source/Runtime/Interchange/Core/Public/Nodes/InterchangeBaseNodeUtilities.h
+++ b/Engine/Source/Runtime/Interchange/Core/Public/Nodes/InterchangeBaseNodeUtilities.h
@@ -21,7 +21,7 @@ template<typename ItemType>
 class TArrayAttributeHelper
 {
 	static_assert(TAttributeTypeTraits<ItemType>::GetType() != EAttributeTypes::None, "The value type must be supported by the attribute storage");
-	
+
 public:
 	~TArrayAttributeHelper()
 	{
@@ -114,7 +114,7 @@ public:
 			OutItems.Empty();
 			return;
 		}
-		
+
 		FAttributeStorage::TAttributeHandle<int32> Handle = AttributePtr->GetAttributeHandle<int32>(GetKeyCount());
 		if (!Handle.IsValid())
 		{
@@ -208,7 +208,8 @@ public:
 			}
 			ItemType Item;
 			HandleItem.Get(Item);
-			if (Item == ItemToDelete)
+			//Remove at most one item
+			if (DecrementKey == 0 && Item == ItemToDelete)
 			{
 				//Remove this entry
 				AttributePtr->UnregisterAttribute(DepIndexKey);
@@ -555,7 +556,7 @@ private:
 		return AttributesPtr->GetAttributeHandle<KeyType>(GetKeyAttribute(CachedKeysAndValues.Num() - 1));
 	}
 
-	bool RemoveBySwap(uint32 Hash, const KeyType& InKey, TPair<TAttributeHandle<FString>, TAttributeHandle<ValueType>>& CachedPair)
+	bool RemoveBySwap(uint32 Hash, const KeyType& InKey, TPair<TAttributeHandle<KeyType>, TAttributeHandle<ValueType>>& CachedPair)
 	{
 		TSharedPtr<FAttributeStorage, ESPMode::ThreadSafe> AttributesPtr = Attributes.Pin();
 		if (!ensure(AttributesPtr.IsValid()))
