Interchange: Setup basic C++ automated tests for ArrayAttributeHelper and MapAttributeHelper.

diff --git a/Engine/Source/Runtime/Interchange/Core/Private/Tests/ArrayAttributeHelperTests.cpp b/Engine/Source/Runtime/Interchange/Core/Private/Tests/ArrayAttributeHelperTests.cpp
new file mode 100644
--- /dev/null
+++ b/Engine/Source/Runtime/Interchange/Core/Private/Tests/ArrayAttributeHelperTests.cpp
@@ -0,0 +1,75 @@
+// Copyright Epic Games, Inc. All Rights Reserved.
+
+#include "Async/ParallelFor.h"
+#include "CoreMinimal.h"
+#include "Misc/AutomationTest.h"
+#include "Nodes/InterchangeBaseNodeUtilities.h"
+#include "Serialization/MemoryReader.h"
+#include "Serialization/MemoryWriter.h"
+#include "Templates/Atomic.h"
+#include "Types/AttributeStorage.h"
+
+#if WITH_DEV_AUTOMATION_TESTS
+
+IMPLEMENT_SIMPLE_AUTOMATION_TEST(
+	FArrayAttributeHelperTests,
+	"System.Runtime.Interchange.ArrayAttributeHelperTests",
+	EAutomationTestFlags::EditorContext | EAutomationTestFlags::EngineFilter
+)
+
+namespace UE::ArrayAttributeHelperTests::Private
+{
+	template<typename T>
+	void RunTestInternal(FArrayAttributeHelperTests& Tester)
+	{
+		using namespace UE::Interchange;
+
+		const FString BaseKeyName = TEXT("TestKey");
+		TSharedPtr<FAttributeStorage, ESPMode::ThreadSafe> TestStorage = MakeShared<FAttributeStorage>();
+
+		UE::Interchange::TArrayAttributeHelper<T> Helper;
+		Helper.Initialize(TestStorage, BaseKeyName);
+
+		const T SomeItem = T{};
+		Helper.AddItem(SomeItem);
+		Helper.AddItem(SomeItem);
+		Helper.AddItem(SomeItem);
+		Tester.TestEqual(TEXT("Count after adding 3"), Helper.GetCount(), 3);
+
+		T GotItem;
+		Helper.GetItem(2, GotItem);
+		Tester.TestEqual(TEXT("Retrieved item after adding 3"), GotItem, SomeItem);
+
+		Helper.RemoveItem(SomeItem);
+		Tester.TestEqual(TEXT("Count after removal with duplicates"), Helper.GetCount(), 2);
+
+		TArray<T> Results;
+		Helper.GetItems(Results);
+		Tester.TestEqual(TEXT("Actual number of items in the result"), Results.Num(), 2);
+
+		for (const T& Result : Results)
+		{
+			Tester.TestEqual(TEXT("Item equals the expected value"), Result, SomeItem);
+		}
+
+		Helper.RemoveAllItems();
+		Tester.TestEqual(TEXT("Count after remove all"), Helper.GetCount(), 0);
+
+		Helper.GetItems(Results);
+		Tester.TestEqual(TEXT("Actual number of items in the result"), Results.Num(), 0);
+	}
+}	 // namespace UE::ArrayAttributeHelperTests::Private
+
+bool FArrayAttributeHelperTests::RunTest(const FString& Parameters)
+{
+	using namespace UE::ArrayAttributeHelperTests::Private;
+
+	RunTestInternal<FString>(*this);
+	RunTestInternal<int32>(*this);
+	RunTestInternal<double>(*this);
+	RunTestInternal<float>(*this);
+
+	return true;
+}
+
+#endif	  // WITH_DEV_AUTOMATION_TESTS
diff --git a/Engine/Source/Runtime/Interchange/Core/Private/Tests/MapAttributeHelperTests.cpp b/Engine/Source/Runtime/Interchange/Core/Private/Tests/MapAttributeHelperTests.cpp
new file mode 100644
--- /dev/null
+++ b/Engine/Source/Runtime/Interchange/Core/Private/Tests/MapAttributeHelperTests.cpp
@@ -0,0 +1,100 @@
+// Copyright Epic Games, Inc. All Rights Reserved.
+
+#include "Async/ParallelFor.h"
+#include "CoreMinimal.h"
+#include "Misc/AutomationTest.h"
+#include "Nodes/InterchangeBaseNodeUtilities.h"
+#include "Serialization/MemoryReader.h"
+#include "Serialization/MemoryWriter.h"
+#include "Templates/Atomic.h"
+#include "Types/AttributeStorage.h"
+
+#if WITH_DEV_AUTOMATION_TESTS
+
+IMPLEMENT_SIMPLE_AUTOMATION_TEST(
+	FMapAttributeHelperTests,
+	"System.Runtime.Interchange.MapAttributeHelperTests",
+	EAutomationTestFlags::EditorContext | EAutomationTestFlags::EngineFilter
+)
+
+namespace UE::MapAttributeHelperTests::Private
+{
+	template<typename T, typename U>
+	void RunTestInternal(FMapAttributeHelperTests& Tester)
+	{
+		using namespace UE::Interchange;
+
+		const FString BaseKeyName = TEXT("TestKey");
+		TSharedPtr<FAttributeStorage, ESPMode::ThreadSafe> TestStorage = MakeShared<FAttributeStorage>();
+
+		UE::Interchange::TMapAttributeHelper<T, U> Helper;
+		Helper.Initialize(TestStorage.ToSharedRef(), BaseKeyName);
+
+		T SomeKey = T{};
+		T AnotherKey = T{};
+		if constexpr (std::is_same_v<T, FString>)
+		{
+			SomeKey = TEXT("SomeKey");
+			AnotherKey = TEXT("AnotherKey");
+		}
+		else if constexpr (std::is_arithmetic_v<T>)
+		{
+			SomeKey = 1;
+			AnotherKey = 20;
+		}
+
+		U SomeValue = U{};
+		U AnotherValue = U{};
+		if constexpr (std::is_same_v<U, FString>)
+		{
+			SomeValue = TEXT("SomeValue");
+			AnotherValue = TEXT("AnotherValue");
+		}
+		else if constexpr (std::is_arithmetic_v<U>)
+		{
+			SomeValue = 1;
+			AnotherValue = 20;
+		}
+
+		Helper.SetKeyValue(SomeKey, AnotherValue);
+		Helper.SetKeyValue(SomeKey, SomeValue);
+		Helper.SetKeyValue(AnotherKey, SomeValue);
+		Helper.SetKeyValue(AnotherKey, AnotherValue);
+
+		TMap<T, U> Map = Helper.ToMap();
+		Tester.TestEqual(TEXT("Compare maps after ToMap"), Map.Num(), 2);
+
+		U SomeGotValue = U{};
+		U AnotherGotValue = U{};
+		bool bSuccess = Helper.GetValue(SomeKey, SomeGotValue);
+		bSuccess = bSuccess && Helper.GetValue(AnotherKey, AnotherGotValue);
+		Tester.TestTrue(TEXT("GetValue succeeded"), bSuccess);
+		Tester.TestEqual(TEXT("Compare result after add 1"), SomeGotValue, SomeValue);
+		Tester.TestEqual(TEXT("Compare result after add 2"), AnotherGotValue, AnotherValue);
+
+		Tester.TestTrue(TEXT("Test removing a key"), Helper.RemoveKey(AnotherKey));
+
+		U RemovedValue;
+		bSuccess = Helper.RemoveKeyAndGetValue(SomeKey, RemovedValue);
+		Tester.TestTrue(TEXT("RemoveKeyAndGetValue succeeded"), bSuccess);
+		Tester.TestEqual(TEXT("Compare removed value"), RemovedValue, SomeValue);
+
+		Tester.TestFalse(TEXT("Test removing missing key"), Helper.RemoveKey(AnotherKey));
+	}
+}	 // namespace UE::MapAttributeHelperTests::Private
+
+bool FMapAttributeHelperTests::RunTest(const FString& Parameters)
+{
+	using namespace UE::MapAttributeHelperTests::Private;
+
+	RunTestInternal<FString, FString>(*this);
+	RunTestInternal<FString, int32>(*this);
+	RunTestInternal<FString, float>(*this);
+	RunTestInternal<int32, FString>(*this);
+	RunTestInternal<int32, int32>(*this);
+	RunTestInternal<int32, float>(*this);
+
+	return true;
+}
+
+#endif	  // WITH_DEV_AUTOMATION_TESTS
