[Interchange] Add a UInterchangeAssetImportData script getter to allow tech art to add/remove/modify re-import pipelines for an asset.

#rnx

diff --git a/Engine/Source/Runtime/Interchange/Engine/Private/InterchangeAssetImportData.cpp b/Engine/Source/Runtime/Interchange/Engine/Private/InterchangeAssetImportData.cpp
--- a/Engine/Source/Runtime/Interchange/Engine/Private/InterchangeAssetImportData.cpp
+++ b/Engine/Source/Runtime/Interchange/Engine/Private/InterchangeAssetImportData.cpp
@@ -257,7 +257,7 @@ void UInterchangeAssetImportData::SetTranslatorSettings(UInterchangeTranslatorSe
 	}
 }
 
-void UInterchangeAssetImportData::SetPipelines(const TArray<UObject*>& InPipelines) const
+void UInterchangeAssetImportData::SetPipelines(const TArray<UObject*>& InPipelines)
 {
 	TransientPipelines.Reset();
 
diff --git a/Engine/Source/Runtime/Interchange/Engine/Private/InterchangeManager.cpp b/Engine/Source/Runtime/Interchange/Engine/Private/InterchangeManager.cpp
--- a/Engine/Source/Runtime/Interchange/Engine/Private/InterchangeManager.cpp
+++ b/Engine/Source/Runtime/Interchange/Engine/Private/InterchangeManager.cpp
@@ -2868,6 +2868,11 @@ bool UInterchangeManager::WarnIfInterchangeIsActive()
 	return true;
 }
 
+UInterchangeAssetImportData* UInterchangeManager::GetAssetImportData(UObject* Asset) const
+{
+	return UInterchangeAssetImportData::GetFromObject(Asset);
+}
+
 bool UInterchangeManager::CanTranslateSourceDataWithPayloadInterface(const UInterchangeSourceData* SourceData, const UClass* PayloadInterfaceClass) const
 {
 	if (GetTranslatorSupportingPayloadInterfaceForSourceData(SourceData, PayloadInterfaceClass))
diff --git a/Engine/Source/Runtime/Interchange/Engine/Public/InterchangeAssetImportData.h b/Engine/Source/Runtime/Interchange/Engine/Public/InterchangeAssetImportData.h
--- a/Engine/Source/Runtime/Interchange/Engine/Public/InterchangeAssetImportData.h
+++ b/Engine/Source/Runtime/Interchange/Engine/Public/InterchangeAssetImportData.h
@@ -154,7 +154,7 @@ public:
 	INTERCHANGEENGINE_API int32 GetNumberOfPipelines() const;
 
 	UFUNCTION(BlueprintCallable, Category = "Interchange | AssetImportData")
-	INTERCHANGEENGINE_API void SetPipelines(const TArray<UObject*>& InPipelines) const;
+	INTERCHANGEENGINE_API void SetPipelines(const TArray<UObject*>& InPipelines);
 
 
 	UFUNCTION(BlueprintCallable, Category = "Interchange | AssetImportData")
diff --git a/Engine/Source/Runtime/Interchange/Engine/Public/InterchangeManager.h b/Engine/Source/Runtime/Interchange/Engine/Public/InterchangeManager.h
--- a/Engine/Source/Runtime/Interchange/Engine/Public/InterchangeManager.h
+++ b/Engine/Source/Runtime/Interchange/Engine/Public/InterchangeManager.h
@@ -807,6 +807,12 @@ public:
 	UFUNCTION(BlueprintCallable, Category = "Interchange | Import Manager")
 	INTERCHANGEENGINE_API bool WarnIfInterchangeIsActive();
 
+	/**
+	 * Returns the list of supported formats for a given translator type.
+	 */
+	UFUNCTION(BlueprintCallable, Category = "Interchange | Import Manager")
+	INTERCHANGEENGINE_API UInterchangeAssetImportData* GetAssetImportData(UObject* Asset) const;
+
 	/**
 	 * Check whether there is a translator registered that can translate the source data with the specified PayloadInterface.
 	 * @Param SourceData - The source data input you want to translate to Uod.
