Interchange: Project setting for showing the reimport options dialog.

Deprecate the Interchange bShowImportDialogAtReimport option from the Editor settings;
Allow per-translator dialog overrides of whether to show the reimport options dialog;
Add per-translator dialog overrides for scene import settings.

diff --git a/Engine/Plugins/Interchange/Extensions/OpenUSD/Source/Import/Private/InterchangeOpenUSDImportModule.cpp b/Engine/Plugins/Interchange/Extensions/OpenUSD/Source/Import/Private/InterchangeOpenUSDImportModule.cpp
--- a/Engine/Plugins/Interchange/Extensions/OpenUSD/Source/Import/Private/InterchangeOpenUSDImportModule.cpp
+++ b/Engine/Plugins/Interchange/Extensions/OpenUSD/Source/Import/Private/InterchangeOpenUSDImportModule.cpp
@@ -2,6 +2,7 @@
 
 #include "InterchangeOpenUSDImportModule.h"
 
+#include "Algo/NoneOf.h"
 #include "Engine/Engine.h"
 #include "InterchangeManager.h"
 #include "InterchangeProjectSettings.h"
@@ -17,31 +18,102 @@ void FInterchangeOpenUSDImportModule::StartupModule()
 		UInterchangeManager& InterchangeManager = UInterchangeManager::GetInterchangeManager();
 		InterchangeManager.RegisterTranslator(UInterchangeUSDTranslator::StaticClass());
 
-		// Add USD translator and pipelines to the Interchange engine
+		// Don't go through FInterchangeProjectSettingsUtils::GetMutableDefaultImportSettings as we'll need a member of the actual
+		// FInterchangeContent/SceneImportSettings struct anyway, and this is likely safer than casting the struct pointer
+		UInterchangeProjectSettings* ProjectSettings = GetMutableDefault<UInterchangeProjectSettings>();
+		if (!ProjectSettings)
+		{
+			return;
+		}
+
+		FInterchangeContentImportSettings& AssetImportSettings = ProjectSettings->ContentImportSettings;
+		FInterchangeSceneImportSettings& SceneImportSettings = ProjectSettings->SceneImportSettings;
+
+		TSoftClassPtr<UInterchangeTranslatorBase> TranslatorClassPath = TSoftClassPtr<UInterchangeTranslatorBase>(
+			UInterchangeUSDTranslator::StaticClass()
+		);
+
+		FInterchangePerTranslatorDialogOverride ImportDialogOverride;
+		ImportDialogOverride.Translator = TranslatorClassPath;
+		ImportDialogOverride.bShowImportDialog = true;
+		ImportDialogOverride.bShowReimportDialog = true;
+
+		// Asset import pipelines
 		{
 			FInterchangeTranslatorPipelines TranslatorPipelines;
-			TranslatorPipelines.Translator = TSoftClassPtr<UInterchangeTranslatorBase>(UInterchangeUSDTranslator::StaticClass());
+			TranslatorPipelines.Translator = TranslatorClassPath;
 			TranslatorPipelines.Pipelines.Add(FSoftObjectPath(TEXT("/InterchangeOpenUSD/Pipelines/DefaultUSDAssetsPipeline.DefaultUSDAssetsPipeline")));
 			TranslatorPipelines.Pipelines.Add(FSoftObjectPath(TEXT("/InterchangeOpenUSD/Pipelines/DefaultUSDPipelineAssetImport.DefaultUSDPipelineAssetImport")));
 			TranslatorPipelines.Pipelines.Add(FSoftObjectPath(TEXT("/Interchange/Pipelines/DefaultMaterialXPipeline.DefaultMaterialXPipeline")));
 
-			FInterchangeImportSettings& ImportSettings = FInterchangeProjectSettingsUtils::GetMutableDefaultImportSettings(false);
-			FInterchangePipelineStack& PipelineStack = ImportSettings.PipelineStacks[TEXT("Assets")];
+			FInterchangePipelineStack& PipelineStack = AssetImportSettings.PipelineStacks[TEXT("Assets")];
 			PipelineStack.PerTranslatorPipelines.Add(TranslatorPipelines);
 		}
 
+		// Scene import pipelines
 		{
 			FInterchangeTranslatorPipelines TranslatorPipelines;
-			TranslatorPipelines.Translator = TSoftClassPtr<UInterchangeTranslatorBase>(UInterchangeUSDTranslator::StaticClass());
+			TranslatorPipelines.Translator = TranslatorClassPath;
 			TranslatorPipelines.Pipelines.Add(FSoftObjectPath(TEXT("/InterchangeOpenUSD/Pipelines/DefaultUSDSceneAssetsPipeline.DefaultUSDSceneAssetsPipeline")));
 			TranslatorPipelines.Pipelines.Add(FSoftObjectPath(TEXT("/Interchange/Pipelines/DefaultSceneLevelPipeline.DefaultSceneLevelPipeline")));
 			TranslatorPipelines.Pipelines.Add(FSoftObjectPath(TEXT("/InterchangeOpenUSD/Pipelines/DefaultUSDPipeline.DefaultUSDPipeline")));
 			TranslatorPipelines.Pipelines.Add(FSoftObjectPath(TEXT("/Interchange/Pipelines/DefaultMaterialXPipeline.DefaultMaterialXPipeline")));
 
-			FInterchangeImportSettings& ImportSettings = FInterchangeProjectSettingsUtils::GetMutableDefaultImportSettings(true);
-			FInterchangePipelineStack& PipelineStack = ImportSettings.PipelineStacks[TEXT("Scene")];
+			FInterchangePipelineStack& PipelineStack = SceneImportSettings.PipelineStacks[TEXT("Scene")];
 			PipelineStack.PerTranslatorPipelines.Add(TranslatorPipelines);
 		}
+
+		// Asset import dialog overrides (we want to show the import and reimport dialog for all USD imports, of all asset types)
+		{
+			// TODO: Add a Count value to the enum and then use ENUM_RANGE_BY_COUNT?
+			static const TArray<EInterchangeTranslatorAssetType> AssetTypes = {
+				EInterchangeTranslatorAssetType::None,
+				EInterchangeTranslatorAssetType::Textures,
+				EInterchangeTranslatorAssetType::Materials,
+				EInterchangeTranslatorAssetType::Meshes,
+				EInterchangeTranslatorAssetType::Animations,
+			};
+
+			for (EInterchangeTranslatorAssetType AssetType : AssetTypes)
+			{
+				FInterchangeDialogOverride& DialogOverrides = AssetImportSettings.ShowImportDialogOverride.FindOrAdd(AssetType);
+
+				// Don't set the dialog overrides for the USD translator if the user already has done that (with potentially different values)
+				bool bHasUSDOverrides = false;
+				for (const FInterchangePerTranslatorDialogOverride& Override : DialogOverrides.PerTranslatorImportDialogOverride)
+				{
+					if (Override.Translator == TranslatorClassPath)
+					{
+						bHasUSDOverrides = true;
+						break;
+					}
+				}
+
+				if (!bHasUSDOverrides)
+				{
+					DialogOverrides.PerTranslatorImportDialogOverride.Add(ImportDialogOverride);
+				}
+			}
+		}
+
+		// Scene import dialog overrides
+		{
+			// Don't set the dialog overrides for the USD translator if the user already has done that (with potentially different values)
+			bool bHasUSDOverrides = false;
+			for (const FInterchangePerTranslatorDialogOverride& Override : SceneImportSettings.PerTranslatorDialogOverride)
+			{
+				if (Override.Translator == TranslatorClassPath)
+				{
+					bHasUSDOverrides = true;
+					break;
+				}
+			}
+
+			if (!bHasUSDOverrides)
+			{
+				SceneImportSettings.PerTranslatorDialogOverride.Add(ImportDialogOverride);
+			}
+		}
 	};
 
 	PostEngineInitHandle = FCoreDelegates::OnPostEngineInit.AddLambda(RegisterItems);
diff --git a/Engine/Plugins/Interchange/Extensions/OpenVDB/Source/Import/Private/InterchangeOpenVDBImportModule.cpp b/Engine/Plugins/Interchange/Extensions/OpenVDB/Source/Import/Private/InterchangeOpenVDBImportModule.cpp
--- a/Engine/Plugins/Interchange/Extensions/OpenVDB/Source/Import/Private/InterchangeOpenVDBImportModule.cpp
+++ b/Engine/Plugins/Interchange/Extensions/OpenVDB/Source/Import/Private/InterchangeOpenVDBImportModule.cpp
@@ -64,6 +64,7 @@ void FInterchangeOpenVDBImportModule::StartupModule()
 			FInterchangePerTranslatorDialogOverride ImportDialogOverride;
 			ImportDialogOverride.Translator = TranslatorClassPath;
 			ImportDialogOverride.bShowImportDialog = true;
+			ImportDialogOverride.bShowReimportDialog = true;
 			DialogOverrides.PerTranslatorImportDialogOverride.Add(ImportDialogOverride);
 		}
 #endif
diff --git a/Engine/Source/Runtime/Interchange/Engine/Private/InterchangeProjectSettings.cpp b/Engine/Source/Runtime/Interchange/Engine/Private/InterchangeProjectSettings.cpp
--- a/Engine/Source/Runtime/Interchange/Engine/Private/InterchangeProjectSettings.cpp
+++ b/Engine/Source/Runtime/Interchange/Engine/Private/InterchangeProjectSettings.cpp
@@ -59,7 +59,7 @@ FInterchangeImportSettings& FInterchangeProjectSettingsUtils::GetMutableDefaultI
 FName FInterchangeProjectSettingsUtils::GetDefaultPipelineStackName(const bool bIsSceneImport, const UInterchangeSourceData& SourceData)
 {
 	const FInterchangeImportSettings& ImportSettings = GetDefaultImportSettings(bIsSceneImport);
-	
+
 	FInterchangeGroup::EUsedGroupStatus UsedGroupStatus;
 	const FInterchangeGroup& UsedInterchangeGroup = FInterchangeProjectSettingsUtils::GetUsedGroup(UsedGroupStatus);
 	bool bInterchangeGroupUsed = (UsedGroupStatus == FInterchangeGroup::EUsedGroupStatus::SetAndValid);
@@ -131,16 +131,66 @@ void FInterchangeProjectSettingsUtils::SetDefaultPipelineStackName(const bool bI
 
 bool FInterchangeProjectSettingsUtils::ShouldShowPipelineStacksConfigurationDialog(const bool bIsSceneImport, const bool bReImport, const UInterchangeSourceData& SourceData)
 {
-	bool bShowImportDialog = true; 
-
 	FInterchangeGroup::EUsedGroupStatus UsedGroupStatus;
 	const FInterchangeGroup& UsedInterchangeGroup = FInterchangeProjectSettingsUtils::GetUsedGroup(UsedGroupStatus);
 	bool bInterchangeGroupUsed = (UsedGroupStatus == FInterchangeGroup::EUsedGroupStatus::SetAndValid);
 
-	bShowImportDialog = bInterchangeGroupUsed ? UsedInterchangeGroup.bShowImportDialog : GetDefaultImportSettings(bIsSceneImport).bShowImportDialog;
-	
+	bool bShowDialog = false;
+	if (bInterchangeGroupUsed)
+	{
+		bShowDialog = bReImport ? UsedInterchangeGroup.bShowReimportDialog
+								: UsedInterchangeGroup.bShowImportDialog;
+	}
+	else
+	{
+		bShowDialog = bReImport ? GetDefaultImportSettings(bIsSceneImport).bShowReimportDialog
+								: GetDefaultImportSettings(bIsSceneImport).bShowImportDialog;
+	}
 
-	if (!bIsSceneImport)
+	if (bIsSceneImport)
+	{
+		// Find the per translator overrides to use
+		const TArray<FInterchangePerTranslatorDialogOverride>* PerTranslatorOverrides = nullptr;
+		if (bInterchangeGroupUsed)
+		{
+			// Use "None" to mean "scene" when authored on the groups, so we don't have to modify the group signature and
+			// lose users' saved groups
+			if (const FInterchangeDialogOverride* FoundNoneOverrides = UsedInterchangeGroup.ShowImportDialogOverride.Find(
+					EInterchangeTranslatorAssetType::None
+				))
+			{
+				PerTranslatorOverrides = &FoundNoneOverrides->PerTranslatorImportDialogOverride;
+			}
+		}
+		else
+		{
+			PerTranslatorOverrides = &GetDefault<UInterchangeProjectSettings>()->SceneImportSettings.PerTranslatorDialogOverride;
+		}
+
+		if (PerTranslatorOverrides)
+		{
+			// Check if the translator for this source data has an override
+			UE::Interchange::FScopedTranslator ScopedTranslator{&SourceData};
+			if (UInterchangeTranslatorBase* Translator = ScopedTranslator.GetTranslator())
+			{
+				EInterchangeTranslatorAssetType SupportedAssetTypes = Translator->GetSupportedAssetTypes();
+				const UClass* TranslatorClass = Translator->GetClass();
+
+				const FInterchangePerTranslatorDialogOverride* PerTranslatorOverride = PerTranslatorOverrides->FindByPredicate(
+					[&TranslatorClass](const FInterchangePerTranslatorDialogOverride& InterchangePerTranslatorDialogOverride)
+					{
+						return (InterchangePerTranslatorDialogOverride.Translator.Get() == TranslatorClass);
+					}
+				);
+
+				if (PerTranslatorOverride)
+				{
+					bShowDialog = (bReImport ? PerTranslatorOverride->bShowReimportDialog : PerTranslatorOverride->bShowImportDialog);
+				}
+			}
+		}
+	}
+	else
 	{
 		UE::Interchange::FScopedTranslator ScopedTranslator(&SourceData);
 
@@ -152,7 +202,9 @@ bool FInterchangeProjectSettingsUtils::ShouldShowPipelineStacksConfigurationDial
 			//Iterate all override, if there is at least one override that show the imnport dialog we will show it.
 			bool bFoundOverride = false;
 			bool bShowFromOverrideStack = false;
-			const TMap<EInterchangeTranslatorAssetType, FInterchangeDialogOverride>& ShowImportDialogOverride = bInterchangeGroupUsed ? UsedInterchangeGroup.ShowImportDialogOverride : GetDefault<UInterchangeProjectSettings>()->ContentImportSettings.ShowImportDialogOverride;
+			const TMap<EInterchangeTranslatorAssetType, FInterchangeDialogOverride>& ShowImportDialogOverride = bInterchangeGroupUsed
+				? UsedInterchangeGroup.ShowImportDialogOverride
+				: GetDefault<UInterchangeProjectSettings>()->ContentImportSettings.ShowImportDialogOverride;
 
 			for (const TPair<EInterchangeTranslatorAssetType, FInterchangeDialogOverride>& ShowImportDialog : ShowImportDialogOverride)
 			{
@@ -167,29 +219,24 @@ bool FInterchangeProjectSettingsUtils::ShouldShowPipelineStacksConfigurationDial
 						});
 					if (FindPerTranslatorOverride)
 					{
-						bShowFromOverrideStack |= FindPerTranslatorOverride->bShowImportDialog;
+						bShowFromOverrideStack |= (bReImport ? FindPerTranslatorOverride->bShowReimportDialog : FindPerTranslatorOverride->bShowImportDialog);
 					}
 					else
 					{
 						//Simple override
-						bShowFromOverrideStack |= InterchangeDialogOverride.bShowImportDialog;
+						bShowFromOverrideStack |= (bReImport ? InterchangeDialogOverride.bShowReimportDialog : InterchangeDialogOverride.bShowImportDialog);
 					}
 					bFoundOverride = true;
 				}
 			}
 			if (bFoundOverride)
 			{
-				bShowImportDialog = bShowFromOverrideStack;
+				bShowDialog = bShowFromOverrideStack;
 			}
 		}
 	}
 
-	if (const UInterchangeEditorSettings* InterchangeEditorSettings = GetDefault<UInterchangeEditorSettings>())
-	{
-		bShowImportDialog &= (bReImport ? InterchangeEditorSettings->bShowImportDialogAtReimport : true);
-	}
-
-	return bShowImportDialog;
+	return bShowDialog;
 }
 
 const FInterchangeGroup& FInterchangeProjectSettingsUtils::GetUsedGroup(FInterchangeGroup::EUsedGroupStatus& UsedGroupStatus)
diff --git a/Engine/Source/Runtime/Interchange/Engine/Public/InterchangeProjectSettings.h b/Engine/Source/Runtime/Interchange/Engine/Public/InterchangeProjectSettings.h
--- a/Engine/Source/Runtime/Interchange/Engine/Public/InterchangeProjectSettings.h
+++ b/Engine/Source/Runtime/Interchange/Engine/Public/InterchangeProjectSettings.h
@@ -29,7 +29,7 @@ USTRUCT()
 struct FInterchangePipelineStack
 {
 	GENERATED_BODY()
-	
+
 	/** The list of pipelines in this stack. The pipelines are executed in fixed order, from top to bottom. */
 	UPROPERTY(EditAnywhere, Category = "Pipelines", meta = (AllowedClasses = "/Script/InterchangeCore.InterchangePipelineBase, /Script/InterchangeEngine.InterchangeBlueprintPipelineBase, /Script/InterchangeEngine.InterchangePythonPipelineAsset"))
 	TArray<FSoftObjectPath> Pipelines;
@@ -39,6 +39,40 @@ struct FInterchangePipelineStack
 	TArray<FInterchangeTranslatorPipelines> PerTranslatorPipelines;
 };
 
+USTRUCT()
+struct FInterchangePerTranslatorDialogOverride
+{
+	GENERATED_BODY()
+
+	UPROPERTY(EditAnywhere, Category = "DialogOverride", meta = (AllowedClasses = "/Script/InterchangeCore.InterchangeTranslatorBase"))
+	TSoftClassPtr<UInterchangeTranslatorBase> Translator;
+
+	/** Show the options dialog when Interchange imports. */
+	UPROPERTY(EditAnywhere, Category = "DialogOverride")
+	bool bShowImportDialog = true;
+
+	/** Show the options dialog when Interchange reimports. */
+	UPROPERTY(EditAnywhere, Category = "DialogOverride")
+	bool bShowReimportDialog = false;
+};
+
+USTRUCT()
+struct FInterchangeDialogOverride
+{
+	GENERATED_BODY()
+
+	/** Show the options dialog when Interchange imports. */
+	UPROPERTY(EditAnywhere, Category = "DialogOverride")
+	bool bShowImportDialog = true;
+
+	/** Show the options dialog when Interchange reimports. */
+	UPROPERTY(EditAnywhere, Category = "DialogOverride")
+	bool bShowReimportDialog = false;
+
+	UPROPERTY(EditAnywhere, Category = "DialogOverride")
+	TArray<FInterchangePerTranslatorDialogOverride> PerTranslatorImportDialogOverride;
+};
+
 USTRUCT()
 struct FInterchangeImportSettings
 {
@@ -56,33 +90,22 @@ struct FInterchangeImportSettings
 	UPROPERTY(EditAnywhere, Category = "Pipeline")
 	TSoftClassPtr <UInterchangePipelineConfigurationBase> ImportDialogClass;
 
-	/** If enabled, the import option dialog will show when interchange import.*/
-	UPROPERTY(EditAnywhere, Category = "Pipeline")
+	/** Show the options dialog when Interchange imports. */
+	UPROPERTY(EditAnywhere, Category = "DialogOverride")
 	bool bShowImportDialog = true;
-};
-
-USTRUCT()
-struct FInterchangePerTranslatorDialogOverride
-{
-	GENERATED_BODY()
-
-	UPROPERTY(EditAnywhere, Category = "DialogOverride", meta = (AllowedClasses = "/Script/InterchangeCore.InterchangeTranslatorBase"))
-	TSoftClassPtr<UInterchangeTranslatorBase> Translator;
 
+	/** Show the options dialog when Interchange reimports. */
 	UPROPERTY(EditAnywhere, Category = "DialogOverride")
-	bool bShowImportDialog = true;
+	bool bShowReimportDialog = false;
 };
 
 USTRUCT()
-struct FInterchangeDialogOverride
+struct FInterchangeSceneImportSettings : public FInterchangeImportSettings
 {
 	GENERATED_BODY()
 
-	UPROPERTY(EditAnywhere, Category = "DialogOverride")
-	bool bShowImportDialog = true;
-
-	UPROPERTY(EditAnywhere, Category = "DialogOverride", meta = (AllowedClasses = "/Script/InterchangeCore.InterchangePipelineBase, /Script/InterchangeEngine.InterchangeBlueprintPipelineBase, /Script/InterchangeEngine.InterchangePythonPipelineAsset"))
-	TArray<FInterchangePerTranslatorDialogOverride> PerTranslatorImportDialogOverride;
+	UPROPERTY(EditAnywhere, Category = "DialogOverride", Meta=(DisplayAfter="bShowReimportDialog"))
+	TArray<FInterchangePerTranslatorDialogOverride> PerTranslatorDialogOverride;
 };
 
 USTRUCT()
@@ -94,8 +117,7 @@ struct FInterchangeContentImportSettings : public FInterchangeImportSettings
 	UPROPERTY(EditAnywhere, Category = "Pipeline", Meta=(DisplayAfter="DefaultPipelineStack"))
 	TMap<EInterchangeTranslatorAssetType, FName> DefaultPipelineStackOverride;
 
-	/** This tell interchange if the import dialog should show or not when importing a particular type of asset.*/
-	UPROPERTY(EditAnywhere, Category = "Pipeline", Meta=(DisplayAfter="bShowImportDialogAtReimport"))
+	UPROPERTY(EditAnywhere, Category = "DialogOverride", Meta=(DisplayAfter="bShowReimportDialog"))
 	TMap<EInterchangeTranslatorAssetType, FInterchangeDialogOverride> ShowImportDialogOverride;
 };
 
@@ -127,11 +149,16 @@ struct FInterchangeGroup
 	UPROPERTY(EditAnywhere, Category = "Pipeline", Meta = (DisplayAfter = "DefaultPipelineStack"))
 	TMap<EInterchangeTranslatorAssetType, FName> DefaultPipelineStackOverride;
 
+	/** Show the options dialog when Interchange imports. */
 	UPROPERTY(EditAnywhere, Category = "DialogOverride")
 	bool bShowImportDialog = true;
 
+	/** Show the options dialog when Interchange reimports. */
+	UPROPERTY(EditAnywhere, Category = "DialogOverride")
+	bool bShowReimportDialog = false;
+
 	/** This tell interchange if the import dialog should show or not when importing a particular type of asset.*/
-	UPROPERTY(EditAnywhere, Category = "Pipeline", Meta = (DisplayAfter = "bShowImportDialog"))
+	UPROPERTY(EditAnywhere, Category = "DialogOverride")
 	TMap<EInterchangeTranslatorAssetType, FInterchangeDialogOverride> ShowImportDialogOverride;
 };
 
@@ -151,7 +178,7 @@ public:
 	 * Settings used when importing into a level.
 	 */
 	UPROPERTY(EditAnywhere, config, Category = "ImportIntoLevel")
-	FInterchangeImportSettings SceneImportSettings;
+	FInterchangeSceneImportSettings SceneImportSettings;
 
 	/** This tells Interchange which file picker class to construct when we need to choose a file for a source. */
 	UPROPERTY(EditAnywhere, config, Category = "EditorInterface")
@@ -214,8 +241,9 @@ class UInterchangeEditorSettings : public UDeveloperSettings
 	GENERATED_BODY()
 
 public:
-	/** If enabled, the import option dialog will show when interchange re-import.*/
-	UPROPERTY(EditAnywhere, config, Category = "Show Dialog")
+	/** Use the dialog overrides on the Interchange category of the Project Settings if you want to control when the Interchange reimport dialog is shown */
+	UE_DEPRECATED(5.6, "Use the dialog overrides on the Interchange category of the Project Settings if you want to control when the Interchange reimport dialog is shown")
+	UPROPERTY(EditAnywhere, config, Category = "Show Dialog", meta=(EditCondition = "false"))
 	bool bShowImportDialogAtReimport = false;
 
 	const FGuid& GetUsedGroupUID() const { return UsedGroupUID; }
@@ -229,7 +257,6 @@ public:
 	void UpdateUsedGroupName();
 
 private:
-	/** If enabled, the import option dialog will show when interchange re-import.*/
 	UPROPERTY(EditAnywhere, Transient, Category = "Group Used", meta = (GetOptions = "GetSelectableItems", AllowPrivateAccess = "true"))
 	FName UsedGroupName;
 
