Interchange: Fix tooltip on deprecated editor settings property.

diff --git a/Engine/Source/Runtime/Interchange/Engine/Public/InterchangeProjectSettings.h b/Engine/Source/Runtime/Interchange/Engine/Public/InterchangeProjectSettings.h
--- a/Engine/Source/Runtime/Interchange/Engine/Public/InterchangeProjectSettings.h
+++ b/Engine/Source/Runtime/Interchange/Engine/Public/InterchangeProjectSettings.h
@@ -241,9 +241,8 @@ class UInterchangeEditorSettings : public UDeveloperSettings
 	GENERATED_BODY()
 
 public:
-	/** Use the dialog overrides on the Interchange category of the Project Settings if you want to control when the Interchange reimport dialog is shown */
 	UE_DEPRECATED(5.6, "Use the dialog overrides on the Interchange category of the Project Settings if you want to control when the Interchange reimport dialog is shown")
-	UPROPERTY(EditAnywhere, config, Category = "Show Dialog", meta=(EditCondition = "false"))
+	UPROPERTY(EditAnywhere, config, Category = "Show Dialog", meta=(EditCondition = "false", ToolTip="Use the dialog overrides on the Interchange category of the Project Settings if you want to control when the Interchange reimport dialog is shown"))
 	bool bShowImportDialogAtReimport = false;
 
 	const FGuid& GetUsedGroupUID() const { return UsedGroupUID; }
