[Interchange] Task Parsing will run on Game Thread
- InterchangeTaskParsing was causing issues with GC when running on AsyncThread.
- Removing an unused lambda function in the Execute()
- Currently changing it to run on the game thread.
- Will be revisited in future.

#rnx

diff --git a/Engine/Source/Runtime/Interchange/Engine/Private/Tasks/InterchangeTaskParsing.cpp b/Engine/Source/Runtime/Interchange/Engine/Private/Tasks/InterchangeTaskParsing.cpp
--- a/Engine/Source/Runtime/Interchange/Engine/Private/Tasks/InterchangeTaskParsing.cpp
+++ b/Engine/Source/Runtime/Interchange/Engine/Private/Tasks/InterchangeTaskParsing.cpp
@@ -206,34 +206,6 @@ void UE::Interchange::FTaskParsing::Execute()
 		}
 	}
 
-	auto CreateTasksForEachTaskData = [](TArray<FTaskData>& TaskDatas, TFunction<uint64(FTaskData&)> CreateTasksFunc) -> TArray<uint64>
-	{
-		TArray<uint64> GraphEvents;
-
-		for (int32 TaskIndex = 0; TaskIndex < TaskDatas.Num(); ++TaskIndex)
-		{
-			FTaskData& TaskData = TaskDatas[TaskIndex];
-
-			if (TaskData.Dependencies.Num() > 0)
-			{
-				//Search the previous node to find the dependence
-				for (int32 DepTaskIndex = 0; DepTaskIndex < TaskIndex; ++DepTaskIndex)
-				{
-					if (TaskData.Dependencies.Contains(TaskDatas[DepTaskIndex].UniqueID))
-					{
-						//Add has prerequisite
-						TaskData.Prerequisites.Add(TaskDatas[DepTaskIndex].GraphEventRef);
-					}
-				}
-			}
-
-			TaskData.GraphEventRef = CreateTasksFunc(TaskData);
-			GraphEvents.Add(TaskData.GraphEventRef);
-		}
-
-		return GraphEvents;
-	};
-
 	struct FTaskParsingRenameInfo
 	{
 		UInterchangeFactoryBaseNode* FactoryNode = nullptr;
diff --git a/Engine/Source/Runtime/Interchange/Engine/Private/Tasks/InterchangeTaskParsing.h b/Engine/Source/Runtime/Interchange/Engine/Private/Tasks/InterchangeTaskParsing.h
--- a/Engine/Source/Runtime/Interchange/Engine/Private/Tasks/InterchangeTaskParsing.h
+++ b/Engine/Source/Runtime/Interchange/Engine/Private/Tasks/InterchangeTaskParsing.h
@@ -28,13 +28,8 @@ namespace UE
 
 			virtual EInterchangeTaskThread GetTaskThread() const override
 			{
-				TSharedPtr<FImportAsyncHelper, ESPMode::ThreadSafe> AsyncHelper = WeakAsyncHelper.Pin();
-				if (AsyncHelper.IsValid() && AsyncHelper->bRunSynchronous)
-				{
-					return EInterchangeTaskThread::GameThread;
-				}
-
-				return EInterchangeTaskThread::AsyncThread;
+				// TODO: Find a way to run this task on Async Thread without causing issues for GC.
+				return EInterchangeTaskThread::GameThread;
 			}
 
 			virtual void Execute() override;
