[BUGFIX] Interchange carash upon reseting scene

- Object Instances to reset are now checked for validity.

#rnx

diff --git a/Engine/Source/Runtime/Interchange/Engine/Private/InterchangeImportReset.cpp b/Engine/Source/Runtime/Interchange/Engine/Private/InterchangeImportReset.cpp
--- a/Engine/Source/Runtime/Interchange/Engine/Private/InterchangeImportReset.cpp
+++ b/Engine/Source/Runtime/Interchange/Engine/Private/InterchangeImportReset.cpp
@@ -161,20 +161,28 @@ TObjectPtr<UInterchangeFactoryBase> FInterchangeResetParameters::GetFactoryForNo
 
 TArray<TObjectPtr<UObject>> FInterchangeResetParameters::GetObjectInstancesForFactoryNode(const UInterchangeFactoryBaseNode* FactoryNode)
 {
+	TArray<TObjectPtr<UObject>> OutObjectInstances;
+
 	if (FactoryNodeDataCache.Contains(FactoryNode))
 	{
 		const FInterchangeResetParameters::FFactoryNodeData& FactoryNodeData = FactoryNodeDataCache[FactoryNode];
-		if (FactoryNodeData.ObjectsToReset.Num())
+		if (!FactoryNodeData.ObjectsToReset.IsEmpty())
 		{
-			return FactoryNodeData.ObjectsToReset;
+			for (const TObjectPtr<UObject>& ObjectInstance : FactoryNodeData.ObjectsToReset)
+			{
+				if (ObjectInstance.Get() != nullptr)
+				{
+					OutObjectInstances.Add(ObjectInstance);
+				}
+			}
 		}
-		else
+		else if (UObject* ReferencedObject = Cast<UObject>(FactoryNodeData.ReferencedObjectPath.TryLoad()))
 		{
-			return TArray<TObjectPtr<UObject>>{ Cast<UObject>(FactoryNodeData.ReferencedObjectPath.TryLoad()) };
+			OutObjectInstances.Add(ReferencedObject);
 		}
 	}
 
-	return TArray<TObjectPtr<UObject>>();
+	return OutObjectInstances;
 }
 
 void FInterchangeResetParameters::AddObjectInstanceToReset(const UInterchangeFactoryBaseNode* FactoryNode, UObject* ObjectToReset)
