Interchange:
- Have the asset reimport strategies visible upon reimport on pipeline assets
- Add reimport strategy method to pass it to Variant Set Factory node from the asset pipeline

#virtualized

diff --git a/Engine/Plugins/Interchange/Runtime/Source/Pipelines/Private/InterchangeGenericAssetsPipeline.cpp b/Engine/Plugins/Interchange/Runtime/Source/Pipelines/Private/InterchangeGenericAssetsPipeline.cpp
--- a/Engine/Plugins/Interchange/Runtime/Source/Pipelines/Private/InterchangeGenericAssetsPipeline.cpp
+++ b/Engine/Plugins/Interchange/Runtime/Source/Pipelines/Private/InterchangeGenericAssetsPipeline.cpp
@@ -1273,6 +1273,9 @@ void UInterchangeGenericAssetsPipeline::ExecutePipeline(UInterchangeBaseNodeCont
 
 	// Setup the source node sub paths, which will be used here but also by all other pipelines that run after us and create asset factory nodes
 	UInterchangeSourceNode* SourceNode = UInterchangeSourceNode::FindOrCreateUniqueInstance(InBaseNodeContainer);
+
+	SourceNode->SetCustomReimportStrategyFlags(int8(ReimportStrategy));
+
 	if (bAssetTypeSubFolders)
 	{
 		SourceNode->SetCustomUseAssetTypeSubPathSuffix(bAssetTypeSubFolders);
diff --git a/Engine/Plugins/Interchange/Runtime/Source/Pipelines/Private/InterchangeGenericScenesPipeline.cpp b/Engine/Plugins/Interchange/Runtime/Source/Pipelines/Private/InterchangeGenericScenesPipeline.cpp
--- a/Engine/Plugins/Interchange/Runtime/Source/Pipelines/Private/InterchangeGenericScenesPipeline.cpp
+++ b/Engine/Plugins/Interchange/Runtime/Source/Pipelines/Private/InterchangeGenericScenesPipeline.cpp
@@ -1153,6 +1153,12 @@ void UInterchangeGenericLevelPipeline::ExecuteSceneVariantSetNodePreImport(const
 	UInterchangeSourceNode* SourceNode = UInterchangeSourceNode::FindOrCreateUniqueInstance(BaseNodeContainer);
 	UE::Interchange::PipelineHelper::FillSubPathFromSourceNode(FactoryNode, SourceNode);
 
+	// The reimport strategy should be the one from the InterchangeAssetsPipeline not the Level
+	if (EReimportStrategyFlags ReimportStrategy; SourceNode->GetCustomReimportStrategyFlags(reinterpret_cast<uint8&>(ReimportStrategy)))
+	{
+		FactoryNode->SetReimportStrategyFlags(ReimportStrategy);
+	}
+
 	TArray<FString> VariantSetUids;
 	SceneVariantSetNode.GetCustomVariantSetUids(VariantSetUids);
 
diff --git a/Engine/Source/Runtime/Interchange/Core/Private/Nodes/InterchangeSourceNode.cpp b/Engine/Source/Runtime/Interchange/Core/Private/Nodes/InterchangeSourceNode.cpp
--- a/Engine/Source/Runtime/Interchange/Core/Private/Nodes/InterchangeSourceNode.cpp
+++ b/Engine/Source/Runtime/Interchange/Core/Private/Nodes/InterchangeSourceNode.cpp
@@ -206,3 +206,13 @@ bool UInterchangeSourceNode::SetCustomUseAssetTypeSubPathSuffix(const bool& Attr
 	IMPLEMENT_NODE_ATTRIBUTE_SETTER_NODELEGATE(UseAssetTypeSubPathSuffix, bool);
 }
 
+INTERCHANGECORE_API bool UInterchangeSourceNode::GetCustomReimportStrategyFlags(uint8& AttributeValue) const
+{
+	IMPLEMENT_NODE_ATTRIBUTE_GETTER(ReimportStrategyFlags, uint8);
+}
+
+INTERCHANGECORE_API bool UInterchangeSourceNode::SetCustomReimportStrategyFlags(uint8 AttributeValue)
+{
+	IMPLEMENT_NODE_ATTRIBUTE_SETTER_NODELEGATE(ReimportStrategyFlags, uint8)
+}
+
diff --git a/Engine/Source/Runtime/Interchange/Core/Public/Nodes/InterchangeSourceNode.h b/Engine/Source/Runtime/Interchange/Core/Public/Nodes/InterchangeSourceNode.h
--- a/Engine/Source/Runtime/Interchange/Core/Public/Nodes/InterchangeSourceNode.h
+++ b/Engine/Source/Runtime/Interchange/Core/Public/Nodes/InterchangeSourceNode.h
@@ -3,15 +3,10 @@
 #pragma once
 
 #include "Containers/UnrealString.h"
-#include "CoreMinimal.h"
 #include "HAL/Platform.h"
 #include "Nodes/InterchangeBaseNode.h"
 #include "Nodes/InterchangeBaseNodeUtilities.h"
-#include "Types/AttributeStorage.h"
-#include "UObject/Class.h"
-#include "UObject/Object.h"
 #include "UObject/ObjectMacros.h"
-#include "UObject/UObjectGlobals.h"
 
 #include "InterchangeSourceNode.generated.h"
 
@@ -188,6 +183,14 @@ public:
 	UFUNCTION(BlueprintCallable, Category = "Interchange | Node | Source")
 	INTERCHANGECORE_API bool SetCustomUseAssetTypeSubPathSuffix(const bool& Suffix);
 
+	/** Get the reimport strategy based on EReimportStrategyFlags */
+	UFUNCTION(BlueprintCallable, Category = "Interchange | Node | Source")
+	INTERCHANGECORE_API bool GetCustomReimportStrategyFlags(uint8& StrategyFlag) const;
+
+	/** Set the reimport strategy based on EReimportStrategyFlags */
+	UFUNCTION(BlueprintCallable, Category = "Interchange | Node | Source")
+	INTERCHANGECORE_API bool SetCustomReimportStrategyFlags(uint8 StrategyFlag);
+
 private:
 	IMPLEMENT_NODE_ATTRIBUTE_KEY(SourceFrameRateNumerator);
 	IMPLEMENT_NODE_ATTRIBUTE_KEY(SourceFrameRateDenominator);
@@ -200,6 +203,7 @@ private:
 	IMPLEMENT_NODE_ATTRIBUTE_KEY(UseLegacySkeletalMeshBakeTransform);
 	IMPLEMENT_NODE_ATTRIBUTE_KEY(SubPathPrefix);
 	IMPLEMENT_NODE_ATTRIBUTE_KEY(UseAssetTypeSubPathSuffix);
+	IMPLEMENT_NODE_ATTRIBUTE_KEY(ReimportStrategyFlags);
 
 	// Extra InformationTo show in the Config Panel.
 	UE::Interchange::TMapAttributeHelper<FString, FString> ExtraInformation;
