Interchange:
- Common operator in Interchange to handle long filenames and reduce them to protect file system length.
- Add CVar Interchange.MaxPathLength to control the budget on filename during the creation of assets

diff --git a/Engine/Source/Runtime/Interchange/Core/Private/InterchangeHelper.cpp b/Engine/Source/Runtime/Interchange/Core/Private/InterchangeHelper.cpp
--- a/Engine/Source/Runtime/Interchange/Core/Private/InterchangeHelper.cpp
+++ b/Engine/Source/Runtime/Interchange/Core/Private/InterchangeHelper.cpp
@@ -1,9 +1,56 @@
 // Copyright Epic Games, Inc. All Rights Reserved.
 
 #include "InterchangeHelper.h"
+#include "InterchangeLogPrivate.h"
+#include "HAL/IConsoleManager.h"
+#include "Misc/Paths.h"
 
 namespace UE::Interchange
 {
+	namespace Private
+	{
+		int32 GInterchangeMaxAssetPathLength = 160; // a base value that should suits most cases
+		static FAutoConsoleVariableRef CVarInterchangeMaxAssetPathLength(
+			TEXT("Interchange.MaxAssetPathLength"),
+			GInterchangeMaxAssetPathLength,
+			TEXT("Interchange will try to limit asset path length to this value. Default: 260")
+		);
+
+		int32 GetUsableMaxAssetPathLength()
+		{
+			const int32 MinWorkable = 60;
+			static int32 ProjectConstraint; // Deduced from os limit and project path
+			static bool RunOnce = [&]()
+			{
+				const FString ProjectContentDir = FPaths::ProjectContentDir();
+				const FString FullPathProjectContentDir = FPaths::ConvertRelativePathToFull(ProjectContentDir);
+				ProjectConstraint = FPlatformMisc::GetMaxPathLength() - FullPathProjectContentDir.Len();
+
+				if (ProjectConstraint < MinWorkable)
+				{
+					UE_LOG(LogInterchangeCore, Error,
+							TEXT("Interchange can encounter import issues due to a Content path too long, and an OS limitation on path length.\n")
+							TEXT("Content path: '%s'\n")
+							TEXT("System max path length: %d.\n")
+							, *FullPathProjectContentDir, FPlatformMisc::GetMaxPathLength());
+				}
+
+				if (GInterchangeMaxAssetPathLength > ProjectConstraint)
+				{
+					UE_LOG(LogInterchangeCore, Warning,
+							TEXT("The Interchange.MaxAssetPathLength value (%d) is too high for the current setup.\n")
+							TEXT("Content path: '%s'\n")
+							TEXT("System max path length: %d.\n")
+							, GInterchangeMaxAssetPathLength, *FullPathProjectContentDir, FPlatformMisc::GetMaxPathLength());
+				}
+
+				return true;
+			}();
+
+			return FMath::Max(FMath::Min(GInterchangeMaxAssetPathLength, ProjectConstraint), MinWorkable);
+		}
+	}
+
 	void SanitizeName(FString& OutName, bool bIsJoint)
 	{
 		const TCHAR* InvalidChar = bIsJoint ? INVALID_OBJECTNAME_CHARACTERS TEXT("+ ") : INVALID_OBJECTNAME_CHARACTERS;
@@ -26,4 +73,47 @@ namespace UE::Interchange
 		SanitizeName(TmpName, bIsJoint);
 		return TmpName;
 	}
+
+	int32 GetAssetNameMaxCharCount(const FString& ParentPackage)
+	{
+		// can be tweaked, the goal is to be more restrictive than the filesystem
+		// so that a project can be shared / moved without breaking the constraint
+		int32 MaxAssetPathLength = Private::GetUsableMaxAssetPathLength();
+
+		// internal limit of FNames + room for prefix, separators and null char. (Asset names occur twice in paths)
+		int32 InternalNameConstraint = (NAME_SIZE - 100);
+
+		int32 PackageLength = 1 + (!ParentPackage.IsEmpty() ? ParentPackage.Len() : 20);
+		int32 Budget = FMath::Min((InternalNameConstraint - PackageLength) / 2, MaxAssetPathLength - PackageLength);
+		Budget = FMath::Min(Budget, 255 - 10); // a filename cannot be longer than 255, and we keep a small buffer for the extension
+		return FMath::Max(0, Budget);
+	}
+
+	FString DesiredAssetName(const FString& ParentPackage, int32 PackageLength, int32 CharBudget, TCHAR CharReplacement)
+	{
+		if (ParentPackage.Len() + PackageLength <= CharBudget)
+		{
+			return ParentPackage;
+		}
+
+		// Arbitrary number to avoid possible name collisions and having a very low char budget
+		if (CharBudget <= 5)
+		{
+			UE_LOG(LogInterchangeCore,
+				   Warning,
+				   TEXT("Char budget is too small for generating a new asset name, please adapt the name of the asset manually or try modifying the value of Interchange.MaxAssetPathLength"),
+				   *ParentPackage);
+			return ParentPackage;
+		}
+
+		// Space for beginning, replacement and ending
+		int32 Remaining = CharBudget - 2;
+		int32 LeftCount = Remaining / 2;
+		int32 RightCount = Remaining - LeftCount;
+
+		FString Left = ParentPackage.Mid(0, LeftCount + 1);
+		FString Right = ParentPackage.Mid(ParentPackage.Len() - (RightCount + 1));
+
+		return Left + CharReplacement + Right;
+	}
 };
\ No newline at end of file
diff --git a/Engine/Source/Runtime/Interchange/Core/Public/InterchangeHelper.h b/Engine/Source/Runtime/Interchange/Core/Public/InterchangeHelper.h
--- a/Engine/Source/Runtime/Interchange/Core/Public/InterchangeHelper.h
+++ b/Engine/Source/Runtime/Interchange/Core/Public/InterchangeHelper.h
@@ -14,6 +14,25 @@ namespace UE::Interchange
 	INTERCHANGECORE_API void SanitizeName(FString& OutName, bool bIsJoint = false);
 	INTERCHANGECORE_API FString MakeName(const FString& InName, bool bIsJoint = false);
 
+	/**
+	 * Try to compute a char budget for asset names, including FNames constraints, OS constraints,
+	 * parent package, and user defined limitation.
+	 *
+	 * @param ParentPackage destination of the asset (Package path consume a part of the budget)
+	 * @return An estimation of the budget for asset names
+	 */
+	INTERCHANGECORE_API int32 GetAssetNameMaxCharCount(const FString& ParentPackage);
+
+	/**
+	 * Generate a new name for ParentPackage based on a character budget. Any characters out of the budget limit will be stripped out
+	 * @param ParentPackage - the name of the package that needs to fulfill the OS constraints
+	 * @param PackageLegth - length of the package, this value can differ from the length of ParentPackage
+	 * @param CharBudget - characters limit for the new generated name (if 0, the string will be returned with a warning, meaning we couldn't generate a new name)
+	 * @param CharReplacement - the character replacement of the trimmed substring
+	 * @return The new generated name if out of budget, or just the string if it's in the budget limitation
+	 */
+	INTERCHANGECORE_API FString DesiredAssetName(const FString& ParentPackage, int32 PackageLength, int32 CharBudget, TCHAR CharReplacement = TEXT('_'));
+
 	class FScopedLambda
 	{
 	public:
diff --git a/Engine/Source/Runtime/Interchange/Engine/Private/Tasks/InterchangeTaskImportObject.cpp b/Engine/Source/Runtime/Interchange/Engine/Private/Tasks/InterchangeTaskImportObject.cpp
--- a/Engine/Source/Runtime/Interchange/Engine/Private/Tasks/InterchangeTaskImportObject.cpp
+++ b/Engine/Source/Runtime/Interchange/Engine/Private/Tasks/InterchangeTaskImportObject.cpp
@@ -25,6 +25,7 @@
 #include "UObject/ObjectMacros.h"
 #include "UObject/ObjectRedirector.h"
 #include "UObject/WeakObjectPtrTemplates.h"
+#include "InterchangeHelper.h"
 
 namespace UE::Interchange::Private
 {
@@ -40,7 +41,14 @@ namespace UE::Interchange::Private
 		UInterchangeManager::GetInterchangeManager().SanitizeNameInline(OutAssetName, ESanitizeNameTypeFlags::ObjectName | ESanitizeNameTypeFlags::ObjectPath | ESanitizeNameTypeFlags::LongPackage);
 		FString SubPath;
 		FactoryNode->GetCustomSubPath(SubPath);
-		OutPackageName = FPaths::Combine(*PackageBasePath, *SubPath, *OutAssetName);
+
+		FString PackagePath = FPaths::Combine(*PackageBasePath, *SubPath);
+		int32 MaxPath = GetAssetNameMaxCharCount(*PackagePath);
+
+		// We don't want to take into account the '/Game' length, hence number 5
+		OutAssetName = DesiredAssetName(OutAssetName, PackagePath.Len() - 5, MaxPath);
+
+		OutPackageName = FPaths::Combine(*PackagePath, *OutAssetName);
 		UInterchangeManager::GetInterchangeManager().SanitizeNameInline(OutPackageName, ESanitizeNameTypeFlags::ObjectPath | ESanitizeNameTypeFlags::LongPackage);
 	}
 
