Fixed name doubling issue

#rnx

diff --git a/Engine/Source/Runtime/Interchange/Core/Private/InterchangeHelper.cpp b/Engine/Source/Runtime/Interchange/Core/Private/InterchangeHelper.cpp
--- a/Engine/Source/Runtime/Interchange/Core/Private/InterchangeHelper.cpp
+++ b/Engine/Source/Runtime/Interchange/Core/Private/InterchangeHelper.cpp
@@ -89,11 +89,11 @@ namespace UE::Interchange
 		return FMath::Max(0, Budget);
 	}
 
-	FString DesiredAssetName(const FString& ParentPackage, int32 PackageLength, int32 CharBudget, TCHAR CharReplacement)
+	FString GetAssetNameWBudget(const FString& DesiredAssetName, int32 CharBudget, TCHAR CharReplacement)
 	{
-		if (ParentPackage.Len() + PackageLength <= CharBudget)
+		if (DesiredAssetName.Len() <= CharBudget)
 		{
-			return ParentPackage;
+			return DesiredAssetName;
 		}
 
 		// Arbitrary number to avoid possible name collisions and having a very low char budget
@@ -102,8 +102,8 @@ namespace UE::Interchange
 			UE_LOG(LogInterchangeCore,
 				   Warning,
 				   TEXT("Char budget is too small for generating a new asset name, please adapt the name of the asset manually or try modifying the value of Interchange.MaxAssetPathLength"),
-				   *ParentPackage);
-			return ParentPackage;
+				   *DesiredAssetName);
+			return DesiredAssetName;
 		}
 
 		// Space for beginning, replacement and ending
@@ -111,8 +111,8 @@ namespace UE::Interchange
 		int32 LeftCount = Remaining / 2;
 		int32 RightCount = Remaining - LeftCount;
 
-		FString Left = ParentPackage.Mid(0, LeftCount + 1);
-		FString Right = ParentPackage.Mid(ParentPackage.Len() - (RightCount + 1));
+		FString Left = DesiredAssetName.Mid(0, LeftCount + 1);
+		FString Right = DesiredAssetName.Mid(DesiredAssetName.Len() - (RightCount + 1));
 
 		return Left + CharReplacement + Right;
 	}
diff --git a/Engine/Source/Runtime/Interchange/Core/Public/InterchangeHelper.h b/Engine/Source/Runtime/Interchange/Core/Public/InterchangeHelper.h
--- a/Engine/Source/Runtime/Interchange/Core/Public/InterchangeHelper.h
+++ b/Engine/Source/Runtime/Interchange/Core/Public/InterchangeHelper.h
@@ -24,14 +24,13 @@ namespace UE::Interchange
 	INTERCHANGECORE_API int32 GetAssetNameMaxCharCount(const FString& ParentPackage);
 
 	/**
-	 * Generate a new name for ParentPackage based on a character budget. Any characters out of the budget limit will be stripped out
-	 * @param ParentPackage - the name of the package that needs to fulfill the OS constraints
-	 * @param PackageLegth - length of the package, this value can differ from the length of ParentPackage
+	 * Generate a new name for DesiredAssetName based on a character budget. Any characters out of the budget limit will be stripped out
+	 * @param DesiredAssetName - the name of the package that needs to fulfill the OS constraints
 	 * @param CharBudget - characters limit for the new generated name (if 0, the string will be returned with a warning, meaning we couldn't generate a new name)
 	 * @param CharReplacement - the character replacement of the trimmed substring
 	 * @return The new generated name if out of budget, or just the string if it's in the budget limitation
 	 */
-	INTERCHANGECORE_API FString DesiredAssetName(const FString& ParentPackage, int32 PackageLength, int32 CharBudget, TCHAR CharReplacement = TEXT('_'));
+	INTERCHANGECORE_API FString GetAssetNameWBudget(const FString& ParentPackage, int32 CharBudget, TCHAR CharReplacement = TEXT('_'));
 
 	class FScopedLambda
 	{
diff --git a/Engine/Source/Runtime/Interchange/Engine/Private/Tasks/InterchangeTaskImportObject.cpp b/Engine/Source/Runtime/Interchange/Engine/Private/Tasks/InterchangeTaskImportObject.cpp
--- a/Engine/Source/Runtime/Interchange/Engine/Private/Tasks/InterchangeTaskImportObject.cpp
+++ b/Engine/Source/Runtime/Interchange/Engine/Private/Tasks/InterchangeTaskImportObject.cpp
@@ -46,7 +46,7 @@ namespace UE::Interchange::Private
 		int32 MaxPath = GetAssetNameMaxCharCount(*PackagePath);
 
 		// We don't want to take into account the '/Game' length, hence number 5
-		OutAssetName = DesiredAssetName(OutAssetName, PackagePath.Len() - 5, MaxPath);
+		OutAssetName = GetAssetNameWBudget(OutAssetName, MaxPath);
 
 		OutPackageName = FPaths::Combine(*PackagePath, *OutAssetName);
 		UInterchangeManager::GetInterchangeManager().SanitizeNameInline(OutPackageName, ESanitizeNameTypeFlags::ObjectPath | ESanitizeNameTypeFlags::LongPackage);
