Fixed ensure when re-importing an asset that was import with a blueprint pipeline

#rnx

diff --git a/Engine/Source/Runtime/Interchange/Engine/Private/InterchangeAssetImportData.cpp b/Engine/Source/Runtime/Interchange/Engine/Private/InterchangeAssetImportData.cpp
--- a/Engine/Source/Runtime/Interchange/Engine/Private/InterchangeAssetImportData.cpp
+++ b/Engine/Source/Runtime/Interchange/Engine/Private/InterchangeAssetImportData.cpp
@@ -1,6 +1,7 @@
 // Copyright Epic Games, Inc. All Rights Reserved.
 
 #include "InterchangeAssetImportData.h"
+#include "InterchangeBlueprintPipelineBase.h"
 #include "InterchangeCustomVersion.h"
 #include "InterchangeEngineLogPrivate.h"
 #include "InterchangeManager.h"
@@ -14,6 +15,23 @@
 
 #include UE_INLINE_GENERATED_CPP_BY_NAME(InterchangeAssetImportData)
 
+TMap<FString, UClass*> UInterchangeAssetImportData::GetPipelineClassPerName()
+{
+	TMap<FString, UClass*> ClassPerName;
+
+	TArray<UClass*> PipelineCandidates;
+	UInterchangeManager::GetInterchangeManager().FindPipelineCandidate(PipelineCandidates);
+
+	ClassPerName.Reserve(PipelineCandidates.NumBytes());
+
+	for (UClass* PipelineClass : PipelineCandidates)
+	{
+		ClassPerName.Add(PipelineClass->GetFullName(), PipelineClass);
+	}
+
+	return ClassPerName;
+}
+
 void UInterchangeAssetImportData::PostLoad()
 {
 	Super::PostLoad();
@@ -401,15 +419,7 @@ void UInterchangeAssetImportData::ProcessPipelinesCache() const
 		{
 			TransientPipelines.Reset(CachedPipelines.Num());
 
-			TMap<FString, UClass*> ClassPerName;
-			for (FThreadSafeObjectIterator It(UClass::StaticClass()); It; ++It)
-			{
-				UClass* Class = Cast<UClass>(*It);
-				if (Class->IsChildOf(UInterchangePipelineBase::StaticClass()))
-				{
-					ClassPerName.Add(Class->GetFullName(), Class);
-				}
-			}
+			TMap<FString, UClass*> ClassPerName = GetPipelineClassPerName();
 
 			for (const TPair<FString, FString>& CachedPipeline : CachedPipelines)
 			{
@@ -419,10 +429,12 @@ void UInterchangeAssetImportData::ProcessPipelinesCache() const
 				{
 					ClassFullName = RedirectedObjectName.ToString();
 				}
-				//This cannot fail to make sure we have a healty serialization
+
+				//This cannot fail to make sure we have a healthy serialization
 				if (!ensure(ClassPerName.Contains(ClassFullName)))
 				{
-					//We did not successfully serialize the content of the file into the node container
+					//We did not successfully serialize the cached pipelines
+					TransientPipelines.Empty();
 					return;
 				}
 
diff --git a/Engine/Source/Runtime/Interchange/Engine/Private/InterchangeManager.cpp b/Engine/Source/Runtime/Interchange/Engine/Private/InterchangeManager.cpp
--- a/Engine/Source/Runtime/Interchange/Engine/Private/InterchangeManager.cpp
+++ b/Engine/Source/Runtime/Interchange/Engine/Private/InterchangeManager.cpp
@@ -2926,17 +2926,14 @@ bool UInterchangeManager::IsAttended()
 	return true;
 }
 
+// #todo_interchange: TODO: Can be costly on large project. Using a cache and using the AssetRegistry to monitor addition will be more efficient.
+// This code could also be moved to a utility class.
 void UInterchangeManager::FindPipelineCandidate(TArray<UClass*>& PipelineCandidates)
 {
 	//Find in memory pipeline class
 	for (TObjectIterator< UClass > ClassIt; ClassIt; ++ClassIt)
 	{
 		UClass* Class = *ClassIt;
-		// Only interested in native C++ classes
-// 		if (!Class->IsNative())
-// 		{
-// 			continue;
-// 		}
 		// Ignore deprecated
 		if (Class->HasAnyClassFlags(CLASS_Deprecated | CLASS_NewerVersionExists))
 		{
@@ -2960,7 +2957,7 @@ void UInterchangeManager::FindPipelineCandidate(TArray<UClass*>& PipelineCandida
 
 	TArray< FString > ContentPaths;
 	ContentPaths.Add(TEXT("/Game"));
-	//TODO do we have an other alternative, this call is synchronous and will wait unitl the registry database have finish the initial scan. If there is a lot of asset it can take multiple second the first time we call it.
+	//#todo_interchange: TODO do we have an other alternative, this call is synchronous and will wait unitl the registry database have finish the initial scan. If there is a lot of asset it can take multiple second the first time we call it.
 	AssetRegistry.ScanPathsSynchronous(ContentPaths);
 
 	FTopLevelAssetPath BaseClassName = UInterchangePipelineBase::StaticClass()->GetClassPathName();
@@ -3000,9 +2997,10 @@ void UInterchangeManager::FindPipelineCandidate(TArray<UClass*>& PipelineCandida
 			}
 
 			UBlueprint* Blueprint = Cast<UBlueprint>(Asset.GetAsset());
-			check(Blueprint);
-			check(Blueprint->ParentClass == UInterchangePipelineBase::StaticClass());
-			PipelineCandidates.AddUnique(Blueprint->GeneratedClass);
+			if (ensure(Blueprint) && ensure(Blueprint->ParentClass->IsChildOf<UInterchangePipelineBase>()))
+			{
+				PipelineCandidates.AddUnique(Blueprint->GeneratedClass);
+			}
 		}
 	}
 }
diff --git a/Engine/Source/Runtime/Interchange/Engine/Public/InterchangeAssetImportData.h b/Engine/Source/Runtime/Interchange/Engine/Public/InterchangeAssetImportData.h
--- a/Engine/Source/Runtime/Interchange/Engine/Public/InterchangeAssetImportData.h
+++ b/Engine/Source/Runtime/Interchange/Engine/Public/InterchangeAssetImportData.h
@@ -201,6 +201,8 @@ private:
 	/** Source file data describing the files that were used to import this asset. Temporary, Primary usage is for Re-import cancellations. */
 	mutable FAssetImportInfo SourceDataBackup;
 #endif
+	// #todo_interchange: TODO: Move to as a IntercnageManager utility method
+	static TMap<FString, UClass*> GetPipelineClassPerName();
 };
 
 /**
diff --git a/Engine/Source/Runtime/Interchange/Engine/Public/InterchangeManager.h b/Engine/Source/Runtime/Interchange/Engine/Public/InterchangeManager.h
--- a/Engine/Source/Runtime/Interchange/Engine/Public/InterchangeManager.h
+++ b/Engine/Source/Runtime/Interchange/Engine/Public/InterchangeManager.h
@@ -1013,4 +1013,6 @@ private:
 	bool bGCEndDelegateCancellAllTask = false;
 
 	friend class UE::Interchange::FScopedTranslator;
+	// #todo_interchange: Remove wen Interchange utility method for pipelines is done
+	friend class UInterchangeAssetImportData;
 };
