[BUGFIX] Interchange, folder creation option inherited from previous import
- Cached Pipeline Settings now additionally use a Package Path Hash to distinguish pipelines.
- InterchangePipelineSettingsCacheHandler is introduced to remove the cached settings for pipeline assets that get deleted in the editor.
- Subpipelines use the path hash of the outer and append the classname to it to create config section.

Unhandled Scenarios:
- Pipeline Assets being renamed and moved will not update the cached settings to correctly reflect the new path/name in the section headings
- Pipeline Assets that are created, cached but not saved in the editor will leave behind the cached settings bloating the ini file.

#rnx

diff --git a/Engine/Plugins/Interchange/Editor/Source/InterchangeEditor/Private/InterchangeEditorModule.cpp b/Engine/Plugins/Interchange/Editor/Source/InterchangeEditor/Private/InterchangeEditorModule.cpp
--- a/Engine/Plugins/Interchange/Editor/Source/InterchangeEditor/Private/InterchangeEditorModule.cpp
+++ b/Engine/Plugins/Interchange/Editor/Source/InterchangeEditor/Private/InterchangeEditorModule.cpp
@@ -1,6 +1,7 @@
 // Copyright Epic Games, Inc. All Rights Reserved.
 #include "InterchangeEditorModule.h"
 
+#include "AssetRegistry/IAssetRegistry.h"
 #include "AssetToolsModule.h"
 #include "Engine/Engine.h"
 #include "IMessageLogListing.h"
@@ -8,6 +9,7 @@
 #include "InterchangeFbxAssetImportDataConverter.h"
 #include "InterchangeManager.h"
 #include "InterchangeResetContextMenuExtender.h"
+#include "InterchangePipelineSettingsCacheHandler.h"
 #include "MessageLogModule.h"
 #include "Misc/App.h"
 #include "Misc/CoreDelegates.h"
@@ -150,10 +152,12 @@ void FInterchangeEditorModule::StartupModule()
 	}
 
 	FInterchangeResetContextMenuExtender::SetupLevelEditorContextMenuExtender();
+	FInterchangePipelineSettingsCacheHandler::InitializeCacheHandler();
 }
 
 void FInterchangeEditorModule::ShutdownModule()
 {
+	FInterchangePipelineSettingsCacheHandler::ShutdownCacheHandler();
 	FInterchangeResetContextMenuExtender::RemoveLevelEditorContextMenuExtender();
 }
 
diff --git a/Engine/Plugins/Interchange/Editor/Source/InterchangeEditor/Private/InterchangePipelineSettingsCacheHandler.cpp b/Engine/Plugins/Interchange/Editor/Source/InterchangeEditor/Private/InterchangePipelineSettingsCacheHandler.cpp
new file mode 100644
--- /dev/null
+++ b/Engine/Plugins/Interchange/Editor/Source/InterchangeEditor/Private/InterchangePipelineSettingsCacheHandler.cpp
@@ -0,0 +1,51 @@
+// Copyright Epic Games, Inc. All Rights Reserved.
+
+#include "InterchangePipelineSettingsCacheHandler.h"
+
+#include "AssetRegistry/IAssetRegistry.h"
+#include "AssetRegistry/AssetRegistryModule.h"
+#include "CoreGlobals.h"
+#include "Containers/Array.h"
+#include "InterchangeEditorModule.h"
+#include "InterchangePipelineBase.h"
+#include "InterchangeEditorLog.h"
+#include "Misc/ConfigCacheIni.h"
+#include "Modules/ModuleManager.h"
+
+TSet<uint32> FInterchangePipelineSettingsCacheHandler::CachedPipelineHashes;
+FDelegateHandle FInterchangePipelineSettingsCacheHandler::AssetRemovedHandle;
+
+void FInterchangePipelineSettingsCacheHandler::InitializeCacheHandler()
+{
+	IAssetRegistry* AssetRegistry = &FModuleManager::LoadModuleChecked<FAssetRegistryModule>(AssetRegistryConstants::ModuleName).Get();
+	AssetRemovedHandle = AssetRegistry->OnAssetRemoved().AddStatic(&FInterchangePipelineSettingsCacheHandler::OnAssetRemoved);
+}
+
+void FInterchangePipelineSettingsCacheHandler::OnAssetRemoved(const FAssetData& RemovedAsset)
+{
+	if (UObject* Asset = RemovedAsset.GetAsset(); !Asset || !Asset->IsA<UInterchangePipelineBase>())
+	{
+		return;
+	}
+
+	const FString PipelinePathString = RemovedAsset.GetObjectPathString();
+	const uint32 PipelinePathHash = GetTypeHash(PipelinePathString);
+	const FString PathHashString = FString::Printf(TEXT("%u"), PipelinePathHash);
+	if (GConfig->EmptySectionsMatchingString(*PathHashString, UInterchangePipelineBase::GetDefaultConfigFileName()))
+	{
+		UE_LOG(LogInterchangeEditor, Log, TEXT("Cached pipeline settings are removed for %s"), *RemovedAsset.AssetName.ToString());
+		constexpr bool bRemoveFromCache = false;
+		GConfig->Flush(bRemoveFromCache, UInterchangePipelineBase::GetDefaultConfigFileName());
+	}
+}
+
+void FInterchangePipelineSettingsCacheHandler::ShutdownCacheHandler()
+{
+	if (AssetRemovedHandle.IsValid() && FModuleManager::Get().IsModuleLoaded(AssetRegistryConstants::ModuleName))
+	{
+		IAssetRegistry* AssetRegistry = &FModuleManager::GetModuleChecked<FAssetRegistryModule>(AssetRegistryConstants::ModuleName).Get();
+		AssetRegistry->OnAssetRemoved().Remove(AssetRemovedHandle);
+	}
+
+	CachedPipelineHashes.Empty();
+}
diff --git a/Engine/Plugins/Interchange/Editor/Source/InterchangeEditor/Private/InterchangePipelineSettingsCacheHandler.h b/Engine/Plugins/Interchange/Editor/Source/InterchangeEditor/Private/InterchangePipelineSettingsCacheHandler.h
new file mode 100644
--- /dev/null
+++ b/Engine/Plugins/Interchange/Editor/Source/InterchangeEditor/Private/InterchangePipelineSettingsCacheHandler.h
@@ -0,0 +1,22 @@
+// Copyright Epic Games, Inc. All Rights Reserved.
+#pragma once
+
+#include "AssetRegistry/AssetData.h"
+#include "Containers/Set.h"
+#include "Delegates/IDelegateInstance.h"
+#include "UObject/SoftObjectPath.h"
+
+class FInterchangePipelineSettingsCacheHandler
+{
+public:
+	static void InitializeCacheHandler();
+
+	static void OnAssetRemoved(const FAssetData& RemovedAsset);
+
+	static void ShutdownCacheHandler();
+
+private:
+	static TSet<uint32> CachedPipelineHashes;
+
+	static FDelegateHandle AssetRemovedHandle;
+};
diff --git a/Engine/Source/Runtime/Interchange/Core/Private/InterchangePipelineBase.cpp b/Engine/Source/Runtime/Interchange/Core/Private/InterchangePipelineBase.cpp
--- a/Engine/Source/Runtime/Interchange/Core/Private/InterchangePipelineBase.cpp
+++ b/Engine/Source/Runtime/Interchange/Core/Private/InterchangePipelineBase.cpp
@@ -10,31 +10,84 @@
 
 #include UE_INLINE_GENERATED_CPP_BY_NAME(InterchangePipelineBase)
 
-namespace UE
-{
-	namespace Interchange
+namespace UE::Interchange::PipelinePrivate
+{	
+	FString CreateConfigSectionName(const FName PipelineStackName, UClass* PipelineClass)
 	{
-		namespace PipelinePrivate
+		const FString Section = TEXT("Interchange_StackName__") + PipelineStackName.ToString() + TEXT("__PipelineClassName_") + PipelineClass->GetName();
+		return Section;
+	}
+
+	uint32 GetConfigSectionPipelinePathHash(UInterchangePipelineBase* PipelineAsset)
+	{
+		uint32 PipelineHash = 0;
+		FSoftObjectPath PipelinePath;
+#if WITH_EDITORONLY_DATA
+		if (PipelineAsset->IsStandAlonePipeline())
+		{
+			PipelinePath = PipelineAsset->OriginalPipelinePath;
+		}
+		else if (UInterchangePipelineBase* OuterMostPipeline = PipelineAsset->GetMostPipelineOuter())
 		{
-			FString CreateConfigSectionName(const FName PipelineStackName, UClass* PipelineClass)
+			PipelinePath = OuterMostPipeline->OriginalPipelinePath;
+		}
+#endif
+
+		if (PipelinePath.IsValid())
+		{
+			PipelineHash = GetTypeHash(PipelinePath.ToString());
+		}
+		
+		return PipelineHash;
+	}
+
+	FString CreateConfigSectionName(const FName PipelineStackName, UInterchangePipelineBase* PipelineAsset)
+	{
+		FString PathString;
+		uint32 PipelineHash = GetConfigSectionPipelinePathHash(PipelineAsset);
+		if (PipelineHash != 0)
+		{
+			PathString = FString::Printf(TEXT("__PipelinePathHash_%u"), PipelineHash);
+			if (!PipelineAsset->IsStandAlonePipeline())
 			{
-				FString Section = TEXT("Interchange_StackName__") + PipelineStackName.ToString() + TEXT("__PipelineClassName_") + PipelineClass->GetName();
-				return Section;
+				PathString = FString::Printf(TEXT("__PipelinePathHash_%u__SubPipelineClassName__%s"), PipelineHash, *PipelineAsset->GetClass()->GetName());
 			}
 		}
+
+		if (PathString.IsEmpty())
+		{
+			UE_LOG(LogInterchangeCore, Warning, TEXT("Failed to resolve pipeline's original pipeline path. Using class name based config section name."));
+			// Fallback to class based config section name
+			PathString = TEXT("__PipelineClassName_") + PipelineAsset->GetClass()->GetName();
+		}
+
+		const FString Section = TEXT("Interchange_StackName__") + PipelineStackName.ToString() + PathString;
+		return Section;
 	}
 }
 
+FString UInterchangePipelineBase::GetDefaultConfigFileName()
+{
+	const static FString PipelineConfigFilename = GEditorPerProjectIni;
+	return PipelineConfigFilename;
+}
+
 void UInterchangePipelineBase::LoadSettings(const FName PipelineStackName, bool bResetPreDialog /*= false*/)
 {
-	LoadSettingsInternal(PipelineStackName, GEditorPerProjectIni, PropertiesStates, bResetPreDialog);
+	bool bOutRequiresSaving = false;
+	LoadSettingsInternal(PipelineStackName, GetDefaultConfigFileName(), PropertiesStates, bResetPreDialog, bOutRequiresSaving);
+	if (bOutRequiresSaving)
+	{
+		SaveSettings(PipelineStackName);
+	}
 }
 
 void UInterchangePipelineBase::SaveSettings(const FName PipelineStackName)
 {
-	SaveSettingsInternal(PipelineStackName, GEditorPerProjectIni);
+	SaveSettingsInternal(PipelineStackName, GetDefaultConfigFileName());
 	//Flush the config only once, we do not want to flush when recursively save the sub object pipeline.
-	GConfig->Flush(0);
+	constexpr bool bRemoveFromCache = false;
+	GConfig->Flush(bRemoveFromCache, GetDefaultConfigFileName());
 }
 
 void UInterchangePipelineBase::AdjustSettingsForContext(const FInterchangePipelineContextParams& ContextParams)
@@ -102,7 +155,7 @@ FName UInterchangePipelineBase::GetResultsPropertyName()
 	return GET_MEMBER_NAME_CHECKED(UInterchangePipelineBase, Results);
 }
 
-void UInterchangePipelineBase::LoadSettingsInternal(const FName PipelineStackName, const FString& ConfigFilename, TMap<FName, FInterchangePipelinePropertyStates>& ParentPropertiesStates, bool bResetPreDialog)
+void UInterchangePipelineBase::LoadSettingsInternal(const FName PipelineStackName, const FString& ConfigFilename, TMap<FName, FInterchangePipelinePropertyStates>& ParentPropertiesStates, bool bResetPreDialog, bool& bOutRequiresSaving)
 {
 	int32 PortFlags = 0;
 	UClass* Class = this->GetClass();
@@ -115,8 +168,17 @@ void UInterchangePipelineBase::LoadSettingsInternal(const FName PipelineStackNam
 			continue;
 		}
 
-		FString SectionName = UE::Interchange::PipelinePrivate::CreateConfigSectionName(PipelineStackName, Class);
-		FString Key = Property->GetName();
+		FString SectionName = UE::Interchange::PipelinePrivate::CreateConfigSectionName(PipelineStackName, this);
+		{
+			constexpr bool bForce = false;
+			const FConfigSection* Section = GConfig->GetSection(*SectionName, bForce, ConfigFilename);
+			if (!Section)
+			{
+				SectionName = UE::Interchange::PipelinePrivate::CreateConfigSectionName(PipelineStackName, Class);
+				bOutRequiresSaving |= true;
+			}
+		}
+
 		const FName PropertyName = Property->GetFName();
 		const FName PropertyPath = FName(Property->GetPathName());
 		if (PropertyName == GET_MEMBER_NAME_CHECKED(UInterchangePipelineBase, PropertiesStates))
@@ -148,6 +210,7 @@ void UInterchangePipelineBase::LoadSettingsInternal(const FName PipelineStackNam
 			}
 		}
 
+		FString Key = Property->GetName();
 		const bool bIsPropertyInherited = Property->GetOwnerClass() != Class;
 		UObject* SuperClassDefaultObject = Class->GetSuperClass()->GetDefaultObject();
 
@@ -155,7 +218,7 @@ void UInterchangePipelineBase::LoadSettingsInternal(const FName PipelineStackNam
 		FArrayProperty* Array = CastField<FArrayProperty>(Property);
 		if (Array)
 		{
-			const bool bForce = false;
+			constexpr bool bForce = false;
 			const FConfigSection* Section = GConfig->GetSection(*SectionName, bForce, ConfigFilename);
 			if (Section != nullptr)
 			{
@@ -208,7 +271,7 @@ void UInterchangePipelineBase::LoadSettingsInternal(const FName PipelineStackNam
 			// Load the settings if the referenced pipeline is a subobject of ours
 			if (SubPipeline->IsInOuter(this))
 			{
-				SubPipeline->LoadSettingsInternal(PipelineStackName, ConfigFilename, ParentPropertiesStates, bResetPreDialog);
+				SubPipeline->LoadSettingsInternal(PipelineStackName, ConfigFilename, ParentPropertiesStates, bResetPreDialog, bOutRequiresSaving);
 			}
 		}
 		else
@@ -249,14 +312,15 @@ void UInterchangePipelineBase::SaveSettingsInternal(const FName PipelineStackNam
 			continue;
 		}
 
-		FString SectionName = UE::Interchange::PipelinePrivate::CreateConfigSectionName(PipelineStackName, Class);
-		FString Key = Property->GetName();
 		const FName PropertyName = Property->GetFName();
 		if (PropertyName == GET_MEMBER_NAME_CHECKED(UInterchangePipelineBase, PropertiesStates))
 		{
 			continue;
 		}
 
+		const FString SectionName = UE::Interchange::PipelinePrivate::CreateConfigSectionName(PipelineStackName, this);
+		FString Key = Property->GetName();
+
 		const bool bIsPropertyInherited = Property->GetOwnerClass() != Class;
 		UObject* SuperClassDefaultObject = Class->GetSuperClass()->GetDefaultObject();
 		FObjectProperty* SubObject = CastField<FObjectProperty>(Property);
@@ -300,7 +364,6 @@ void UInterchangePipelineBase::SaveSettingsInternal(const FName PipelineStackNam
 	}
 }
 
-
 UInterchangePipelineBase* UInterchangePipelineBase::GetMostPipelineOuter() const
 {
 	UInterchangePipelineBase* Top = (UInterchangePipelineBase*)this;
@@ -483,4 +546,4 @@ void UInterchangePipelineBase::PostDuplicate(bool bDuplicateForPIE)
 	Super::PostDuplicate(bDuplicateForPIE);
 
 	UpdateWeakObjectPtrs();
-}
+}
\ No newline at end of file
diff --git a/Engine/Source/Runtime/Interchange/Core/Public/InterchangePipelineBase.h b/Engine/Source/Runtime/Interchange/Core/Public/InterchangePipelineBase.h
--- a/Engine/Source/Runtime/Interchange/Core/Public/InterchangePipelineBase.h
+++ b/Engine/Source/Runtime/Interchange/Core/Public/InterchangePipelineBase.h
@@ -341,6 +341,8 @@ public:
 		return true;
 	}
 
+	INTERCHANGECORE_API static FString GetDefaultConfigFileName();
+	
 	INTERCHANGECORE_API void LoadSettings(const FName PipelineStackName, bool bResetPreDialog = false);
 
 	INTERCHANGECORE_API void SaveSettings(const FName PipelineStackName);
@@ -513,7 +515,7 @@ public:
 	 * The pipeline itself can call the function to know whether a property must be used.
 	 * Example: the override asset name property.
 	 */
-	bool IsStandAlonePipeline()
+	bool IsStandAlonePipeline() const
 	{
 		return this == GetMostPipelineOuter();
 	};
@@ -549,6 +551,17 @@ public:
 	UPROPERTY()
 	FString ContentImportPath;
 
+#if WITH_EDITORONLY_DATA
+	/**
+	 * Path of the pipeline object used to create an instance of the pipeline.
+	 * Useful for creating section names for caching the properties in ini file.
+	*/
+	UPROPERTY()
+	FSoftObjectPath OriginalPipelinePath;
+#endif
+
+	INTERCHANGECORE_API UInterchangePipelineBase* GetMostPipelineOuter() const;
+
 protected:
 
 	INTERCHANGECORE_API virtual FString GetPipelineDisplayName() const;
@@ -591,12 +604,10 @@ protected:
 	{
 	}
 
-	INTERCHANGECORE_API void LoadSettingsInternal(const FName PipelineStackName, const FString& ConfigFilename, TMap<FName, FInterchangePipelinePropertyStates>& ParentPropertiesStates, bool bResetPreDialog);
+	INTERCHANGECORE_API void LoadSettingsInternal(const FName PipelineStackName, const FString& ConfigFilename, TMap<FName, FInterchangePipelinePropertyStates>& ParentPropertiesStates, bool bResetPreDialog, bool& bOutRequiresSaving);
 
 	INTERCHANGECORE_API void SaveSettingsInternal(const FName PipelineStackName, const FString& ConfigFilename);
 
-	INTERCHANGECORE_API UInterchangePipelineBase* GetMostPipelineOuter() const;
-
 #if WITH_EDITOR
 	static INTERCHANGECORE_API void InternalToggleVisibilityPropertiesOfMetaDataValue(UInterchangePipelineBase* OuterMostPipeline, UInterchangePipelineBase* Pipeline, bool bDoTransientSubPipeline, const FString& MetaDataKey, const FString& MetaDataValue, const bool bVisibilityState);
 	static INTERCHANGECORE_API void HidePropertiesOfCategory(UInterchangePipelineBase* OuterMostPipeline, UInterchangePipelineBase* Pipeline, const FString& HideCategoryName, bool bDoTransientSubPipeline = false);
@@ -649,5 +660,5 @@ protected:
 	mutable TMap<FName, FInterchangePipelinePropertyStates> CachePropertiesStates;
 	mutable FInterchangePipelineContextParams CacheContextParam;
 
-	TArray<FInterchangeConflictInfo> ConflictInfos;
+	TArray<FInterchangeConflictInfo> ConflictInfos; 
 };
diff --git a/Engine/Source/Runtime/Interchange/Engine/Private/InterchangeManager.cpp b/Engine/Source/Runtime/Interchange/Engine/Private/InterchangeManager.cpp
--- a/Engine/Source/Runtime/Interchange/Engine/Private/InterchangeManager.cpp
+++ b/Engine/Source/Runtime/Interchange/Engine/Private/InterchangeManager.cpp
@@ -970,6 +970,9 @@ UInterchangePipelineBase* UE::Interchange::GeneratePipelineInstance(const FSoftO
 		if (BlueprintPipeline->GeneratedClass.Get())
 		{
 			GeneratedPipeline = NewObject<UInterchangePipelineBase>(PipelineInstancePackage, BlueprintPipeline->GeneratedClass);
+#if WITH_EDITORONLY_DATA
+			GeneratedPipeline->OriginalPipelinePath = FSoftObjectPath(BlueprintPipeline);
+#endif
 		}
 		else
 		{
@@ -982,6 +985,9 @@ UInterchangePipelineBase* UE::Interchange::GeneratePipelineInstance(const FSoftO
 		if (PythonPipeline->GeneratedPipeline)
 		{
 			GeneratedPipeline = DuplicateObject<UInterchangePipelineBase>(PythonPipeline->GeneratedPipeline.Get(), PipelineInstancePackage);
+#if WITH_EDITORONLY_DATA
+			GeneratedPipeline->OriginalPipelinePath = FSoftObjectPath(PythonPipeline);
+#endif
 		}
 		else
 		{
@@ -992,6 +998,16 @@ UInterchangePipelineBase* UE::Interchange::GeneratePipelineInstance(const FSoftO
 	else if (const UInterchangePipelineBase* DefaultPipeline = Cast<UInterchangePipelineBase>(ReferenceInstance))
 	{
 		GeneratedPipeline = DuplicateObject<UInterchangePipelineBase>(DefaultPipeline, PipelineInstancePackage);
+#if WITH_EDITORONLY_DATA
+		if (DefaultPipeline->OriginalPipelinePath.IsValid())
+		{
+			GeneratedPipeline->OriginalPipelinePath = DefaultPipeline->OriginalPipelinePath;
+		}
+		else
+		{
+			GeneratedPipeline->OriginalPipelinePath = FSoftObjectPath(DefaultPipeline);
+		}
+#endif
 	}
 	else
 	{
