[Interchange] Expose interchange project settings to python and blueprint.

#rnx

diff --git a/Engine/Source/Runtime/Interchange/Engine/Private/InterchangeProjectSettings.cpp b/Engine/Source/Runtime/Interchange/Engine/Private/InterchangeProjectSettings.cpp
--- a/Engine/Source/Runtime/Interchange/Engine/Private/InterchangeProjectSettings.cpp
+++ b/Engine/Source/Runtime/Interchange/Engine/Private/InterchangeProjectSettings.cpp
@@ -4,6 +4,23 @@
 #include "InterchangeManager.h"
 #include "InterchangeSourceData.h"
 
+namespace UE::Interchange::Private
+{
+	TArray<UInterchangePipelineBase*> GetPipelineArrayFromSoftObjectArray(const TArray<FSoftObjectPath>& PipelineSoftObjectPaths)
+	{
+		TArray<UInterchangePipelineBase*> Pipelines;
+		Pipelines.Reserve(PipelineSoftObjectPaths.Num());
+		for (const FSoftObjectPath& PipelinePath : PipelineSoftObjectPaths)
+		{
+			if (UInterchangePipelineBase* PipelineBase = GeneratePipelineInstance(PipelinePath))
+			{
+				Pipelines.Add(PipelineBase);
+			}
+		}
+		return Pipelines;
+	}
+}
+
 #if WITH_EDITOR
 void UInterchangeProjectSettings::PostEditChangeProperty(FPropertyChangedEvent& PropertyChangedEvent)
 {
@@ -281,6 +298,53 @@ TArray<FName> FInterchangeProjectSettingsUtils::GetGroupNames()
 	return GroupNames;
 }
 
+TArray<UInterchangePipelineBase*> UInterchangeProjectSettingsScript::GetPipelineArrayFromPipelineStack(const FInterchangePipelineStack& InterchangePipelineStack)
+{
+	return UE::Interchange::Private::GetPipelineArrayFromSoftObjectArray(InterchangePipelineStack.Pipelines);
+}
+
+TArray<UInterchangePipelineBase*> UInterchangeProjectSettingsScript::GetPipelineArrayFromTranslatorPipelines(const FInterchangeTranslatorPipelines& InterchangeTranslatorPipeline)
+{
+	return UE::Interchange::Private::GetPipelineArrayFromSoftObjectArray(InterchangeTranslatorPipeline.Pipelines);
+}
+
+TArray<UInterchangePipelineBase*> UInterchangeProjectSettingsScript::GetPipelineStackFromSourceData(const bool bIsSceneImport, const UInterchangeSourceData* SourceData)
+{
+	TArray<UInterchangePipelineBase*> ResultPipelines;
+	FName DefaultStackName = FInterchangeProjectSettingsUtils::GetDefaultPipelineStackName(bIsSceneImport, *SourceData);
+	const FInterchangeImportSettings& ImportSettings = FInterchangeProjectSettingsUtils::GetDefaultImportSettings(bIsSceneImport);
+	if (!ImportSettings.PipelineStacks.Contains(DefaultStackName))
+	{
+		return ResultPipelines;
+	}
+
+	const FInterchangePipelineStack* InterchangePipelineStack = ImportSettings.PipelineStacks.Find(DefaultStackName);
+	if (!InterchangePipelineStack)
+	{
+		return ResultPipelines;
+	}
+
+	UE::Interchange::FScopedTranslator ScopedTranslator(SourceData);
+	const TArray<FSoftObjectPath>* Pipelines = &InterchangePipelineStack->Pipelines;
+	// If applicable, check to see if a specific pipeline stack is associated with this translator
+	for (const FInterchangeTranslatorPipelines& TranslatorPipelines : InterchangePipelineStack->PerTranslatorPipelines)
+	{
+		const UClass* TranslatorClass = TranslatorPipelines.Translator.LoadSynchronous();
+		if (ScopedTranslator.GetTranslator() && ScopedTranslator.GetTranslator()->IsA(TranslatorClass))
+		{
+			Pipelines = &TranslatorPipelines.Pipelines;
+			break;
+		}
+	}
+
+	if(Pipelines)
+	{
+		//Copy the pipelines
+		ResultPipelines = UE::Interchange::Private::GetPipelineArrayFromSoftObjectArray(*Pipelines);
+	}
+	return ResultPipelines;
+}
+
 void UInterchangeEditorSettings::SetUsedGroupName(const FName& InUsedGroupName)
 {
 	UsedGroupName = InUsedGroupName;
diff --git a/Engine/Source/Runtime/Interchange/Engine/Public/InterchangeProjectSettings.h b/Engine/Source/Runtime/Interchange/Engine/Public/InterchangeProjectSettings.h
--- a/Engine/Source/Runtime/Interchange/Engine/Public/InterchangeProjectSettings.h
+++ b/Engine/Source/Runtime/Interchange/Engine/Public/InterchangeProjectSettings.h
@@ -13,19 +13,19 @@
 
 class UInterchangeSourceData;
 
-USTRUCT()
+USTRUCT(BlueprintType)
 struct FInterchangeTranslatorPipelines
 {
 	GENERATED_BODY()
 
-	UPROPERTY(EditAnywhere, Category = "TranslatorPipelines", meta = (AllowedClasses = "/Script/InterchangeCore.InterchangeTranslatorBase"))
+	UPROPERTY(EditAnywhere, BlueprintReadWrite, Category = "TranslatorPipelines", meta = (AllowedClasses = "/Script/InterchangeCore.InterchangeTranslatorBase"))
 	TSoftClassPtr<UInterchangeTranslatorBase> Translator;
 
-	UPROPERTY(EditAnywhere, Category = "TranslatorPipelines", meta = (AllowedClasses = "/Script/InterchangeCore.InterchangePipelineBase, /Script/InterchangeEngine.InterchangeBlueprintPipelineBase, /Script/InterchangeEngine.InterchangePythonPipelineAsset"))
+	UPROPERTY(EditAnywhere, BlueprintReadWrite, Category = "TranslatorPipelines", meta = (AllowedClasses = "/Script/InterchangeCore.InterchangePipelineBase, /Script/InterchangeEngine.InterchangeBlueprintPipelineBase, /Script/InterchangeEngine.InterchangePythonPipelineAsset"))
 	TArray<FSoftObjectPath> Pipelines;
 };
 
-USTRUCT()
+USTRUCT(BlueprintType)
 struct FInterchangePipelineStack
 {
 	GENERATED_BODY()
@@ -35,93 +35,93 @@ struct FInterchangePipelineStack
 	TArray<FSoftObjectPath> Pipelines;
 
 	/** Specifies a different list of pipelines for this stack to use when importing data from specific translators. */
-	UPROPERTY(EditAnywhere, Category = "TranslatorPipelines")
+	UPROPERTY(EditAnywhere, BlueprintReadWrite, Category = "TranslatorPipelines")
 	TArray<FInterchangeTranslatorPipelines> PerTranslatorPipelines;
 };
 
-USTRUCT()
+USTRUCT(BlueprintType)
 struct FInterchangePerTranslatorDialogOverride
 {
 	GENERATED_BODY()
 
-	UPROPERTY(EditAnywhere, Category = "DialogOverride", meta = (AllowedClasses = "/Script/InterchangeCore.InterchangeTranslatorBase"))
+	UPROPERTY(EditAnywhere, BlueprintReadWrite, Category = "DialogOverride", meta = (AllowedClasses = "/Script/InterchangeCore.InterchangeTranslatorBase"))
 	TSoftClassPtr<UInterchangeTranslatorBase> Translator;
 
 	/** Show the options dialog when Interchange imports. */
-	UPROPERTY(EditAnywhere, Category = "DialogOverride")
+	UPROPERTY(EditAnywhere, BlueprintReadWrite, Category = "DialogOverride")
 	bool bShowImportDialog = true;
 
 	/** Show the options dialog when Interchange reimports. */
-	UPROPERTY(EditAnywhere, Category = "DialogOverride")
+	UPROPERTY(EditAnywhere, BlueprintReadWrite, Category = "DialogOverride")
 	bool bShowReimportDialog = false;
 };
 
-USTRUCT()
+USTRUCT(BlueprintType)
 struct FInterchangeDialogOverride
 {
 	GENERATED_BODY()
 
 	/** Show the options dialog when Interchange imports. */
-	UPROPERTY(EditAnywhere, Category = "DialogOverride")
+	UPROPERTY(EditAnywhere, BlueprintReadWrite, Category = "DialogOverride")
 	bool bShowImportDialog = true;
 
 	/** Show the options dialog when Interchange reimports. */
-	UPROPERTY(EditAnywhere, Category = "DialogOverride")
+	UPROPERTY(EditAnywhere, BlueprintReadWrite, Category = "DialogOverride")
 	bool bShowReimportDialog = false;
 
-	UPROPERTY(EditAnywhere, Category = "DialogOverride")
+	UPROPERTY(EditAnywhere, BlueprintReadWrite, Category = "DialogOverride")
 	TArray<FInterchangePerTranslatorDialogOverride> PerTranslatorImportDialogOverride;
 };
 
-USTRUCT()
+USTRUCT(BlueprintType)
 struct FInterchangeImportSettings
 {
 	GENERATED_BODY()
 
 	/** Configures the pipeline stacks that are available when importing assets with Interchange. */
-	UPROPERTY(EditAnywhere, Category = "Pipeline")
+	UPROPERTY(EditAnywhere, BlueprintReadWrite, Category = "Pipeline")
 	TMap<FName, FInterchangePipelineStack> PipelineStacks;
 
 	/** Specifies which pipeline stack Interchange should use by default. */
-	UPROPERTY(EditAnywhere, Category = "Pipeline")
+	UPROPERTY(EditAnywhere, BlueprintReadWrite, Category = "Pipeline")
 	FName DefaultPipelineStack = NAME_None;
 
 	/** Specifies the class that should be used to define the configuration dialog that Interchange shows on import. */
-	UPROPERTY(EditAnywhere, Category = "Pipeline")
+	UPROPERTY(EditAnywhere, BlueprintReadWrite, Category = "Pipeline")
 	TSoftClassPtr <UInterchangePipelineConfigurationBase> ImportDialogClass;
 
 	/** Show the options dialog when Interchange imports. */
-	UPROPERTY(EditAnywhere, Category = "DialogOverride")
+	UPROPERTY(EditAnywhere, BlueprintReadWrite, Category = "DialogOverride")
 	bool bShowImportDialog = true;
 
 	/** Show the options dialog when Interchange reimports. */
-	UPROPERTY(EditAnywhere, Category = "DialogOverride")
+	UPROPERTY(EditAnywhere, BlueprintReadWrite, Category = "DialogOverride")
 	bool bShowReimportDialog = false;
 };
 
-USTRUCT()
+USTRUCT(BlueprintType)
 struct FInterchangeSceneImportSettings : public FInterchangeImportSettings
 {
 	GENERATED_BODY()
 
-	UPROPERTY(EditAnywhere, Category = "DialogOverride", Meta=(DisplayAfter="bShowReimportDialog"))
+	UPROPERTY(EditAnywhere, BlueprintReadWrite, Category = "DialogOverride", Meta=(DisplayAfter="bShowReimportDialog"))
 	TArray<FInterchangePerTranslatorDialogOverride> PerTranslatorDialogOverride;
 };
 
-USTRUCT()
+USTRUCT(BlueprintType)
 struct FInterchangeContentImportSettings : public FInterchangeImportSettings
 {
 	GENERATED_BODY()
 
 	/** Specifies a different pipeline stack for Interchange to use by default when importing specific types of assets. */
-	UPROPERTY(EditAnywhere, Category = "Pipeline", Meta=(DisplayAfter="DefaultPipelineStack"))
+	UPROPERTY(EditAnywhere, BlueprintReadWrite, Category = "Pipeline", Meta=(DisplayAfter="DefaultPipelineStack"))
 	TMap<EInterchangeTranslatorAssetType, FName> DefaultPipelineStackOverride;
 
-	UPROPERTY(EditAnywhere, Category = "DialogOverride", Meta=(DisplayAfter="bShowReimportDialog"))
+	UPROPERTY(EditAnywhere, BlueprintReadWrite, Category = "DialogOverride", Meta=(DisplayAfter="bShowReimportDialog"))
 	TMap<EInterchangeTranslatorAssetType, FInterchangeDialogOverride> ShowImportDialogOverride;
 };
 
-USTRUCT()
+USTRUCT(BlueprintType)
 struct FInterchangeGroup
 {
 	GENERATED_BODY()
@@ -134,35 +134,35 @@ struct FInterchangeGroup
 	};
 
 	/** Specifies a different pipeline stack for Interchange to use by default when importing specific types of assets. */
-	UPROPERTY(EditAnywhere, Category = "InterchangeGroup")
+	UPROPERTY(EditAnywhere, BlueprintReadWrite, Category = "InterchangeGroup")
 	FName DisplayName;
 
 	/** This tell interchange if the import dialog should show or not when importing a particular type of asset.*/
-	UPROPERTY(VisibleAnywhere, Category = "InterchangeGroup", meta = (IgnoreForMemberInitializationTest, EditCondition = "false", EditConditionHides))
+	UPROPERTY(VisibleAnywhere, BlueprintReadWrite, Category = "InterchangeGroup", meta = (IgnoreForMemberInitializationTest, EditCondition = "false", EditConditionHides))
 	FGuid UniqueID = FGuid::NewGuid();
 
 	/** Specifies which pipeline stack Interchange should use by default. */
-	UPROPERTY(EditAnywhere, Category = "Pipeline")
+	UPROPERTY(EditAnywhere, BlueprintReadWrite, Category = "Pipeline")
 	FName DefaultPipelineStack = NAME_None;
 
 	/** Specifies a different pipeline stack for Interchange to use by default when importing specific types of assets. */
-	UPROPERTY(EditAnywhere, Category = "Pipeline", Meta = (DisplayAfter = "DefaultPipelineStack"))
+	UPROPERTY(EditAnywhere, BlueprintReadWrite, Category = "Pipeline", Meta = (DisplayAfter = "DefaultPipelineStack"))
 	TMap<EInterchangeTranslatorAssetType, FName> DefaultPipelineStackOverride;
 
 	/** Show the options dialog when Interchange imports. */
-	UPROPERTY(EditAnywhere, Category = "DialogOverride")
+	UPROPERTY(EditAnywhere, BlueprintReadWrite, Category = "DialogOverride")
 	bool bShowImportDialog = true;
 
 	/** Show the options dialog when Interchange reimports. */
-	UPROPERTY(EditAnywhere, Category = "DialogOverride")
+	UPROPERTY(EditAnywhere, BlueprintReadWrite, Category = "DialogOverride")
 	bool bShowReimportDialog = false;
 
 	/** This tell interchange if the import dialog should show or not when importing a particular type of asset.*/
-	UPROPERTY(EditAnywhere, Category = "DialogOverride")
+	UPROPERTY(EditAnywhere, BlueprintReadWrite, Category = "DialogOverride")
 	TMap<EInterchangeTranslatorAssetType, FInterchangeDialogOverride> ShowImportDialogOverride;
 };
 
-UCLASS(config=Engine, meta=(DisplayName=Interchange), MinimalAPI)
+UCLASS(config=Engine, BlueprintType, meta=(DisplayName=Interchange), MinimalAPI)
 class UInterchangeProjectSettings : public UDeveloperSettings
 {
 	GENERATED_BODY()
@@ -171,44 +171,44 @@ public:
 	/**
 	 * Settings used when importing into the Content Browser.
 	 */
-	UPROPERTY(EditAnywhere, config, Category = "ImportContent")
+	UPROPERTY(EditAnywhere, BlueprintReadWrite, config, Category = "ImportContent")
 	FInterchangeContentImportSettings ContentImportSettings;
 
 	/**
 	 * Settings used when importing into a level.
 	 */
-	UPROPERTY(EditAnywhere, config, Category = "ImportIntoLevel")
+	UPROPERTY(EditAnywhere, BlueprintReadWrite, config, Category = "ImportIntoLevel")
 	FInterchangeSceneImportSettings SceneImportSettings;
 
 	/** This tells Interchange which file picker class to construct when we need to choose a file for a source. */
-	UPROPERTY(EditAnywhere, config, Category = "EditorInterface")
+	UPROPERTY(EditAnywhere, BlueprintReadWrite, config, Category = "EditorInterface")
 	TSoftClassPtr <UInterchangeFilePickerBase> FilePickerClass;
 
 	/**
 	 * If enabled, both Interchange translators and the legacy import process smooth the edges of static meshes that don't contain smoothing information.
 	 * If you have an older project that relies on leaving hard edges by default, you can disable this setting to preserve consistency with older assets.
 	 */
-	UPROPERTY(EditAnywhere, config, Category = "Generic|ImportSettings")
+	UPROPERTY(EditAnywhere, BlueprintReadWrite, config, Category = "Generic|ImportSettings")
 	bool bStaticMeshUseSmoothEdgesIfSmoothingInformationIsMissing = true;
 
 	/**
 	 * Specifies which pipeline class Interchange should use when editor tools import or reimport an asset with base settings.
 	 * Unreal Editor depends on this class to be set. You can only edit this property in the .ini file.
 	 */
-	UPROPERTY(EditAnywhere, config, Category = "Editor Generic Pipeline Class")
+	UPROPERTY(EditAnywhere, BlueprintReadWrite, config, Category = "Editor Generic Pipeline Class")
 	TSoftClassPtr <UInterchangePipelineBase> GenericPipelineClass;
 
 	/**
 	 * Optional, the pipeline asset converters will duplicate to create interchange import data pipeline.
 	 * If not set, converters will duplicate a pipeline class CDO of there choice.
 	 */
-	UPROPERTY(EditAnywhere, config, Category = "Converters", meta = (AllowedClasses = "/Script/InterchangeCore.InterchangePipelineBase, /Script/InterchangeEngine.InterchangeBlueprintPipelineBase, /Script/InterchangeEngine.InterchangePythonPipelineAsset"))
+	UPROPERTY(EditAnywhere, BlueprintReadWrite, config, Category = "Converters", meta = (AllowedClasses = "/Script/InterchangeCore.InterchangePipelineBase, /Script/InterchangeEngine.InterchangeBlueprintPipelineBase, /Script/InterchangeEngine.InterchangePythonPipelineAsset"))
 	FSoftObjectPath ConverterDefaultPipeline;
 
 	/**
 	* Groups that define PerTransalatorPipelines that user can select to use.
 	*/
-	UPROPERTY(EditAnywhere, config, Category = "Groups", meta = (NoElementDuplicate))
+	UPROPERTY(EditAnywhere, BlueprintReadWrite, config, Category = "Groups", meta = (NoElementDuplicate))
 	TArray<FInterchangeGroup> InterchangeGroups;
 
 private:
@@ -235,14 +235,44 @@ public:
 	static INTERCHANGEENGINE_API TArray<FName> GetGroupNames();
 };
 
-UCLASS(config = EditorPerProjectUserSettings, meta = (DisplayName = Interchange), MinimalAPI)
+UCLASS()
+class INTERCHANGEENGINE_API UInterchangeProjectSettingsScript : public UObject
+{
+	GENERATED_BODY()
+public:
+	/**
+	 * Pipeline stack store SoftObjectPath which are not supported in python, you can use this function to retrieve the object from FInterchangePipelineStack.
+	 * 
+	 * @note: It return duplicated pipeline objects
+	 */
+	UFUNCTION(BlueprintCallable, Category = "Interchange Settings")
+	static TArray<UInterchangePipelineBase*> GetPipelineArrayFromPipelineStack(const FInterchangePipelineStack& InterchangePipelineStack);
+
+	/**
+	 * Pipeline stack store SoftObjectPath which are not supported in python, you can use this function to retrieve the object from FInterchangeTranslatorPipelines.
+	 * 
+	 * @note: It return duplicated pipeline objects
+	 */
+	UFUNCTION(BlueprintCallable, Category = "Interchange Settings")
+	static TArray<UInterchangePipelineBase*> GetPipelineArrayFromTranslatorPipelines(const FInterchangeTranslatorPipelines& InterchangeTranslatorPipeline);
+
+	/**
+	 * Return the default settings pipeline stack for the specified source data.
+	 * 
+	 * @note: It return duplicated pipeline objects
+	 */
+	UFUNCTION(BlueprintCallable, Category = "Interchange Settings")
+	static TArray<UInterchangePipelineBase*> GetPipelineStackFromSourceData(const bool bIsSceneImport, const UInterchangeSourceData* SourceData);
+};
+
+UCLASS(config = EditorPerProjectUserSettings, BlueprintType, meta = (DisplayName = Interchange), MinimalAPI)
 class UInterchangeEditorSettings : public UDeveloperSettings
 {
 	GENERATED_BODY()
 
 public:
 	UE_DEPRECATED(5.6, "Use the dialog overrides on the Interchange category of the Project Settings if you want to control when the Interchange reimport dialog is shown")
-	UPROPERTY(EditAnywhere, config, Category = "Show Dialog", meta=(EditCondition = "false", ToolTip="Use the dialog overrides on the Interchange category of the Project Settings if you want to control when the Interchange reimport dialog is shown"))
+	UPROPERTY(EditAnywhere, BlueprintReadWrite, config, Category = "Show Dialog", meta=(EditCondition = "false", ToolTip="Use the dialog overrides on the Interchange category of the Project Settings if you want to control when the Interchange reimport dialog is shown"))
 	bool bShowImportDialogAtReimport = false;
 
 	const FGuid& GetUsedGroupUID() const { return UsedGroupUID; }
