Interchange: Fix some properties having been made Transient by accident with cl 42291753

diff --git a/Engine/Source/Runtime/Interchange/Core/Public/InterchangePipelineBase.h b/Engine/Source/Runtime/Interchange/Core/Public/InterchangePipelineBase.h
--- a/Engine/Source/Runtime/Interchange/Core/Public/InterchangePipelineBase.h
+++ b/Engine/Source/Runtime/Interchange/Core/Public/InterchangePipelineBase.h
@@ -660,17 +660,17 @@ protected:
 	 * If true, this pipeline was create to re-import an asset or override the project settings pipelines.
 	 * That kind of pipeline will not be treat like project settings pipeline in the UI. PredialogCleanup will not be called.
 	 */
-	UPROPERTY(Transient)
+	UPROPERTY()
 	bool bFromReimportOrOverride = false;
 
-	UPROPERTY(Transient)
+	UPROPERTY()
 	TObjectPtr<UInterchangeResultsContainer> Results;
 
 	/**
 	 * Map of property path and lock status. Any properties that have a true lock status will be readonly when showing the import dialog.
 	 * Use the API to Get and Set the properties states.
 	*/
-	UPROPERTY(Transient)
+	UPROPERTY()
 	TMap<FName, FInterchangePipelinePropertyStates> PropertiesStates;
 
 	UPROPERTY(Transient)
