Change the label of the Reimport Strategy in Interchange Import Dialog

#rnx

diff --git a/Engine/Source/Runtime/Interchange/Core/Public/Nodes/InterchangeFactoryBaseNode.h b/Engine/Source/Runtime/Interchange/Core/Public/Nodes/InterchangeFactoryBaseNode.h
--- a/Engine/Source/Runtime/Interchange/Core/Public/Nodes/InterchangeFactoryBaseNode.h
+++ b/Engine/Source/Runtime/Interchange/Core/Public/Nodes/InterchangeFactoryBaseNode.h
@@ -175,7 +175,7 @@ enum class EReimportStrategyFlags : uint8
 	//Always apply all properties specified by the pipeline.
 	ApplyPipelineProperties,
 	//Always apply all properties specified by the pipeline, but leave the properties that were modified in the editor since the last import.
-	ApplyEditorChangedProperties
+	ApplyEditorChangedProperties UMETA(DisplayName = "Preserve Editor Changed Properties")
 };
 
 /**
