Fixed ensure when re-importing an asset that was import with a blueprint or python pipeline
The support for Python pipeline was missing in the first submission.

#rnx

diff --git a/Engine/Source/Runtime/Interchange/Engine/Private/InterchangeAssetImportData.cpp b/Engine/Source/Runtime/Interchange/Engine/Private/InterchangeAssetImportData.cpp
--- a/Engine/Source/Runtime/Interchange/Engine/Private/InterchangeAssetImportData.cpp
+++ b/Engine/Source/Runtime/Interchange/Engine/Private/InterchangeAssetImportData.cpp
@@ -6,6 +6,7 @@
 #include "InterchangeEngineLogPrivate.h"
 #include "InterchangeManager.h"
 #include "InterchangePipelineBase.h"
+#include "InterchangePythonPipelineBase.h"
 #include "InterchangeProjectSettings.h"
 #include "JsonObjectConverter.h"
 #include "Serialization/LargeMemoryReader.h"
@@ -104,7 +105,8 @@ FString SerializeTranslatorSettings(UObject* TranslatorSettings)
 
 UObject* DeSerializePipeline(const FString& PipelineStr, UClass* PipelineClass)
 {
-	UInterchangePipelineBase* GeneratedPipeline = NewObject<UInterchangePipelineBase>(GetTransientPackage(), PipelineClass);
+	// Note: PipelineClass can be a child of either UInterchangePipelineBase or UInterchangePythonPipelineAsset
+	UObject* GeneratedPipeline = NewObject<UObject>(GetTransientPackage(), PipelineClass);
 
 	TSharedPtr<FJsonObject> RootObject;
 	TSharedRef< TJsonReader<> > JsonReader = TJsonReaderFactory<>::Create(PipelineStr);
@@ -114,7 +116,14 @@ UObject* DeSerializePipeline(const FString& PipelineStr, UClass* PipelineClass)
 		FJsonObjectConverter::JsonObjectToUStruct(JsonPipelineProperties.ToSharedRef(), PipelineClass, GeneratedPipeline, 0, 0);
 	}
 
-	GeneratedPipeline->UpdateWeakObjectPtrs();
+	if (UInterchangePipelineBase* InterchangePipelineBase = Cast<UInterchangePipelineBase>(GeneratedPipeline))
+	{
+		InterchangePipelineBase->UpdateWeakObjectPtrs();
+	}
+	else if (UInterchangePythonPipelineAsset* PythonPipelineAsset = Cast<UInterchangePythonPipelineAsset>(GeneratedPipeline))
+	{
+		PythonPipelineAsset->GeneratePipeline();
+	}
 
 	return GeneratedPipeline;
 }
diff --git a/Engine/Source/Runtime/Interchange/Engine/Private/InterchangeManager.cpp b/Engine/Source/Runtime/Interchange/Engine/Private/InterchangeManager.cpp
--- a/Engine/Source/Runtime/Interchange/Engine/Private/InterchangeManager.cpp
+++ b/Engine/Source/Runtime/Interchange/Engine/Private/InterchangeManager.cpp
@@ -2988,34 +2988,44 @@ void UInterchangeManager::FindPipelineCandidate(TArray<UClass*>& PipelineCandida
 		AssetRegistry.GetDerivedClassNames(BaseNames, Excluded, DerivedNames);
 	}
 
+	// Iterate over retrieved blueprint and python pipeline assets
 	FARFilter Filter;
 	Filter.ClassPaths.Add(UBlueprint::StaticClass()->GetClassPathName());
+	Filter.ClassPaths.Add(UInterchangePythonPipelineAsset::StaticClass()->GetClassPathName());
 	Filter.bRecursiveClasses = true;
 	Filter.bRecursivePaths = true;
 
 	TArray< FAssetData > AssetList;
 	AssetRegistry.GetAssets(Filter, AssetList);
 
-	// Iterate over retrieved blueprint assets
 	for (const FAssetData& Asset : AssetList)
 	{
-		//Only get the asset with the native parent class using UInterchangePipelineBase
-		FAssetDataTagMapSharedView::FFindTagResult GeneratedClassPath = Asset.TagsAndValues.FindTag(TEXT("GeneratedClass"));
-		if (GeneratedClassPath.IsSet())
+		if (UBlueprint* Blueprint = Cast<UBlueprint>(Asset.GetAsset()))
 		{
-			// Convert path to just the name part
-			const FTopLevelAssetPath ClassObjectPath(FPackageName::ExportTextPathToObjectPath(*GeneratedClassPath.GetValue()));
-
-			// Check if this class is in the derived set
-			if (!DerivedNames.Contains(ClassObjectPath))
+			//Only get the asset with the native parent class using UInterchangePipelineBase
+			FAssetDataTagMapSharedView::FFindTagResult GeneratedClassPath = Asset.TagsAndValues.FindTag(TEXT("GeneratedClass"));
+			if (GeneratedClassPath.IsSet())
 			{
-				continue;
-			}
+				// Convert path to just the name part
+				const FTopLevelAssetPath ClassObjectPath(FPackageName::ExportTextPathToObjectPath(*GeneratedClassPath.GetValue()));
+
+				// Check if this class is in the derived set
+				if (!DerivedNames.Contains(ClassObjectPath))
+				{
+					continue;
+				}
 
-			UBlueprint* Blueprint = Cast<UBlueprint>(Asset.GetAsset());
-			if (ensure(Blueprint) && ensure(Blueprint->ParentClass->IsChildOf<UInterchangePipelineBase>()))
+				if (ensure(Blueprint->ParentClass->IsChildOf<UInterchangePipelineBase>()))
+				{
+					PipelineCandidates.AddUnique(Blueprint->GeneratedClass);
+				}
+			}
+		}
+		else if (UInterchangePythonPipelineAsset* PythonPipeline = Cast<UInterchangePythonPipelineAsset>(Asset.GetAsset()))
+		{
+			if (PythonPipeline->PythonClass.Get()->IsChildOf<UInterchangePipelineBase>())
 			{
-				PipelineCandidates.AddUnique(Blueprint->GeneratedClass);
+				PipelineCandidates.AddUnique(PythonPipeline->GetClass());
 			}
 		}
 	}
