[Backout] - CL42152089
#rnx
Original CL Desc
-----------------------------------------------------------------
Fixed ensure when re-importing an asset that was import with a blueprint pipeline

#rnx

diff --git a/Engine/Source/Runtime/Interchange/Engine/Private/InterchangeAssetImportData.cpp b/Engine/Source/Runtime/Interchange/Engine/Private/InterchangeAssetImportData.cpp
--- a/Engine/Source/Runtime/Interchange/Engine/Private/InterchangeAssetImportData.cpp
+++ b/Engine/Source/Runtime/Interchange/Engine/Private/InterchangeAssetImportData.cpp
@@ -1,12 +1,10 @@
 // Copyright Epic Games, Inc. All Rights Reserved.
 
 #include "InterchangeAssetImportData.h"
-#include "InterchangeBlueprintPipelineBase.h"
 #include "InterchangeCustomVersion.h"
 #include "InterchangeEngineLogPrivate.h"
 #include "InterchangeManager.h"
 #include "InterchangePipelineBase.h"
-#include "InterchangePythonPipelineBase.h"
 #include "InterchangeProjectSettings.h"
 #include "JsonObjectConverter.h"
 #include "Serialization/LargeMemoryReader.h"
@@ -16,23 +14,6 @@
 
 #include UE_INLINE_GENERATED_CPP_BY_NAME(InterchangeAssetImportData)
 
-TMap<FString, UClass*> UInterchangeAssetImportData::GetPipelineClassPerName()
-{
-	TMap<FString, UClass*> ClassPerName;
-
-	TArray<UClass*> PipelineCandidates;
-	UInterchangeManager::GetInterchangeManager().FindPipelineCandidate(PipelineCandidates);
-
-	ClassPerName.Reserve(PipelineCandidates.NumBytes());
-
-	for (UClass* PipelineClass : PipelineCandidates)
-	{
-		ClassPerName.Add(PipelineClass->GetFullName(), PipelineClass);
-	}
-
-	return ClassPerName;
-}
-
 void UInterchangeAssetImportData::PostLoad()
 {
 	Super::PostLoad();
@@ -105,8 +86,7 @@ FString SerializeTranslatorSettings(UObject* TranslatorSettings)
 
 UObject* DeSerializePipeline(const FString& PipelineStr, UClass* PipelineClass)
 {
-	// Note: PipelineClass can be a child of either UInterchangePipelineBase or UInterchangePythonPipelineAsset
-	UObject* GeneratedPipeline = NewObject<UObject>(GetTransientPackage(), PipelineClass);
+	UInterchangePipelineBase* GeneratedPipeline = NewObject<UInterchangePipelineBase>(GetTransientPackage(), PipelineClass);
 
 	TSharedPtr<FJsonObject> RootObject;
 	TSharedRef< TJsonReader<> > JsonReader = TJsonReaderFactory<>::Create(PipelineStr);
@@ -116,14 +96,7 @@ UObject* DeSerializePipeline(const FString& PipelineStr, UClass* PipelineClass)
 		FJsonObjectConverter::JsonObjectToUStruct(JsonPipelineProperties.ToSharedRef(), PipelineClass, GeneratedPipeline, 0, 0);
 	}
 
-	if (UInterchangePipelineBase* InterchangePipelineBase = Cast<UInterchangePipelineBase>(GeneratedPipeline))
-	{
-		InterchangePipelineBase->UpdateWeakObjectPtrs();
-	}
-	else if (UInterchangePythonPipelineAsset* PythonPipelineAsset = Cast<UInterchangePythonPipelineAsset>(GeneratedPipeline))
-	{
-		PythonPipelineAsset->GeneratePipeline();
-	}
+	GeneratedPipeline->UpdateWeakObjectPtrs();
 
 	return GeneratedPipeline;
 }
@@ -428,7 +401,15 @@ void UInterchangeAssetImportData::ProcessPipelinesCache() const
 		{
 			TransientPipelines.Reset(CachedPipelines.Num());
 
-			TMap<FString, UClass*> ClassPerName = GetPipelineClassPerName();
+			TMap<FString, UClass*> ClassPerName;
+			for (FThreadSafeObjectIterator It(UClass::StaticClass()); It; ++It)
+			{
+				UClass* Class = Cast<UClass>(*It);
+				if (Class->IsChildOf(UInterchangePipelineBase::StaticClass()))
+				{
+					ClassPerName.Add(Class->GetFullName(), Class);
+				}
+			}
 
 			for (const TPair<FString, FString>& CachedPipeline : CachedPipelines)
 			{
@@ -438,12 +419,10 @@ void UInterchangeAssetImportData::ProcessPipelinesCache() const
 				{
 					ClassFullName = RedirectedObjectName.ToString();
 				}
-
-				//This cannot fail to make sure we have a healthy serialization
+				//This cannot fail to make sure we have a healty serialization
 				if (!ensure(ClassPerName.Contains(ClassFullName)))
 				{
-					//We did not successfully serialize the cached pipelines
-					TransientPipelines.Empty();
+					//We did not successfully serialize the content of the file into the node container
 					return;
 				}
 
diff --git a/Engine/Source/Runtime/Interchange/Engine/Private/InterchangeManager.cpp b/Engine/Source/Runtime/Interchange/Engine/Private/InterchangeManager.cpp
--- a/Engine/Source/Runtime/Interchange/Engine/Private/InterchangeManager.cpp
+++ b/Engine/Source/Runtime/Interchange/Engine/Private/InterchangeManager.cpp
@@ -2942,8 +2942,6 @@ bool UInterchangeManager::IsAttended()
 	return true;
 }
 
-// #todo_interchange: TODO: Can be costly on large project. Using a cache and using the AssetRegistry to monitor addition will be more efficient.
-// This code could also be moved to a utility class.
 void UInterchangeManager::FindPipelineCandidate(TArray<UClass*>& PipelineCandidates)
 {
 	//Find in memory pipeline class
@@ -2973,7 +2971,7 @@ void UInterchangeManager::FindPipelineCandidate(TArray<UClass*>& PipelineCandida
 
 	TArray< FString > ContentPaths;
 	ContentPaths.Add(TEXT("/Game"));
-	//#todo_interchange: TODO do we have an other alternative, this call is synchronous and will wait unitl the registry database have finish the initial scan. If there is a lot of asset it can take multiple second the first time we call it.
+	//TODO do we have an other alternative, this call is synchronous and will wait unitl the registry database have finish the initial scan. If there is a lot of asset it can take multiple second the first time we call it.
 	AssetRegistry.ScanPathsSynchronous(ContentPaths);
 
 	FTopLevelAssetPath BaseClassName = UInterchangePipelineBase::StaticClass()->GetClassPathName();
@@ -2988,45 +2986,34 @@ void UInterchangeManager::FindPipelineCandidate(TArray<UClass*>& PipelineCandida
 		AssetRegistry.GetDerivedClassNames(BaseNames, Excluded, DerivedNames);
 	}
 
-	// Iterate over retrieved blueprint and python pipeline assets
 	FARFilter Filter;
 	Filter.ClassPaths.Add(UBlueprint::StaticClass()->GetClassPathName());
-	Filter.ClassPaths.Add(UInterchangePythonPipelineAsset::StaticClass()->GetClassPathName());
 	Filter.bRecursiveClasses = true;
 	Filter.bRecursivePaths = true;
 
 	TArray< FAssetData > AssetList;
 	AssetRegistry.GetAssets(Filter, AssetList);
 
+	// Iterate over retrieved blueprint assets
 	for (const FAssetData& Asset : AssetList)
 	{
-		if (UBlueprint* Blueprint = Cast<UBlueprint>(Asset.GetAsset()))
+		//Only get the asset with the native parent class using UInterchangePipelineBase
+		FAssetDataTagMapSharedView::FFindTagResult GeneratedClassPath = Asset.TagsAndValues.FindTag(TEXT("GeneratedClass"));
+		if (GeneratedClassPath.IsSet())
 		{
-			//Only get the asset with the native parent class using UInterchangePipelineBase
-			FAssetDataTagMapSharedView::FFindTagResult GeneratedClassPath = Asset.TagsAndValues.FindTag(TEXT("GeneratedClass"));
-			if (GeneratedClassPath.IsSet())
-			{
-				// Convert path to just the name part
-				const FTopLevelAssetPath ClassObjectPath(FPackageName::ExportTextPathToObjectPath(*GeneratedClassPath.GetValue()));
-
-				// Check if this class is in the derived set
-				if (!DerivedNames.Contains(ClassObjectPath))
-				{
-					continue;
-				}
+			// Convert path to just the name part
+			const FTopLevelAssetPath ClassObjectPath(FPackageName::ExportTextPathToObjectPath(*GeneratedClassPath.GetValue()));
 
-				if (ensure(Blueprint->ParentClass->IsChildOf<UInterchangePipelineBase>()))
-				{
-					PipelineCandidates.AddUnique(Blueprint->GeneratedClass);
-				}
-			}
-		}
-		else if (UInterchangePythonPipelineAsset* PythonPipeline = Cast<UInterchangePythonPipelineAsset>(Asset.GetAsset()))
-		{
-			if (PythonPipeline->PythonClass.Get()->IsChildOf<UInterchangePipelineBase>())
+			// Check if this class is in the derived set
+			if (!DerivedNames.Contains(ClassObjectPath))
 			{
-				PipelineCandidates.AddUnique(PythonPipeline->GetClass());
+				continue;
 			}
+
+			UBlueprint* Blueprint = Cast<UBlueprint>(Asset.GetAsset());
+			check(Blueprint);
+			check(Blueprint->ParentClass == UInterchangePipelineBase::StaticClass());
+			PipelineCandidates.AddUnique(Blueprint->GeneratedClass);
 		}
 	}
 }
diff --git a/Engine/Source/Runtime/Interchange/Engine/Public/InterchangeAssetImportData.h b/Engine/Source/Runtime/Interchange/Engine/Public/InterchangeAssetImportData.h
--- a/Engine/Source/Runtime/Interchange/Engine/Public/InterchangeAssetImportData.h
+++ b/Engine/Source/Runtime/Interchange/Engine/Public/InterchangeAssetImportData.h
@@ -201,8 +201,6 @@ private:
 	/** Source file data describing the files that were used to import this asset. Temporary, Primary usage is for Re-import cancellations. */
 	mutable FAssetImportInfo SourceDataBackup;
 #endif
-	// #todo_interchange: TODO: Move to as a IntercnageManager utility method
-	static TMap<FString, UClass*> GetPipelineClassPerName();
 };
 
 /**
diff --git a/Engine/Source/Runtime/Interchange/Engine/Public/InterchangeManager.h b/Engine/Source/Runtime/Interchange/Engine/Public/InterchangeManager.h
--- a/Engine/Source/Runtime/Interchange/Engine/Public/InterchangeManager.h
+++ b/Engine/Source/Runtime/Interchange/Engine/Public/InterchangeManager.h
@@ -1013,6 +1013,4 @@ private:
 	bool bGCEndDelegateCancellAllTask = false;
 
 	friend class UE::Interchange::FScopedTranslator;
-	// #todo_interchange: Remove wen Interchange utility method for pipelines is done
-	friend class UInterchangeAssetImportData;
 };
