Restore fix to ensure when re-importing an asset that was import with a blueprint and/or python pipeline

#rnx

diff --git a/Engine/Source/Runtime/Interchange/Engine/Private/InterchangeAssetImportData.cpp b/Engine/Source/Runtime/Interchange/Engine/Private/InterchangeAssetImportData.cpp
--- a/Engine/Source/Runtime/Interchange/Engine/Private/InterchangeAssetImportData.cpp
+++ b/Engine/Source/Runtime/Interchange/Engine/Private/InterchangeAssetImportData.cpp
@@ -1,10 +1,12 @@
 // Copyright Epic Games, Inc. All Rights Reserved.
 
 #include "InterchangeAssetImportData.h"
+#include "InterchangeBlueprintPipelineBase.h"
 #include "InterchangeCustomVersion.h"
 #include "InterchangeEngineLogPrivate.h"
 #include "InterchangeManager.h"
 #include "InterchangePipelineBase.h"
+#include "InterchangePythonPipelineBase.h"
 #include "InterchangeProjectSettings.h"
 #include "JsonObjectConverter.h"
 #include "Serialization/LargeMemoryReader.h"
@@ -14,6 +16,23 @@
 
 #include UE_INLINE_GENERATED_CPP_BY_NAME(InterchangeAssetImportData)
 
+TMap<FString, UClass*> UInterchangeAssetImportData::GetPipelineClassPerName()
+{
+	TMap<FString, UClass*> ClassPerName;
+
+	TArray<UClass*> PipelineCandidates;
+	UInterchangeManager::GetInterchangeManager().FindPipelineCandidate(PipelineCandidates);
+
+	ClassPerName.Reserve(PipelineCandidates.NumBytes());
+
+	for (UClass* PipelineClass : PipelineCandidates)
+	{
+		ClassPerName.Add(PipelineClass->GetFullName(), PipelineClass);
+	}
+
+	return ClassPerName;
+}
+
 void UInterchangeAssetImportData::PostLoad()
 {
 	Super::PostLoad();
@@ -86,7 +105,8 @@ FString SerializeTranslatorSettings(UObject* TranslatorSettings)
 
 UObject* DeSerializePipeline(const FString& PipelineStr, UClass* PipelineClass)
 {
-	UInterchangePipelineBase* GeneratedPipeline = NewObject<UInterchangePipelineBase>(GetTransientPackage(), PipelineClass);
+	// Note: PipelineClass can be a child of either UInterchangePipelineBase or UInterchangePythonPipelineAsset
+	UObject* GeneratedPipeline = NewObject<UObject>(GetTransientPackage(), PipelineClass);
 
 	TSharedPtr<FJsonObject> RootObject;
 	TSharedRef< TJsonReader<> > JsonReader = TJsonReaderFactory<>::Create(PipelineStr);
@@ -96,7 +116,14 @@ UObject* DeSerializePipeline(const FString& PipelineStr, UClass* PipelineClass)
 		FJsonObjectConverter::JsonObjectToUStruct(JsonPipelineProperties.ToSharedRef(), PipelineClass, GeneratedPipeline, 0, 0);
 	}
 
-	GeneratedPipeline->UpdateWeakObjectPtrs();
+	if (UInterchangePipelineBase* InterchangePipelineBase = Cast<UInterchangePipelineBase>(GeneratedPipeline))
+	{
+		InterchangePipelineBase->UpdateWeakObjectPtrs();
+	}
+	else if (UInterchangePythonPipelineAsset* PythonPipelineAsset = Cast<UInterchangePythonPipelineAsset>(GeneratedPipeline))
+	{
+		PythonPipelineAsset->GeneratePipeline();
+	}
 
 	return GeneratedPipeline;
 }
@@ -401,15 +428,7 @@ void UInterchangeAssetImportData::ProcessPipelinesCache() const
 		{
 			TransientPipelines.Reset(CachedPipelines.Num());
 
-			TMap<FString, UClass*> ClassPerName;
-			for (FThreadSafeObjectIterator It(UClass::StaticClass()); It; ++It)
-			{
-				UClass* Class = Cast<UClass>(*It);
-				if (Class->IsChildOf(UInterchangePipelineBase::StaticClass()))
-				{
-					ClassPerName.Add(Class->GetFullName(), Class);
-				}
-			}
+			TMap<FString, UClass*> ClassPerName = GetPipelineClassPerName();
 
 			for (const TPair<FString, FString>& CachedPipeline : CachedPipelines)
 			{
@@ -419,10 +438,12 @@ void UInterchangeAssetImportData::ProcessPipelinesCache() const
 				{
 					ClassFullName = RedirectedObjectName.ToString();
 				}
-				//This cannot fail to make sure we have a healty serialization
+
+				//This cannot fail to make sure we have a healthy serialization
 				if (!ensure(ClassPerName.Contains(ClassFullName)))
 				{
-					//We did not successfully serialize the content of the file into the node container
+					//We did not successfully serialize the cached pipelines
+					TransientPipelines.Empty();
 					return;
 				}
 
diff --git a/Engine/Source/Runtime/Interchange/Engine/Private/InterchangeManager.cpp b/Engine/Source/Runtime/Interchange/Engine/Private/InterchangeManager.cpp
--- a/Engine/Source/Runtime/Interchange/Engine/Private/InterchangeManager.cpp
+++ b/Engine/Source/Runtime/Interchange/Engine/Private/InterchangeManager.cpp
@@ -5,6 +5,7 @@
 #include "AssetCompilingManager.h"
 #include "AssetRegistry/AssetData.h"
 #include "AssetRegistry/AssetDataTagMap.h"
+#include "AssetRegistry/AssetRegistryHelpers.h"
 #include "AssetRegistry/AssetRegistryModule.h"
 #if WITH_EDITOR
 #include "AssetToolsModule.h"
@@ -2942,6 +2943,8 @@ bool UInterchangeManager::IsAttended()
 	return true;
 }
 
+// #todo_interchange: TODO: Can be costly on large project. Using a cache and using the AssetRegistry to monitor addition will be more efficient.
+// This code could also be moved to a utility class.
 void UInterchangeManager::FindPipelineCandidate(TArray<UClass*>& PipelineCandidates)
 {
 	//Find in memory pipeline class
@@ -2971,23 +2974,23 @@ void UInterchangeManager::FindPipelineCandidate(TArray<UClass*>& PipelineCandida
 
 	TArray< FString > ContentPaths;
 	ContentPaths.Add(TEXT("/Game"));
-	//TODO do we have an other alternative, this call is synchronous and will wait unitl the registry database have finish the initial scan. If there is a lot of asset it can take multiple second the first time we call it.
+	//#todo_interchange: TODO do we have an other alternative, this call is synchronous and will wait unitl the registry database have finish the initial scan. If there is a lot of asset it can take multiple second the first time we call it.
 	AssetRegistry.ScanPathsSynchronous(ContentPaths);
 
-	FTopLevelAssetPath BaseClassName = UInterchangePipelineBase::StaticClass()->GetClassPathName();
 
 	// Use the asset registry to get the set of all class names deriving from Base
 	TSet< FTopLevelAssetPath > DerivedNames;
 	{
-		TArray< FTopLevelAssetPath > BaseNames;
-		BaseNames.Add(BaseClassName);
-
+		FTopLevelAssetPath BaseClassName = UInterchangePipelineBase::StaticClass()->GetClassPathName();
 		TSet< FTopLevelAssetPath > Excluded;
-		AssetRegistry.GetDerivedClassNames(BaseNames, Excluded, DerivedNames);
+
+		AssetRegistry.GetDerivedClassNames({ BaseClassName }, Excluded, DerivedNames);
 	}
 
+	// Iterate over retrieved blueprint and python pipeline assets
 	FARFilter Filter;
 	Filter.ClassPaths.Add(UBlueprint::StaticClass()->GetClassPathName());
+	Filter.ClassPaths.Add(UInterchangePythonPipelineAsset::StaticClass()->GetClassPathName());
 	Filter.bRecursiveClasses = true;
 	Filter.bRecursivePaths = true;
 
@@ -2997,23 +3000,25 @@ void UInterchangeManager::FindPipelineCandidate(TArray<UClass*>& PipelineCandida
 	// Iterate over retrieved blueprint assets
 	for (const FAssetData& Asset : AssetList)
 	{
-		//Only get the asset with the native parent class using UInterchangePipelineBase
-		FAssetDataTagMapSharedView::FFindTagResult GeneratedClassPath = Asset.TagsAndValues.FindTag(TEXT("GeneratedClass"));
-		if (GeneratedClassPath.IsSet())
+		if (UAssetRegistryHelpers::IsAssetDataBlueprintOfClassSet(Asset, DerivedNames))
 		{
-			// Convert path to just the name part
-			const FTopLevelAssetPath ClassObjectPath(FPackageName::ExportTextPathToObjectPath(*GeneratedClassPath.GetValue()));
+			UBlueprint* Blueprint = Cast<UBlueprint>(Asset.GetAsset());
 
-			// Check if this class is in the derived set
-			if (!DerivedNames.Contains(ClassObjectPath))
+			if (ensure(Blueprint) && ensure(Blueprint->ParentClass->IsChildOf<UInterchangePipelineBase>()))
 			{
-				continue;
+				PipelineCandidates.AddUnique(Blueprint->GeneratedClass);
+			}
+		}
+		else if (UClass* AssetClass = Asset.GetClass())
+		{
+			if (AssetClass->IsChildOf<UInterchangePythonPipelineAsset>())
+			{
+				UInterchangePythonPipelineAsset* PythonPipeline = Cast<UInterchangePythonPipelineAsset>(Asset.GetAsset());
+				if (PythonPipeline && PythonPipeline->PythonClass.Get()->IsChildOf<UInterchangePipelineBase>())
+				{
+					PipelineCandidates.AddUnique(PythonPipeline->GetClass());
+				}
 			}
-
-			UBlueprint* Blueprint = Cast<UBlueprint>(Asset.GetAsset());
-			check(Blueprint);
-			check(Blueprint->ParentClass == UInterchangePipelineBase::StaticClass());
-			PipelineCandidates.AddUnique(Blueprint->GeneratedClass);
 		}
 	}
 }
diff --git a/Engine/Source/Runtime/Interchange/Engine/Public/InterchangeAssetImportData.h b/Engine/Source/Runtime/Interchange/Engine/Public/InterchangeAssetImportData.h
--- a/Engine/Source/Runtime/Interchange/Engine/Public/InterchangeAssetImportData.h
+++ b/Engine/Source/Runtime/Interchange/Engine/Public/InterchangeAssetImportData.h
@@ -201,6 +201,8 @@ private:
 	/** Source file data describing the files that were used to import this asset. Temporary, Primary usage is for Re-import cancellations. */
 	mutable FAssetImportInfo SourceDataBackup;
 #endif
+	// #todo_interchange: TODO: Move to as a IntercnageManager utility method
+	static TMap<FString, UClass*> GetPipelineClassPerName();
 };
 
 /**
diff --git a/Engine/Source/Runtime/Interchange/Engine/Public/InterchangeManager.h b/Engine/Source/Runtime/Interchange/Engine/Public/InterchangeManager.h
--- a/Engine/Source/Runtime/Interchange/Engine/Public/InterchangeManager.h
+++ b/Engine/Source/Runtime/Interchange/Engine/Public/InterchangeManager.h
@@ -1013,4 +1013,6 @@ private:
 	bool bGCEndDelegateCancellAllTask = false;
 
 	friend class UE::Interchange::FScopedTranslator;
+	// #todo_interchange: Remove when Interchange utility method for pipelines is done
+	friend class UInterchangeAssetImportData;
 };
