- Allow InterchangePrivateNodeBase::FindPropertyByPathChecked() to follow object pointers (will be used to follow the EditorOnlyData pointer in StaticMesh when getting the NaniteSettings property, in upcoming CL)

diff --git a/Engine/Source/Runtime/Interchange/Core/Private/Nodes/InterchangeBaseNode.cpp b/Engine/Source/Runtime/Interchange/Core/Private/Nodes/InterchangeBaseNode.cpp
--- a/Engine/Source/Runtime/Interchange/Core/Private/Nodes/InterchangeBaseNode.cpp
+++ b/Engine/Source/Runtime/Interchange/Core/Private/Nodes/InterchangeBaseNode.cpp
@@ -643,6 +643,20 @@ FProperty* InterchangePrivateNodeBase::FindPropertyByPathChecked(TVariant<UObjec
 					Container.Set<uint8*>(Property->ContainerPtrToValuePtr<uint8>(Container.Get<uint8*>()));
 				}
 			}
+			else if (FObjectProperty* ObjectProperty = CastField<FObjectProperty>(Property))
+			{
+				UObject* DestObject;
+				if (Container.IsType<UObject*>())
+				{
+					DestObject = *Property->ContainerPtrToValuePtr<UObject*>(Container.Get<UObject*>());
+				}
+				else
+				{
+					DestObject = *Property->ContainerPtrToValuePtr<UObject*>(Container.Get<UObject*>());
+				}
+				Outer = DestObject->GetClass();
+				Container.Set<UObject*>(DestObject);
+			}
 		}
 	} while (RestOfPropertyPath.Len() > 0);
 
