- FIxed PVS issue

diff --git a/Engine/Source/Runtime/Interchange/Core/Private/Nodes/InterchangeBaseNode.cpp b/Engine/Source/Runtime/Interchange/Core/Private/Nodes/InterchangeBaseNode.cpp
--- a/Engine/Source/Runtime/Interchange/Core/Private/Nodes/InterchangeBaseNode.cpp
+++ b/Engine/Source/Runtime/Interchange/Core/Private/Nodes/InterchangeBaseNode.cpp
@@ -652,7 +652,7 @@ FProperty* InterchangePrivateNodeBase::FindPropertyByPathChecked(TVariant<UObjec
 				}
 				else
 				{
-					DestObject = *Property->ContainerPtrToValuePtr<UObject*>(Container.Get<UObject*>());
+					DestObject = *Property->ContainerPtrToValuePtr<UObject*>(Container.Get<uint8*>());
 				}
 				Outer = DestObject->GetClass();
 				Container.Set<UObject*>(DestObject);
