[Interchange] FInterchangeTaskSystem::Tick optimization.
Gaining 30% import speed performance on A Lab scene (without textures).
Going from 6000 to 4000 seconds.

Also introduced Timings for FImportAsyncHelper to have easier time tracking import times.

#rnx

diff --git a/Engine/Source/Runtime/Interchange/Core/Private/InterchangeTaskSystem.cpp b/Engine/Source/Runtime/Interchange/Core/Private/InterchangeTaskSystem.cpp
--- a/Engine/Source/Runtime/Interchange/Core/Private/InterchangeTaskSystem.cpp
+++ b/Engine/Source/Runtime/Interchange/Core/Private/InterchangeTaskSystem.cpp
@@ -254,7 +254,7 @@ void UE::Interchange::FInterchangeTaskSystem::Tick()
 
 	TArray<uint64> WaitingGameThreadTick;
 	TArray<uint64> WaitingTaskGraph;
-	TArray<uint64> DoneTasks;
+	TSet<uint64> DoneTasks;
 	for (TPair<uint64, TSharedPtr<FInterchangeTaskBase, ESPMode::ThreadSafe>>& TaskInfo : TaskPerIdMap)
 	{
 		TSharedPtr<FInterchangeTaskBase, ESPMode::ThreadSafe> Task = TaskInfo.Value;
@@ -361,9 +361,9 @@ void UE::Interchange::FInterchangeTaskSystem::Tick()
 	UpdateWaitingTask(WaitingTaskGraph, bTimeLimitFalse);
 
 	//Release Done task
+	ReleaseAndDoneTasks.Append(DoneTasks);
 	for (uint64 TaskId : DoneTasks)
 	{
 		TaskPerIdMap.Remove(TaskId);
-		ReleaseAndDoneTasks.Add(TaskId);
 	}
 }
diff --git a/Engine/Source/Runtime/Interchange/Core/Public/InterchangeTaskSystem.h b/Engine/Source/Runtime/Interchange/Core/Public/InterchangeTaskSystem.h
--- a/Engine/Source/Runtime/Interchange/Core/Public/InterchangeTaskSystem.h
+++ b/Engine/Source/Runtime/Interchange/Core/Public/InterchangeTaskSystem.h
@@ -223,7 +223,7 @@ namespace UE::Interchange
 
 		//When a task is done, we add it here and release the task. Prerequisite are search in the list and here.
 		//A clean up is done when the TaskPerIdMap is empty.
-		TArray<uint64> ReleaseAndDoneTasks;
+		TSet<uint64> ReleaseAndDoneTasks;
 
 		//The priority value is a counter that is increment for each task added to the system to create the task id
 		//The lower task id are execute earlier if possible.
diff --git a/Engine/Source/Runtime/Interchange/Engine/Private/InterchangeManager.cpp b/Engine/Source/Runtime/Interchange/Engine/Private/InterchangeManager.cpp
--- a/Engine/Source/Runtime/Interchange/Engine/Private/InterchangeManager.cpp
+++ b/Engine/Source/Runtime/Interchange/Engine/Private/InterchangeManager.cpp
@@ -2506,6 +2506,8 @@ UInterchangeManager::ImportInternal(const FString& ContentPath, const UInterchan
 
 	QueueTaskCount = FMath::Clamp(QueueTaskCount + 1, 0, MAX_int32);
 
+	AsyncHelper->TaskExecuteStarted = FPlatformTime::Seconds();
+
 	StartQueuedTasks();
 
 	PreReturn();
@@ -2727,6 +2729,8 @@ TSharedRef<UE::Interchange::FImportAsyncHelper, ESPMode::ThreadSafe> UInterchang
 
 	AsyncHelper->bRuntimeOrPIE = GetEditorUtilities() ? GetEditorUtilities()->IsRuntimeOrPIE() : false;
 
+	AsyncHelper->TaskExecuteStarted = FPlatformTime::Seconds();
+
 	return AsyncHelper;
 }
 
@@ -2737,11 +2741,14 @@ void UInterchangeManager::ReleaseAsyncHelper(TWeakPtr<UE::Interchange::FImportAs
 	check(AsyncHelper.IsValid());
 	
 	constexpr bool bLogWarningsAndErrors = false;
-
+	
+	double ImportTimeDelta;
 	bool bSucceeded = false;
 	{
 		TSharedPtr<FImportAsyncHelper> AsyncHelperPtr = AsyncHelper.Pin();
 
+		ImportTimeDelta = FPlatformTime::Seconds() - AsyncHelperPtr->TaskExecuteStarted;
+
 		//Free the lock to allow the next import to happen
 		if (AsyncHelperPtr->Translators.IsValidIndex(0))
 		{
@@ -2815,6 +2822,9 @@ void UInterchangeManager::ReleaseAsyncHelper(TWeakPtr<UE::Interchange::FImportAs
 			{
 				if(bSucceeded)
 				{
+					FString ImportDoneText = TEXT("Import Done in ") + FString::SanitizeFloat(ImportTimeDelta) + TEXT(" seconds.");
+					UE_LOG(LogInterchangeEngine, Display, TEXT("%s"), *ImportDoneText);
+					
 					TitleText = NSLOCTEXT("Interchange", "Asynchronous_import_end", "Import Done");
 				}
 				else
diff --git a/Engine/Source/Runtime/Interchange/Engine/Public/InterchangeManager.h b/Engine/Source/Runtime/Interchange/Engine/Public/InterchangeManager.h
--- a/Engine/Source/Runtime/Interchange/Engine/Public/InterchangeManager.h
+++ b/Engine/Source/Runtime/Interchange/Engine/Public/InterchangeManager.h
@@ -326,6 +326,8 @@ namespace UE
 
 			void CleanUp();
 
+			double TaskExecuteStarted;
+
 		private:
 			FCriticalSection ClassPermissionLock;
 			// Set of classes whose creation has been denied.
