Interchange: Fix name collision suffixes appending to previous suffixes.

diff --git a/Engine/Source/Runtime/Interchange/Engine/Private/Tasks/InterchangeTaskParsing.cpp b/Engine/Source/Runtime/Interchange/Engine/Private/Tasks/InterchangeTaskParsing.cpp
--- a/Engine/Source/Runtime/Interchange/Engine/Private/Tasks/InterchangeTaskParsing.cpp
+++ b/Engine/Source/Runtime/Interchange/Engine/Private/Tasks/InterchangeTaskParsing.cpp
@@ -25,262 +25,197 @@
 #include "UObject/UObjectGlobals.h"
 #include "UObject/WeakObjectPtrTemplates.h"
 
-struct FTaskData
+namespace UE::InterchangeTaskParsing::Private
 {
-	FString UniqueID;
-	int32 SourceIndex = INDEX_NONE;
-	bool bIsSceneNode = false;
-	TArray<FString> Dependencies;
-	uint64 GraphEventRef;
-	TArray<uint64> Prerequisites;
-	const UClass* FactoryClass;
-
-	TArray<UInterchangeFactoryBaseNode*, TInlineAllocator<1>> Nodes; // For scenes, we can group multiple nodes into a single task as they are usually very light
-};
-
-/**
- * For the Dependency sort to work the predicate must be transitive ( A > B > C implying A > C).
- * That means we must take into account the whole dependency chain, not just the immediate dependencies.
- * 
- * This is a helper struct to quickly create the dependencies chain of a node using a cache.
- */
-struct FNodeDependencyCache
-{
-	void CacheAccumulatedDependencies(const UInterchangeBaseNodeContainer* NodeContainer, const TArray<FTaskData>& SourceTaskDatas)
+	struct FTaskParsingRenameInfo
 	{
-		CachedDependencies.Reset();
-		for (const FTaskData& TaskData : SourceTaskDatas)
-		{
-			TSet<FString> NodeStack;
-			GetAccumulatedDependencies(NodeContainer, TaskData.UniqueID, NodeStack);
-		}
-	}
+		UInterchangeFactoryBaseNode* FactoryNode = nullptr;
+		int32 SourceIndex = INDEX_NONE;
+		FString OriginalName = FString();
+		FString NewName = FString();
+	};
 
-	const TSet<FString>& GetAccumulatedDependencies(const UInterchangeBaseNodeContainer* NodeContainer, const FString& NodeID)
+	struct FTaskData
 	{
-		TRACE_CPUPROFILER_EVENT_SCOPE(FNodeDependencyCache::GetAccumulatedDependencies)
-		ensure(CachedDependencies.Find(NodeID));
-		TSet<FString> NodeStack;
-		return GetAccumulatedDependencies(NodeContainer, NodeID, NodeStack);
-	}
+		FString UniqueID;
+		int32 SourceIndex = INDEX_NONE;
+		bool bIsSceneNode = false;
+		TArray<FString> Dependencies;
+		uint64 GraphEventRef;
+		TArray<uint64> Prerequisites;
+		const UClass* FactoryClass;
 
-private:
+		TArray<UInterchangeFactoryBaseNode*, TInlineAllocator<1>> Nodes; // For scenes, we can group multiple nodes into a single task as they are usually very light
+	};
 
-	const TSet<FString>& GetAccumulatedDependencies(const UInterchangeBaseNodeContainer* NodeContainer, const FString& NodeID, TSet<FString>& NodeStack)
+	/**
+	* For the Dependency sort to work the predicate must be transitive ( A > B > C implying A > C).
+	* That means we must take into account the whole dependency chain, not just the immediate dependencies.
+	* 
+	* This is a helper struct to quickly create the dependencies chain of a node using a cache.
+	*/
+	struct FNodeDependencyCache
 	{
-		if (const TSet<FString>* DependenciesPtr = CachedDependencies.Find(NodeID))
+		void CacheAccumulatedDependencies(const UInterchangeBaseNodeContainer* NodeContainer, const TArray<FTaskData>& SourceTaskDatas)
 		{
-			return *DependenciesPtr;
+			CachedDependencies.Reset();
+			for (const FTaskData& TaskData : SourceTaskDatas)
+			{
+				TSet<FString> NodeStack;
+				GetAccumulatedDependencies(NodeContainer, TaskData.UniqueID, NodeStack);
+			}
 		}
 
-		TSet<FString> Dependencies;
-		AccumulateDependencies(NodeContainer, NodeID, Dependencies, NodeStack);
-		return CachedDependencies.Add(NodeID, MoveTemp(Dependencies));
-	}
-
-	void AccumulateDependencies(const UInterchangeBaseNodeContainer* NodeContainer, const FString& NodeID, TSet<FString>& OutDependenciesSet, TSet<FString>& NodeStack)
-	{
-		const UInterchangeFactoryBaseNode* FactoryNode = NodeContainer->GetFactoryNode(NodeID);
-		if (!FactoryNode)
+		const TSet<FString>& GetAccumulatedDependencies(const UInterchangeBaseNodeContainer* NodeContainer, const FString& NodeID)
 		{
-			return;
+			TRACE_CPUPROFILER_EVENT_SCOPE(FNodeDependencyCache::GetAccumulatedDependencies)
+			ensure(CachedDependencies.Find(NodeID));
+			TSet<FString> NodeStack;
+			return GetAccumulatedDependencies(NodeContainer, NodeID, NodeStack);
 		}
 
-		bool bAlreadyInSet = false;
-		NodeStack.Add(NodeID, &bAlreadyInSet);
-		if (ensureMsgf(!bAlreadyInSet, TEXT("FNodeDependencyCache::AccumulateDependencies - Node \"%s\" is in a circular dependency, assets may not be imported properly."), *NodeID))
+	private:
+
+		const TSet<FString>& GetAccumulatedDependencies(const UInterchangeBaseNodeContainer* NodeContainer, const FString& NodeID, TSet<FString>& NodeStack)
 		{
-			TArray<FString> FactoryDependencies;
-			FactoryNode->GetFactoryDependencies(FactoryDependencies);
-			OutDependenciesSet.Reserve(OutDependenciesSet.Num() + FactoryDependencies.Num());
-			for (const FString& DependencyID : FactoryDependencies)
+			if (const TSet<FString>* DependenciesPtr = CachedDependencies.Find(NodeID))
 			{
-				bAlreadyInSet = false;
-				OutDependenciesSet.Add(DependencyID, &bAlreadyInSet);
-				// Avoid infinite recursion.
-				if (!bAlreadyInSet)
-				{
-					OutDependenciesSet.Append(GetAccumulatedDependencies(NodeContainer, DependencyID, NodeStack));
-				}
+				return *DependenciesPtr;
 			}
-			NodeStack.Remove(NodeID);
-		}
-	}
-
-	TMap<FString, TSet<FString>> CachedDependencies;
-};
-
-
-void UE::Interchange::FTaskParsing::Execute()
-{
-	TRACE_CPUPROFILER_EVENT_SCOPE(UE::Interchange::FTaskParsing::DoTask)
-#if INTERCHANGE_TRACE_ASYNCHRONOUS_TASK_ENABLED
-	INTERCHANGE_TRACE_ASYNCHRONOUS_TASK(ParsingGraph)
-#endif
-	LLM_SCOPE_BYNAME(TEXT("Interchange"));
 
-	TOptional<FGCScopeGuard> GCScopeGuard;
-	if (!IsInGameThread())
-	{
-		GCScopeGuard.Emplace();
-	}
-
-	TSharedPtr<FImportAsyncHelper, ESPMode::ThreadSafe> AsyncHelper = WeakAsyncHelper.Pin();
-	check(AsyncHelper.IsValid());
-
-	//Parse each graph and prepare import task data, we will then be able to create all the task with the correct dependencies
-	TArray<FTaskData> TaskDatas;
+			TSet<FString> Dependencies;
+			AccumulateDependencies(NodeContainer, NodeID, Dependencies, NodeStack);
+			return CachedDependencies.Add(NodeID, MoveTemp(Dependencies));
+		}
 
-	//Avoid creating asset if the asynchronous import is canceled, just create the completion task
-	if (!AsyncHelper->bCancel)
-	{
-		for (int32 SourceIndex = 0; SourceIndex < AsyncHelper->SourceDatas.Num(); ++SourceIndex)
+		void AccumulateDependencies(const UInterchangeBaseNodeContainer* NodeContainer, const FString& NodeID, TSet<FString>& OutDependenciesSet, TSet<FString>& NodeStack)
 		{
-			TArray<FTaskData> SourceTaskDatas;
-
-			if (!AsyncHelper->BaseNodeContainers.IsValidIndex(SourceIndex))
+			const UInterchangeFactoryBaseNode* FactoryNode = NodeContainer->GetFactoryNode(NodeID);
+			if (!FactoryNode)
 			{
-				continue;
+				return;
 			}
 
-			UInterchangeBaseNodeContainer* BaseNodeContainer = AsyncHelper->BaseNodeContainers[SourceIndex].Get();
-			if (!BaseNodeContainer)
+			bool bAlreadyInSet = false;
+			NodeStack.Add(NodeID, &bAlreadyInSet);
+			if (ensureMsgf(!bAlreadyInSet, TEXT("FNodeDependencyCache::AccumulateDependencies - Node \"%s\" is in a circular dependency, assets may not be imported properly."), *NodeID))
 			{
-				continue;
-			}
-
-			//Translation and pipelines are not executed, compute the children cache for translated and factory nodes
-			BaseNodeContainer->ComputeChildrenCache();
-
-			const bool bCanImportSceneNode = AsyncHelper->TaskData.ImportType == EImportType::ImportType_Scene;
-			BaseNodeContainer->IterateNodesOfType<UInterchangeFactoryBaseNode>([&](const FString& NodeUID, UInterchangeFactoryBaseNode* FactoryNode)
-			{
-				UClass* ObjectClass = FactoryNode->GetObjectClass();
-				if (ObjectClass != nullptr)
+				TArray<FString> FactoryDependencies;
+				FactoryNode->GetFactoryDependencies(FactoryDependencies);
+				OutDependenciesSet.Reserve(OutDependenciesSet.Num() + FactoryDependencies.Num());
+				for (const FString& DependencyID : FactoryDependencies)
 				{
-					const UClass* RegisteredFactoryClass = InterchangeManager->GetRegisteredFactoryClass(ObjectClass);
-
-					const bool bIsSceneNode = FactoryNode->GetObjectClass()->IsChildOf<AActor>() || FactoryNode->GetObjectClass()->IsChildOf<UActorComponent>();
-
-					if (!RegisteredFactoryClass || (bIsSceneNode && !bCanImportSceneNode))
+					bAlreadyInSet = false;
+					OutDependenciesSet.Add(DependencyID, &bAlreadyInSet);
+					// Avoid infinite recursion.
+					if (!bAlreadyInSet)
 					{
-						//nothing we can import from this element
-						return;
+						OutDependenciesSet.Append(GetAccumulatedDependencies(NodeContainer, DependencyID, NodeStack));
 					}
-
-					FTaskData& NodeTaskData = SourceTaskDatas.AddDefaulted_GetRef();
-					NodeTaskData.UniqueID = FactoryNode->GetUniqueID();
-					NodeTaskData.SourceIndex = SourceIndex;
-					NodeTaskData.bIsSceneNode = bIsSceneNode;
-					NodeTaskData.Nodes.Add(FactoryNode);
-					FactoryNode->GetFactoryDependencies(NodeTaskData.Dependencies);
-					NodeTaskData.FactoryClass = RegisteredFactoryClass;
 				}
-			});
+				NodeStack.Remove(NodeID);
+			}
+		}
 
-			{
-				FNodeDependencyCache DependencyCache;
-				//Reserving the cache will allow us to use reference when sorting the dependencies.
-				DependencyCache.CacheAccumulatedDependencies(BaseNodeContainer, SourceTaskDatas);
+		TMap<FString, TSet<FString>> CachedDependencies;
+	};
 
-				//Sort per dependencies
-				auto SortByDependencies =
-					[&BaseNodeContainer, &DependencyCache](const FTaskData& A, const FTaskData& B)
-				{
-					const TSet<FString>& BDependencies = DependencyCache.GetAccumulatedDependencies(BaseNodeContainer, B.UniqueID);
-					const TSet<FString>& ADependencies = DependencyCache.GetAccumulatedDependencies(BaseNodeContainer, A.UniqueID);
+	uint64 CreateTasksFromData(
+		FTaskData& InTaskData,
+		const TSharedPtr<Interchange::FImportAsyncHelper, ESPMode::ThreadSafe>& InAsyncHelper,
+		TMap<FString, FTaskParsingRenameInfo>& OutRenameAssets,
+		TSet<FString>& OutCreatedTasksAssetNames
+	)
+	{
+		LLM_SCOPE_BYNAME(TEXT("Interchange"));
+		TRACE_CPUPROFILER_EVENT_SCOPE(CreateTasksFromData::CreateTasksFromData);
 
-					if (ADependencies.Num() == BDependencies.Num())
-					{
-						return A.UniqueID < B.UniqueID;
-					}
+		using namespace UE::Interchange;
 
-					return ADependencies.Num() < BDependencies.Num();
-				};
+		// We expect 1 node per asset task
+		if (!ensure(InTaskData.Nodes.Num() == 1))
+		{
+			for (TObjectPtr<UInterchangeTranslatorBase> Translator : InAsyncHelper->Translators)
+			{
+				if (!Translator || InAsyncHelper->bCancel)
+				{
+					continue;
+				}
 
-				// Nodes cannot depend on a node from another source, so it's faster to sort the dependencies per-source and then append those to the TaskData arrays.
-				SourceTaskDatas.Sort(SortByDependencies);
+				UInterchangeResultWarning_Generic* WarningResult = NewObject<UInterchangeResultWarning_Generic>(
+					GetTransientPackage(),
+					UInterchangeResultWarning_Generic::StaticClass()
+				);
+				WarningResult->Text = NSLOCTEXT("InterchangeTaskParsing", "TranslatorNoAssetImported", "Found multiple nodes in an asset task.");
+				Translator->AddMessage(WarningResult);
 			}
 
-			TaskDatas.Append(MoveTemp(SourceTaskDatas));
+			constexpr uint64 EmptyGraphEvent = 0;
+			return EmptyGraphEvent;
 		}
-	}
-
-	struct FTaskParsingRenameInfo
-	{
-		UInterchangeFactoryBaseNode* FactoryNode = nullptr;
-		int32 SourceIndex = INDEX_NONE;
-		FString OriginalName = FString();
-		FString NewName = FString();
-	};
-
-	TMap<FString, FTaskParsingRenameInfo> RenameAssets;
-	TSet<FString> CreatedTasksAssetNames; // Tracks for which asset name we have created a task so that we don't have 2 tasks for the same asset name
-	TFunction<uint64(FTaskData&)> CreateTasksFromData = [this, &AsyncHelper, &RenameAssets, &CreatedTasksAssetNames](FTaskData& TaskData)
-	{
-		LLM_SCOPE_BYNAME(TEXT("Interchange"));
-		check(TaskData.Nodes.Num() == 1); //We expect 1 node per asset task
 
-		const int32 SourceIndex = TaskData.SourceIndex;
-		const UClass* const FactoryClass = TaskData.FactoryClass;
-		UInterchangeFactoryBaseNode* FactoryNode = TaskData.Nodes[0];
+		const int32 SourceIndex = InTaskData.SourceIndex;
+		const UClass* const FactoryClass = InTaskData.FactoryClass;
+		UInterchangeFactoryBaseNode* FactoryNode = InTaskData.Nodes[0];
 
-		if (TaskData.bIsSceneNode)
+		if (InTaskData.bIsSceneNode)
 		{
-			TSharedPtr<UE::Interchange::FTaskCreateSceneObjects_GameThread, ESPMode::ThreadSafe> TaskCreateSceneObjects = MakeShared<UE::Interchange::FTaskCreateSceneObjects_GameThread, ESPMode::ThreadSafe>(AsyncHelper->ContentBasePath
+			TSharedPtr<UE::Interchange::FTaskCreateSceneObjects_GameThread, ESPMode::ThreadSafe> TaskCreateSceneObjects = MakeShared<UE::Interchange::FTaskCreateSceneObjects_GameThread, ESPMode::ThreadSafe>(InAsyncHelper->ContentBasePath
 				, SourceIndex
-				, WeakAsyncHelper
-				, TaskData.Nodes
+				, InAsyncHelper
+				, InTaskData.Nodes
 				, FactoryClass);
 
-			return AsyncHelper->SceneTasks.Add_GetRef(UE::Interchange::FInterchangeTaskSystem::Get().AddTask(TaskCreateSceneObjects, TaskData.Prerequisites));
+			return InAsyncHelper->SceneTasks.Add_GetRef(UE::Interchange::FInterchangeTaskSystem::Get().AddTask(TaskCreateSceneObjects, InTaskData.Prerequisites));
 		}
 		else
 		{
 			FString PackageSubPath;
 			FactoryNode->GetCustomSubPath(PackageSubPath);
 
-			FString AssetFullPath = FPaths::Combine(AsyncHelper->ContentBasePath, PackageSubPath, FactoryNode->GetAssetName());
+			FString AssetFullPath = FPaths::Combine(InAsyncHelper->ContentBasePath, PackageSubPath, FactoryNode->GetAssetName());
 
-			//Make sure there is no duplicate name full path
-			uint32 NameIndex = 1;
-			FString NewName = AssetFullPath;
-			bool NameClash = FactoryNode->IsEnabled();
-			
-			while (NameClash)
+			// Make sure there is no duplicate name full path as we'll need to create a new asset
+			bool bCheckForNameClash = FactoryNode->IsEnabled();
+			if (bCheckForNameClash)
 			{
-				if (CreatedTasksAssetNames.Contains(NewName))
+				uint32 NameIndex = 1;
+				FString NewName = AssetFullPath;
+				while (OutCreatedTasksAssetNames.Contains(NewName))
 				{
-					const FString NameIndexString = FString::FromInt(NameIndex++);
-					NewName = AssetFullPath + NameIndexString;
-					FTaskParsingRenameInfo& RenameInfo = RenameAssets.FindOrAdd(AssetFullPath);
+					NewName = AssetFullPath + FString::FromInt(NameIndex++);
+				}
+
+				// If we renamed anything
+				if (NameIndex != 1)
+				{
+					FTaskParsingRenameInfo& RenameInfo = OutRenameAssets.FindOrAdd(AssetFullPath);
 					RenameInfo.FactoryNode = FactoryNode;
 					RenameInfo.OriginalName = AssetFullPath;
 					RenameInfo.NewName = NewName;
 					RenameInfo.SourceIndex = SourceIndex;
+
+					const FString NameIndexString = FString::FromInt(NameIndex);
 					FactoryNode->SetDisplayLabel(FactoryNode->GetDisplayLabel() + NameIndexString);
 					if (FactoryNode->HasAttribute(UE::Interchange::FBaseNodeStaticData::AssetNameKey()))
 					{
 						//TextureFactorNodes automatically set up the AssetNames which then are used for asset creation:
 						FactoryNode->SetAssetName(FactoryNode->GetDisplayLabel());
 					}
-				}
-				else
-				{
+
 					AssetFullPath = NewName;
-					NameClash = false;
 				}
 			}
 
-			if (!FactoryNode->IsEnabled() || ensureMsgf(!CreatedTasksAssetNames.Contains(AssetFullPath),
+			if (!FactoryNode->IsEnabled() || ensureMsgf(!OutCreatedTasksAssetNames.Contains(AssetFullPath),
 				TEXT("Found multiple task data with the same asset name (%s). Only one will be executed."), *AssetFullPath))
 			{
 				LLM_SCOPE_BYNAME(TEXT("Interchange"));
-				
+
 				// Create factory
 				UInterchangeFactoryBase* Factory = NewObject<UInterchangeFactoryBase>(GetTransientPackage(), FactoryClass);
-				Factory->SetResultsContainer(AsyncHelper->AssetImportResult->GetResults());
-				AsyncHelper->AddCreatedFactory(FactoryNode->GetUniqueID(), Factory);
+				Factory->SetResultsContainer(InAsyncHelper->AssetImportResult->GetResults());
+				InAsyncHelper->AddCreatedFactory(FactoryNode->GetUniqueID(), Factory);
 
 				//If the factory wants to get async payloads it can add the task here so the payloads are in the dependency graph
 				TArray<uint64> ImportObjectQueryPayloadsTasksPrerequistes;
@@ -289,69 +224,67 @@ void UE::Interchange::FTaskParsing::Execute()
 					UInterchangeFactoryBase::FImportAssetObjectParams ImportAssetObjectParams;
 					ImportAssetObjectParams.AssetName = FactoryNode->GetDisplayLabel();
 					ImportAssetObjectParams.AssetNode = FactoryNode;
-					ImportAssetObjectParams.NodeContainer = AsyncHelper->BaseNodeContainers[SourceIndex].Get();
-					ImportAssetObjectParams.SourceData = AsyncHelper->SourceDatas[SourceIndex];
-					ImportAssetObjectParams.Translator = AsyncHelper->Translators[SourceIndex];
-					bool bAsyncPayloads = !AsyncHelper->bRunSynchronous;
+					ImportAssetObjectParams.NodeContainer = InAsyncHelper->BaseNodeContainers[SourceIndex].Get();
+					ImportAssetObjectParams.SourceData = InAsyncHelper->SourceDatas[SourceIndex];
+					ImportAssetObjectParams.Translator = InAsyncHelper->Translators[SourceIndex];
+					bool bAsyncPayloads = !InAsyncHelper->bRunSynchronous;
 					if (!ImportAssetObjectParams.Translator || !ImportAssetObjectParams.Translator->IsThreadSafe())
 					{
 						//Translator not thread safe cannot get payloads asynchronously
 						bAsyncPayloads = false;
 					}
 
-					if (!AsyncHelper->bRuntimeOrPIE || FactoryNode->IsRuntimeImportAllowed())
+					if (!InAsyncHelper->bRuntimeOrPIE || FactoryNode->IsRuntimeImportAllowed())
 					{
 						//Only create payloads if it's allowed.
 						Factory->CreatePayloadTasks(ImportAssetObjectParams, bAsyncPayloads, PayloadTasks);
 					}
-					
+
 					if (!PayloadTasks.IsEmpty())
 					{
 						for (TSharedPtr<FInterchangeTaskBase> PayloadTask : PayloadTasks)
 						{
-							int32 ImportObjectQueryPayloadsTaskIndex = AsyncHelper->ImportObjectQueryPayloadsTasks.Add(FInterchangeTaskSystem::Get().AddTask(PayloadTask, TaskData.Prerequisites));
-							ImportObjectQueryPayloadsTasksPrerequistes.Add(AsyncHelper->ImportObjectQueryPayloadsTasks[ImportObjectQueryPayloadsTaskIndex]);
+							int32 ImportObjectQueryPayloadsTaskIndex = InAsyncHelper->ImportObjectQueryPayloadsTasks.Add(FInterchangeTaskSystem::Get().AddTask(PayloadTask, InTaskData.Prerequisites));
+							ImportObjectQueryPayloadsTasksPrerequistes.Add(InAsyncHelper->ImportObjectQueryPayloadsTasks[ImportObjectQueryPayloadsTaskIndex]);
 						}
 					}
 					else
 					{
-						ImportObjectQueryPayloadsTasksPrerequistes = TaskData.Prerequisites;
+						ImportObjectQueryPayloadsTasksPrerequistes = InTaskData.Prerequisites;
 					}
 				}
 
 				TArray<uint64> ImportObjectTasksPrerequistes;
-				TSharedPtr<FTaskImportObject_GameThread, ESPMode::ThreadSafe> TaskImportObject_GameThread = MakeShared<FTaskImportObject_GameThread, ESPMode::ThreadSafe>(AsyncHelper->ContentBasePath
+				TSharedPtr<FTaskImportObject_GameThread, ESPMode::ThreadSafe> TaskImportObject_GameThread = MakeShared<FTaskImportObject_GameThread, ESPMode::ThreadSafe>(InAsyncHelper->ContentBasePath
 					, SourceIndex
-					, WeakAsyncHelper
+					, InAsyncHelper
 					, FactoryNode
 					, FactoryClass);
-				int32 BeginImportObjectTaskIndex = AsyncHelper->BeginImportObjectTasks.Add(FInterchangeTaskSystem::Get().AddTask(TaskImportObject_GameThread, ImportObjectQueryPayloadsTasksPrerequistes));
-				ImportObjectTasksPrerequistes.Add(AsyncHelper->BeginImportObjectTasks[BeginImportObjectTaskIndex]);
+				int32 BeginImportObjectTaskIndex = InAsyncHelper->BeginImportObjectTasks.Add(FInterchangeTaskSystem::Get().AddTask(TaskImportObject_GameThread, ImportObjectQueryPayloadsTasksPrerequistes));
+				ImportObjectTasksPrerequistes.Add(InAsyncHelper->BeginImportObjectTasks[BeginImportObjectTaskIndex]);
 
-				TSharedPtr<FTaskImportObject_Async, ESPMode::ThreadSafe> TaskImportObject_Async = MakeShared<FTaskImportObject_Async, ESPMode::ThreadSafe>(AsyncHelper->ContentBasePath
+				TSharedPtr<FTaskImportObject_Async, ESPMode::ThreadSafe> TaskImportObject_Async = MakeShared<FTaskImportObject_Async, ESPMode::ThreadSafe>(InAsyncHelper->ContentBasePath
 					, SourceIndex
-					, WeakAsyncHelper
+					, InAsyncHelper
 					, FactoryNode);
-				int32 ImportObjectTaskIndex = AsyncHelper->ImportObjectTasks.Add(FInterchangeTaskSystem::Get().AddTask(TaskImportObject_Async, ImportObjectTasksPrerequistes));
-				
+				int32 ImportObjectTaskIndex = InAsyncHelper->ImportObjectTasks.Add(FInterchangeTaskSystem::Get().AddTask(TaskImportObject_Async, ImportObjectTasksPrerequistes));
+
 				TArray<uint64> FinalizeImportObjectTasksPrerequistes;
-				FinalizeImportObjectTasksPrerequistes.Add(AsyncHelper->ImportObjectTasks[ImportObjectTaskIndex]);
+				FinalizeImportObjectTasksPrerequistes.Add(InAsyncHelper->ImportObjectTasks[ImportObjectTaskIndex]);
 
-				
-				
-				TSharedPtr<FTaskImportObjectFinalize_GameThread, ESPMode::ThreadSafe> TaskImportObjectFinalize_GameThread = MakeShared<FTaskImportObjectFinalize_GameThread, ESPMode::ThreadSafe>(AsyncHelper->ContentBasePath
+				TSharedPtr<FTaskImportObjectFinalize_GameThread, ESPMode::ThreadSafe> TaskImportObjectFinalize_GameThread = MakeShared<FTaskImportObjectFinalize_GameThread, ESPMode::ThreadSafe>(InAsyncHelper->ContentBasePath
 					, SourceIndex
-					, WeakAsyncHelper
+					, InAsyncHelper
 					, FactoryNode);
-				int32 FinalizeCreateTaskIndex = AsyncHelper->FinalizeImportObjectTasks.Add(FInterchangeTaskSystem::Get().AddTask(TaskImportObjectFinalize_GameThread, FinalizeImportObjectTasksPrerequistes));
+				int32 FinalizeCreateTaskIndex = InAsyncHelper->FinalizeImportObjectTasks.Add(FInterchangeTaskSystem::Get().AddTask(TaskImportObjectFinalize_GameThread, FinalizeImportObjectTasksPrerequistes));
 
 				//Only add the name if the factory node is enabled
 				if (FactoryNode->IsEnabled())
 				{
-					CreatedTasksAssetNames.Add(AssetFullPath);
+					OutCreatedTasksAssetNames.Add(AssetFullPath);
 				}
 
-				return AsyncHelper->FinalizeImportObjectTasks[FinalizeCreateTaskIndex];
+				return InAsyncHelper->FinalizeImportObjectTasks[FinalizeCreateTaskIndex];
 			}
 			else
 			{
@@ -359,7 +292,108 @@ void UE::Interchange::FTaskParsing::Execute()
 				return EmptyGraphEvent;
 			}
 		}
-	};
+	}
+}
+
+void UE::Interchange::FTaskParsing::Execute()
+{
+	TRACE_CPUPROFILER_EVENT_SCOPE(UE::Interchange::FTaskParsing::DoTask)
+#if INTERCHANGE_TRACE_ASYNCHRONOUS_TASK_ENABLED
+	INTERCHANGE_TRACE_ASYNCHRONOUS_TASK(ParsingGraph)
+#endif
+	LLM_SCOPE_BYNAME(TEXT("Interchange"));
+
+	using namespace UE::InterchangeTaskParsing::Private;
+
+	TOptional<FGCScopeGuard> GCScopeGuard;
+	if (!IsInGameThread())
+	{
+		GCScopeGuard.Emplace();
+	}
+
+	TSharedPtr<FImportAsyncHelper, ESPMode::ThreadSafe> AsyncHelper = WeakAsyncHelper.Pin();
+	check(AsyncHelper.IsValid());
+
+	//Parse each graph and prepare import task data, we will then be able to create all the task with the correct dependencies
+	TArray<FTaskData> TaskDatas;
+
+	//Avoid creating asset if the asynchronous import is canceled, just create the completion task
+	if (!AsyncHelper->bCancel)
+	{
+		for (int32 SourceIndex = 0; SourceIndex < AsyncHelper->SourceDatas.Num(); ++SourceIndex)
+		{
+			TArray<FTaskData> SourceTaskDatas;
+
+			if (!AsyncHelper->BaseNodeContainers.IsValidIndex(SourceIndex))
+			{
+				continue;
+			}
+
+			UInterchangeBaseNodeContainer* BaseNodeContainer = AsyncHelper->BaseNodeContainers[SourceIndex].Get();
+			if (!BaseNodeContainer)
+			{
+				continue;
+			}
+
+			//Translation and pipelines are not executed, compute the children cache for translated and factory nodes
+			BaseNodeContainer->ComputeChildrenCache();
+
+			const bool bCanImportSceneNode = AsyncHelper->TaskData.ImportType == EImportType::ImportType_Scene;
+			BaseNodeContainer->IterateNodesOfType<UInterchangeFactoryBaseNode>([&](const FString& NodeUID, UInterchangeFactoryBaseNode* FactoryNode)
+			{
+				UClass* ObjectClass = FactoryNode->GetObjectClass();
+				if (ObjectClass != nullptr)
+				{
+					const UClass* RegisteredFactoryClass = InterchangeManager->GetRegisteredFactoryClass(ObjectClass);
+
+					const bool bIsSceneNode = FactoryNode->GetObjectClass()->IsChildOf<AActor>() || FactoryNode->GetObjectClass()->IsChildOf<UActorComponent>();
+
+					if (!RegisteredFactoryClass || (bIsSceneNode && !bCanImportSceneNode))
+					{
+						//nothing we can import from this element
+						return;
+					}
+
+					FTaskData& NodeTaskData = SourceTaskDatas.AddDefaulted_GetRef();
+					NodeTaskData.UniqueID = FactoryNode->GetUniqueID();
+					NodeTaskData.SourceIndex = SourceIndex;
+					NodeTaskData.bIsSceneNode = bIsSceneNode;
+					NodeTaskData.Nodes.Add(FactoryNode);
+					FactoryNode->GetFactoryDependencies(NodeTaskData.Dependencies);
+					NodeTaskData.FactoryClass = RegisteredFactoryClass;
+				}
+			});
+
+			{
+				FNodeDependencyCache DependencyCache;
+				//Reserving the cache will allow us to use reference when sorting the dependencies.
+				DependencyCache.CacheAccumulatedDependencies(BaseNodeContainer, SourceTaskDatas);
+
+				//Sort per dependencies
+				auto SortByDependencies =
+					[&BaseNodeContainer, &DependencyCache](const FTaskData& A, const FTaskData& B)
+				{
+					const TSet<FString>& BDependencies = DependencyCache.GetAccumulatedDependencies(BaseNodeContainer, B.UniqueID);
+					const TSet<FString>& ADependencies = DependencyCache.GetAccumulatedDependencies(BaseNodeContainer, A.UniqueID);
+
+					if (ADependencies.Num() == BDependencies.Num())
+					{
+						return A.UniqueID < B.UniqueID;
+					}
+
+					return ADependencies.Num() < BDependencies.Num();
+				};
+
+				// Nodes cannot depend on a node from another source, so it's faster to sort the dependencies per-source and then append those to the TaskData arrays.
+				SourceTaskDatas.Sort(SortByDependencies);
+			}
+
+			TaskDatas.Append(MoveTemp(SourceTaskDatas));
+		}
+	}
+
+	TMap<FString, FTaskParsingRenameInfo> RenameAssets;
+	TSet<FString> CreatedTasksAssetNames; // Tracks for which asset name we have created a task so that we don't have 2 tasks for the same asset name
 
 	//If we find nothing to import, we notify the user
 	if (TaskDatas.Num() == 0)
@@ -406,7 +440,7 @@ void UE::Interchange::FTaskParsing::Execute()
 			TaskData.Prerequisites.Append(GroupPrerequistes);
 			GroupPrerequistes.Reset();
 		}
-		TaskData.GraphEventRef = CreateTasksFromData(TaskData);
+		TaskData.GraphEventRef = CreateTasksFromData(TaskData, AsyncHelper, RenameAssets, CreatedTasksAssetNames);
 		GroupPrerequistes.Add(TaskData.GraphEventRef);
 	}
 	CompletionPrerequistes.Append(GroupPrerequistes);
