Interchange framework is bypassed on the second attempt after a failed import
- Ensured proper cleanup when the import fails.

#rnx

diff --git a/Engine/Source/Runtime/Interchange/Engine/Private/Tasks/InterchangeTaskCompletion.cpp b/Engine/Source/Runtime/Interchange/Engine/Private/Tasks/InterchangeTaskCompletion.cpp
--- a/Engine/Source/Runtime/Interchange/Engine/Private/Tasks/InterchangeTaskCompletion.cpp
+++ b/Engine/Source/Runtime/Interchange/Engine/Private/Tasks/InterchangeTaskCompletion.cpp
@@ -221,6 +221,7 @@ void UE::Interchange::FTaskCompletion_GameThread::Execute()
 					{
 						if (Asset->HasAnyFlags(RF_MirroredGarbage))
 						{
+							UE::Interchange::Private::ObjectDeletionUtils::PurgeObject(Asset);
 							continue;
 						}
 
