Interchange: Fix off-by-one bug when adding suffixes to produce unique asset names.

diff --git a/Engine/Source/Runtime/Interchange/Engine/Private/Tasks/InterchangeTaskParsing.cpp b/Engine/Source/Runtime/Interchange/Engine/Private/Tasks/InterchangeTaskParsing.cpp
--- a/Engine/Source/Runtime/Interchange/Engine/Private/Tasks/InterchangeTaskParsing.cpp
+++ b/Engine/Source/Runtime/Interchange/Engine/Private/Tasks/InterchangeTaskParsing.cpp
@@ -179,15 +179,15 @@ namespace UE::InterchangeTaskParsing::Private
 			bool bCheckForNameClash = FactoryNode->IsEnabled();
 			if (bCheckForNameClash)
 			{
-				uint32 NameIndex = 1;
+				uint32 NameIndex = 0;
 				FString NewName = AssetFullPath;
 				while (OutCreatedTasksAssetNames.Contains(NewName))
 				{
-					NewName = AssetFullPath + FString::FromInt(NameIndex++);
+					NewName = AssetFullPath + FString::FromInt(++NameIndex);
 				}
 
 				// If we renamed anything
-				if (NameIndex != 1)
+				if (NameIndex != 0)
 				{
 					FTaskParsingRenameInfo& RenameInfo = OutRenameAssets.FindOrAdd(AssetFullPath);
 					RenameInfo.FactoryNode = FactoryNode;
