Interchange: Minor cleanup of AttributeStorage.

Move implementations to .cpp file.

diff --git a/Engine/Plugins/Interchange/Runtime/Source/Nodes/Public/InterchangeAnimationDefinitions.h b/Engine/Plugins/Interchange/Runtime/Source/Nodes/Public/InterchangeAnimationDefinitions.h
--- a/Engine/Plugins/Interchange/Runtime/Source/Nodes/Public/InterchangeAnimationDefinitions.h
+++ b/Engine/Plugins/Interchange/Runtime/Source/Nodes/Public/InterchangeAnimationDefinitions.h
@@ -349,23 +349,3 @@ enum class EInterchangePropertyTracks : int32
 	Visibility UE_DEPRECATED(5.6, "Renamed to 'ActorHiddenInGame'") = ActorHiddenInGame UMETA(hidden)
 };
 
-namespace UE
-{
-	namespace Interchange
-	{
-		template<> struct TAttributeTypeTraits<EInterchangePropertyTracks>
-		{
-			static constexpr EAttributeTypes GetType()
-			{
-				return EAttributeTypes::Int32;
-			}
-			static FString ToString(const uint16& Value)
-			{
-				int32 ValueConv = Value;
-				FStringFormatOrderedArguments OrderedArguments;
-				OrderedArguments.Add(FStringFormatArg(ValueConv));
-				return FString::Format(TEXT("{0}"), OrderedArguments);
-			}
-		};
-	}
-}
\ No newline at end of file
diff --git a/Engine/Plugins/Interchange/Runtime/Source/Nodes/Public/InterchangeMeshDefinitions.h b/Engine/Plugins/Interchange/Runtime/Source/Nodes/Public/InterchangeMeshDefinitions.h
--- a/Engine/Plugins/Interchange/Runtime/Source/Nodes/Public/InterchangeMeshDefinitions.h
+++ b/Engine/Plugins/Interchange/Runtime/Source/Nodes/Public/InterchangeMeshDefinitions.h
@@ -53,21 +53,6 @@ namespace UE
 {
 	namespace Interchange
 	{
-		template<> struct TAttributeTypeTraits<EInterchangeMeshCollision>
-		{
-			static constexpr EAttributeTypes GetType()
-			{
-				return EAttributeTypes::UInt8;
-			}
-			static FString ToString(const uint16& Value)
-			{
-				int32 ValueConv = Value;
-				FStringFormatOrderedArguments OrderedArguments;
-				OrderedArguments.Add(FStringFormatArg(ValueConv));
-				return FString::Format(TEXT("{0}"), OrderedArguments);
-			}
-		};
-
 		namespace MeshPayload
 		{
 			namespace Attributes
diff --git a/Engine/Source/Runtime/Interchange/Core/Private/Tests/StorageTest.cpp b/Engine/Source/Runtime/Interchange/Core/Private/Tests/StorageTest.cpp
--- a/Engine/Source/Runtime/Interchange/Core/Private/Tests/StorageTest.cpp
+++ b/Engine/Source/Runtime/Interchange/Core/Private/Tests/StorageTest.cpp
@@ -34,6 +34,19 @@ IMPLEMENT_SIMPLE_AUTOMATION_TEST(FAttributeStorageTest, "System.Runtime.Intercha
 		AddError(FString(TEXT("`AttributeStorage` must handle adding ATTRIBUTE_TYPE attribute")));														\
 	}
 
+
+namespace UE::InterchangeStorageTest::Private
+{
+	namespace ETestEnum
+	{
+		enum Type : int
+		{
+			Perspective,
+			Orthographic
+		};
+	};
+}
+
 bool FAttributeStorageTest::RunTest(const FString& Parameters)
 {
 	using namespace UE::Interchange;
@@ -319,6 +332,25 @@ bool FAttributeStorageTest::RunTest(const FString& Parameters)
 		ADD_ATTRIBUTE_STORAGE(EAttributeStorageEnumTest, EAttributeStorageEnumTest::EnumValue255);
 	}
 
+	//Store a TEnumAsByte class
+	{
+		enum EBlueprintEnum : int
+		{
+			Opaque,
+			Masked,
+			Translucent,
+			Additive
+		};
+		ADD_ATTRIBUTE_STORAGE(TEnumAsByte<EBlueprintEnum>, EBlueprintEnum::Opaque);
+		ADD_ATTRIBUTE_STORAGE(TEnumAsByte<EBlueprintEnum>, EBlueprintEnum::Masked);
+		ADD_ATTRIBUTE_STORAGE(TEnumAsByte<EBlueprintEnum>, EBlueprintEnum::Translucent);
+		ADD_ATTRIBUTE_STORAGE(TEnumAsByte<EBlueprintEnum>, EBlueprintEnum::Additive);
+
+		using namespace UE::InterchangeStorageTest::Private;
+		ADD_ATTRIBUTE_STORAGE(TEnumAsByte<ETestEnum::Type>, ETestEnum::Perspective);
+		ADD_ATTRIBUTE_STORAGE(TEnumAsByte<ETestEnum::Type>, ETestEnum::Orthographic);
+	}
+
 	//Get the Hash we will use it later to do some test
 	FGuid HashGuidRef = TestStorage.GetStorageHash();
 
diff --git a/Engine/Source/Runtime/Interchange/Core/Private/Types/AttributeStorage.cpp b/Engine/Source/Runtime/Interchange/Core/Private/Types/AttributeStorage.cpp
--- a/Engine/Source/Runtime/Interchange/Core/Private/Types/AttributeStorage.cpp
+++ b/Engine/Source/Runtime/Interchange/Core/Private/Types/AttributeStorage.cpp
@@ -1,16 +1,451 @@
 // Copyright Epic Games, Inc. All Rights Reserved.
+
 #include "Types/AttributeStorage.h"
 #include "InterchangeLogPrivate.h"
 
 #include "CoreMinimal.h"
+#include "HAL/UnrealMemory.h"
+#include "Math/Box.h"
+#include "Math/Box2D.h"
+#include "Math/BoxSphereBounds.h"
+#include "Math/Color.h"
+#include "Math/Float16.h"
+#include "Math/IntPoint.h"
+#include "Math/IntRect.h"
+#include "Math/IntVector.h"
+#include "Math/Matrix.h"
+#include "Math/Plane.h"
+#include "Math/Quat.h"
+#include "Math/RandomStream.h"
+#include "Math/Rotator.h"
+#include "Math/Sphere.h"
+#include "Math/TwoVectors.h"
+#include "Math/Vector.h"
+#include "Math/Vector2D.h"
+#include "Math/Vector2DHalf.h"
+#include "Math/Vector4.h"
+#include "Misc/FrameRate.h"
 #include "Misc/ScopeLock.h"
+#include "Misc/ScopeLock.h"
 #include "Misc/SecureHash.h"
 
+#include <type_traits>
+
 //Interchange namespace
 namespace UE
 {
 	namespace Interchange
 	{
+		/**
+		 * Stub for attribute type traits.
+		 *
+		 * Actual type traits need to be declared through template specialization for custom
+		 * data types that are to be used internally by FAttributeStorage. Traits for the most commonly used built-in
+		 * types are declared below.
+		 *
+		 * Complex types, such as structures and classes, can be serialized into a byte array
+		 * and then assigned to an attribute. Note that you will be responsible for ensuring
+		 * correct byte ordering when serializing those types.
+		 *
+		 * @param T The type to be used in FAttributeStorage.
+		 */
+		template<typename T> 
+		struct TAttributeTypeTraitsInternal
+		{
+		};
+
+		template<> 
+		struct TAttributeTypeTraitsInternal<bool>
+		{
+			static constexpr EAttributeTypes GetType() { return EAttributeTypes::Bool; }
+			static FString ToString(const bool& Value) { return Value ? TEXT("true") : TEXT("false"); }
+		};
+
+		template<> 
+		struct TAttributeTypeTraitsInternal<FColor>
+		{
+			static constexpr EAttributeTypes GetType() { return EAttributeTypes::Color; }
+			static FString ToString(const FColor& Value) { return Value.ToString(); }
+		};
+
+		template<> 
+		struct TAttributeTypeTraitsInternal<FDateTime>
+		{
+			static constexpr EAttributeTypes GetType() { return EAttributeTypes::DateTime; }
+			static FString ToString(const FDateTime& Value) { return Value.ToString(); }
+		};
+
+		template<> 
+		struct TAttributeTypeTraitsInternal<double>
+		{
+			static constexpr EAttributeTypes GetType() { return EAttributeTypes::Double; }
+			static FString ToString(const double& Value) { return LexToString(Value); }
+		};
+
+		template<> 
+		struct TAttributeTypeTraitsInternal<float>
+		{
+			static constexpr EAttributeTypes GetType() { return EAttributeTypes::Float; }
+			static FString ToString(const float& Value) { return LexToString(Value); }
+		};
+
+		template<> 
+		struct TAttributeTypeTraitsInternal<FGuid>
+		{
+			static constexpr EAttributeTypes GetType() { return EAttributeTypes::Guid; }
+			static FString ToString(const FGuid& Value) { return Value.ToString(EGuidFormats::DigitsLower); }
+		};
+
+		template<> 
+		struct TAttributeTypeTraitsInternal<int8>
+		{
+			static constexpr EAttributeTypes GetType() { return EAttributeTypes::Int8; }
+			static FString ToString(const int8& Value) { return FString::FromInt(Value); }
+		};
+
+		template<> 
+		struct TAttributeTypeTraitsInternal<int16>
+		{
+			static constexpr EAttributeTypes GetType() { return EAttributeTypes::Int16; }
+			static FString ToString(const int16& Value) { return FString::FromInt(Value); }
+		};
+
+		template<> 
+		struct TAttributeTypeTraitsInternal<int32>
+		{
+			static constexpr EAttributeTypes GetType() { return EAttributeTypes::Int32; }
+			static FString ToString(const int32& Value) { return FString::FromInt(Value); }
+		};
+
+		template<> 
+		struct TAttributeTypeTraitsInternal<int64>
+		{
+			static constexpr EAttributeTypes GetType() { return EAttributeTypes::Int64; }
+			static FString ToString(const int64& Value) { return LexToString(Value); }
+		};
+
+		template<> 
+		struct TAttributeTypeTraitsInternal<FIntRect>
+		{
+			static constexpr EAttributeTypes GetType() { return EAttributeTypes::IntRect; }
+			static FString ToString(const FIntRect& Value) { return Value.ToString(); }
+		};
+
+		template<> 
+		struct TAttributeTypeTraitsInternal<FLinearColor>
+		{
+			static constexpr EAttributeTypes GetType() { return EAttributeTypes::LinearColor; }
+			static FString ToString(const FLinearColor& Value) { return Value.ToString(); }
+		};
+
+		template<> 
+		struct TAttributeTypeTraitsInternal<FName>
+		{
+			static constexpr EAttributeTypes GetType() { return EAttributeTypes::Name; }
+			static FString ToString(const FName& Value) { return Value.ToString(); }
+		};
+
+		template<> 
+		struct TAttributeTypeTraitsInternal<FRandomStream>
+		{
+			static constexpr EAttributeTypes GetType() { return EAttributeTypes::RandomStream; }
+			static FString ToString(const FRandomStream& Value) { return Value.ToString(); }
+		};
+
+		template<> 
+		struct TAttributeTypeTraitsInternal<FString>
+		{
+			static constexpr EAttributeTypes GetType() { return EAttributeTypes::String; }
+			static FString ToString(const FString& Value) { return Value; }
+		};
+
+		template<> 
+		struct TAttributeTypeTraitsInternal<FTimespan>
+		{
+			static constexpr EAttributeTypes GetType() { return EAttributeTypes::Timespan; }
+			static FString ToString(const FTimespan& Value) { return Value.ToString(); }
+		};
+
+		template<> 
+		struct TAttributeTypeTraitsInternal<FTwoVectors>
+		{
+			static constexpr EAttributeTypes GetType() { return EAttributeTypes::TwoVectors; }
+			static FString ToString(const FTwoVectors& Value) { return Value.ToString(); }
+		};
+
+		template<> 
+		struct TAttributeTypeTraitsInternal<uint8>
+		{
+			static constexpr EAttributeTypes GetType() { return EAttributeTypes::UInt8; }
+			static FString ToString(const uint8& Value) { return LexToString(Value); }
+		};
+
+		template<> 
+		struct TAttributeTypeTraitsInternal<uint16>
+		{
+			static constexpr EAttributeTypes GetType() { return EAttributeTypes::UInt16; }
+			static FString ToString(const uint16& Value) { return LexToString(Value); }
+		};
+
+		template<> 
+		struct TAttributeTypeTraitsInternal<uint32>
+		{
+			static constexpr EAttributeTypes GetType() { return EAttributeTypes::UInt32; }
+			static FString ToString(const uint32& Value) { return LexToString(Value); }
+		};
+
+		template<> 
+		struct TAttributeTypeTraitsInternal<uint64>
+		{
+			static constexpr EAttributeTypes GetType() { return EAttributeTypes::UInt64; }
+			static FString ToString(const uint64& Value) { return LexToString(Value); }
+		};
+
+		template<> 
+		struct TAttributeTypeTraitsInternal<FVector2D>
+		{
+			static constexpr EAttributeTypes GetType() { return EAttributeTypes::Vector2d; }
+			static FString ToString(const FVector2D& Value) { return Value.ToString(); }
+		};
+
+		template<> 
+		struct TAttributeTypeTraitsInternal<FIntPoint>
+		{
+			static constexpr EAttributeTypes GetType() { return EAttributeTypes::IntPoint; }
+			static FString ToString(const FIntPoint& Value) { return Value.ToString(); }
+		};
+
+		template<> 
+		struct TAttributeTypeTraitsInternal<FIntVector>
+		{
+			static constexpr EAttributeTypes GetType() { return EAttributeTypes::IntVector; }
+			static FString ToString(const FIntVector& Value) { return Value.ToString(); }
+		};
+
+		template<> 
+		struct TAttributeTypeTraitsInternal<FVector2DHalf>
+		{
+			static constexpr EAttributeTypes GetType() { return EAttributeTypes::Vector2DHalf; }
+			static FString ToString(const FVector2DHalf& Value) { return Value.ToString(); }
+		};
+
+		template<> 
+		struct TAttributeTypeTraitsInternal<FFloat16>
+		{
+			static constexpr EAttributeTypes GetType() { return EAttributeTypes::Float16; }
+			static FString ToString(const FFloat16& Value) { return LexToString(Value.GetFloat()); }
+		};
+
+		template<> 
+		struct TAttributeTypeTraitsInternal<FOrientedBox>
+		{
+			static constexpr EAttributeTypes GetType() { return EAttributeTypes::OrientedBox; }
+			static FString ToString(const FOrientedBox& Value) { return TEXT("FOrientedBox"); }
+		};
+
+		template<> 
+		struct TAttributeTypeTraitsInternal<FFrameNumber>
+		{
+			static constexpr EAttributeTypes GetType() { return EAttributeTypes::FrameNumber; }
+			static FString ToString(const FFrameNumber& Value) { return LexToString(Value); }
+		};
+
+		template<> 
+		struct TAttributeTypeTraitsInternal<FFrameRate>
+		{
+			static constexpr EAttributeTypes GetType() { return EAttributeTypes::FrameRate; }
+			static FString ToString(const FFrameRate& Value) { return Value.ToPrettyText().ToString(); }
+		};
+
+		template<> 
+		struct TAttributeTypeTraitsInternal<FFrameTime>
+		{
+			static constexpr EAttributeTypes GetType() { return EAttributeTypes::FrameTime; }
+			static FString ToString(const FFrameTime& Value) { return LexToString(Value.AsDecimal()); }
+		};
+
+		template<> 
+		struct TAttributeTypeTraitsInternal<FSoftObjectPath>
+		{
+			static constexpr EAttributeTypes GetType() { return EAttributeTypes::SoftObjectPath; }
+			static FString ToString(const FSoftObjectPath& Value) { return Value.ToString(); }
+		};
+
+		template<> 
+		struct TAttributeTypeTraitsInternal<FMatrix44f>
+		{
+			static constexpr EAttributeTypes GetType() { return EAttributeTypes::Matrix44f; }
+			static FString ToString(const FMatrix44f& Value) { return Value.ToString(); }
+		};
+
+		template<> 
+		struct TAttributeTypeTraitsInternal<FMatrix44d>
+		{
+			static constexpr EAttributeTypes GetType() { return EAttributeTypes::Matrix44d; }
+			static FString ToString(const FMatrix44d& Value) { return Value.ToString(); }
+		};
+
+		template<> 
+		struct TAttributeTypeTraitsInternal<FPlane4f>
+		{
+			static constexpr EAttributeTypes GetType() { return EAttributeTypes::Plane4f; }
+			static FString ToString(const FPlane4f& Value) { return Value.ToString(); }
+		};
+
+		template<> 
+		struct TAttributeTypeTraitsInternal<FPlane4d>
+		{
+			static constexpr EAttributeTypes GetType() { return EAttributeTypes::Plane4d; }
+			static FString ToString(const FPlane4d& Value) { return Value.ToString(); }
+		};
+
+		template<> 
+		struct TAttributeTypeTraitsInternal<FQuat4f>
+		{
+			static constexpr EAttributeTypes GetType() { return EAttributeTypes::Quat4f; }
+			static FString ToString(const FQuat4f& Value) { return Value.ToString(); }
+		};
+
+		template<> 
+		struct TAttributeTypeTraitsInternal<FQuat4d>
+		{
+			static constexpr EAttributeTypes GetType() { return EAttributeTypes::Quat4d; }
+			static FString ToString(const FQuat4d& Value) { return Value.ToString(); }
+		};
+
+		template<> 
+		struct TAttributeTypeTraitsInternal<FRotator3f>
+		{
+			static constexpr EAttributeTypes GetType() { return EAttributeTypes::Rotator3f; }
+			static FString ToString(const FRotator3f& Value) { return Value.ToString(); }
+		};
+
+		template<> 
+		struct TAttributeTypeTraitsInternal<FRotator3d>
+		{
+			static constexpr EAttributeTypes GetType() { return EAttributeTypes::Rotator3d; }
+			static FString ToString(const FRotator3d& Value) { return Value.ToString(); }
+		};
+
+		template<> 
+		struct TAttributeTypeTraitsInternal<FTransform3f>
+		{
+			static constexpr EAttributeTypes GetType() { return EAttributeTypes::Transform3f; }
+			static FString ToString(const FTransform3f& Value) { return Value.ToString(); }
+		};
+
+		template<> 
+		struct TAttributeTypeTraitsInternal<FTransform3d>
+		{
+			static constexpr EAttributeTypes GetType() { return EAttributeTypes::Transform3d; }
+			static FString ToString(const FTransform3d& Value) { return Value.ToString(); }
+		};
+
+		template<> 
+		struct TAttributeTypeTraitsInternal<FVector3f>
+		{
+			static constexpr EAttributeTypes GetType() { return EAttributeTypes::Vector3f; }
+			static FString ToString(const FVector3f& Value) { return Value.ToString(); }
+		};
+
+		template<> 
+		struct TAttributeTypeTraitsInternal<FVector3d>
+		{
+			static constexpr EAttributeTypes GetType() { return EAttributeTypes::Vector3d; }
+			static FString ToString(const FVector3d& Value) { return Value.ToString(); }
+		};
+
+		template<> 
+		struct TAttributeTypeTraitsInternal<FVector2f>
+		{
+			static constexpr EAttributeTypes GetType() { return EAttributeTypes::Vector2f; }
+			static FString ToString(const FVector2f& Value) { return Value.ToString(); }
+		};
+
+		template<> 
+		struct TAttributeTypeTraitsInternal<FVector4f>
+		{
+			static constexpr EAttributeTypes GetType() { return EAttributeTypes::Vector4f; }
+			static FString ToString(const FVector4f& Value) { return Value.ToString(); }
+		};
+
+		template<> 
+		struct TAttributeTypeTraitsInternal<FVector4d>
+		{
+			static constexpr EAttributeTypes GetType() { return EAttributeTypes::Vector4d; }
+			static FString ToString(const FVector4d& Value) { return Value.ToString(); }
+		};
+
+		template<> 
+		struct TAttributeTypeTraitsInternal<FBox2f>
+		{
+			static constexpr EAttributeTypes GetType() { return EAttributeTypes::Box2f; }
+			static FString ToString(const FBox2f& Value) { return Value.ToString(); }
+		};
+
+		template<> 
+		struct TAttributeTypeTraitsInternal<FBox2D>
+		{
+			static constexpr EAttributeTypes GetType() { return EAttributeTypes::Box2D; }
+			static FString ToString(const FBox2D& Value) { return Value.ToString(); }
+		};
+
+		template<> 
+		struct TAttributeTypeTraitsInternal<FBox3f>
+		{
+			static constexpr EAttributeTypes GetType() { return EAttributeTypes::Box3f; }
+			static FString ToString(const FBox3f& Value) { return Value.ToString(); }
+		};
+
+		template<> 
+		struct TAttributeTypeTraitsInternal<FBox3d>
+		{
+			static constexpr EAttributeTypes GetType() { return EAttributeTypes::Box3d; }
+			static FString ToString(const FBox3d& Value) { return Value.ToString(); }
+		};
+
+		template<> 
+		struct TAttributeTypeTraitsInternal<FBoxSphereBounds3f>
+		{
+			static constexpr EAttributeTypes GetType() { return EAttributeTypes::BoxSphereBounds3f; }
+			static FString ToString(const FBoxSphereBounds3f& Value) { return Value.ToString(); }
+		};
+
+		template<> 
+		struct TAttributeTypeTraitsInternal<FBoxSphereBounds3d>
+		{
+			static constexpr EAttributeTypes GetType() { return EAttributeTypes::BoxSphereBounds3d; }
+			static FString ToString(const FBoxSphereBounds3d& Value) { return Value.ToString(); }
+		};
+
+		template<> 
+		struct TAttributeTypeTraitsInternal<FSphere3f>
+		{
+			static constexpr EAttributeTypes GetType() { return EAttributeTypes::Sphere3f; }
+			static FString ToString(const FSphere3f& Value) { return Value.ToString(); }
+		};
+
+		template<> 
+		struct TAttributeTypeTraitsInternal<FSphere3d>
+		{
+			static constexpr EAttributeTypes GetType() { return EAttributeTypes::Sphere3d; }
+			static FString ToString(const FSphere3d& Value) { return Value.ToString(); }
+		};
+
+		template<> 
+		struct TAttributeTypeTraitsInternal<TArray<uint8>>
+		{
+			static constexpr EAttributeTypes GetType() { return EAttributeTypes::ByteArray; }
+			static FString ToString(const TArray<uint8>& Value) {return TEXT("Array<uint8>"); }
+		};
+
+		template<> 
+		struct TAttributeTypeTraitsInternal<TArray64<uint8>>
+		{
+			static constexpr EAttributeTypes GetType() { return EAttributeTypes::ByteArray64; }
+			static FString ToString(const TArray64<uint8>& Value) { return TEXT("Array64<uint8>"); }
+		};
+
 		FString AttributeTypeToString(EAttributeTypes AttributeType)
 		{
 			FString AttributeTypeString;
@@ -99,6 +534,68 @@
 			return EAttributeTypes::None;
 		}
 
+		template<typename ValueType>
+		FString AttributeValueToString(const ValueType& Value)
+		{
+			return TAttributeTypeTraitsInternal<ValueType>::ToString(Value);
+		}
+
+		FGuid GetValueHash(const uint8* Value, uint64 ValueSize)
+		{
+			FSHA1 Sha;
+			Sha.Update(Value, ValueSize);
+			Sha.Final();
+			// Retrieve the hash and use it to construct a pseudo-GUID. 
+			uint32 Hash[5];
+			Sha.GetHash(reinterpret_cast<uint8*>(Hash));
+			return FGuid(Hash[0] ^ Hash[4], Hash[1], Hash[2], Hash[3]);
+		}
+
+		/** Return the storage size of the value. */
+		template<typename T>
+		uint64 GetValueSizeInternal(const T& Value)
+		{
+			const EAttributeTypes ValueType = TAttributeTypeTraitsInternal<T>::GetType();
+			if (ValueType == EAttributeTypes::None)
+			{
+				return 0;
+			}
+			return sizeof(T);
+		}
+
+		template<>
+		uint64 GetValueSizeInternal(const FString& Value)
+		{
+			//We must add the null character '/0' terminate string
+			return (Value.Len() + 1) * sizeof(TCHAR);
+		}
+
+		template<>
+		uint64 GetValueSizeInternal(const FName& Value)
+		{
+			FString ValueStr = Value.ToString();
+			return GetValueSizeInternal(ValueStr);
+		}
+
+		template<>
+		uint64 GetValueSizeInternal(const FSoftObjectPath& Value)
+		{
+			FString ValueStr = Value.ToString();
+			return GetValueSizeInternal(ValueStr);
+		}
+
+		template<>
+		uint64 GetValueSizeInternal(const TArray<uint8>& Value)
+		{
+			return Value.Num();
+		}
+
+		template<>
+		uint64 GetValueSizeInternal(const TArray64<uint8>& Value)
+		{
+			return Value.Num();
+		}
+
 		void LogAttributeStorageErrors(const EAttributeStorageResult Result, const FString OperationName, const FAttributeKey AttributeKey)
 		{
 			//////////////////////////////////////////////////////////////////////////
@@ -174,6 +671,120 @@
 			return *this;
 		}
 
+		template<typename T>
+		FAttributeStorage::TAttributeHandle<T, std::enable_if_t<TIsNonEnumType<T>>>::TAttributeHandle()
+			: AttributeStorage(nullptr)
+			, Key()
+		{
+			static_assert(TAttributeTypeTraitsInternal<T>::GetType() != EAttributeTypes::None, "Unsupported attribute type");
+		}
+
+		template<typename T>
+		bool FAttributeStorage::TAttributeHandle<T, std::enable_if_t<TIsNonEnumType<T>>>::IsValid() const
+		{
+			const EAttributeTypes ValueType = TAttributeTypeTraitsInternal<T>::GetType();
+			return (AttributeStorage && AttributeStorage->ContainAttribute(Key) && AttributeStorage->GetAttributeType(Key) == ValueType);
+		}
+
+		template<typename T>
+		EAttributeStorageResult FAttributeStorage::TAttributeHandle<T, std::enable_if_t<TIsNonEnumType<T>>>::Get(T& Value) const
+		{
+			if (AttributeStorage)
+			{
+				return AttributeStorage->GetAttribute<T>(Key, Value);
+			}
+			return EAttributeStorageResult::Operation_Error_InvalidStorage;
+		}
+
+		template<typename T>
+		EAttributeStorageResult FAttributeStorage::TAttributeHandle<T, std::enable_if_t<TIsNonEnumType<T>>>::Set(const T& Value)
+		{
+			if (AttributeStorage)
+			{
+				return AttributeStorage->SetAttribute<T>(Key, Value);
+			}
+			return EAttributeStorageResult::Operation_Error_InvalidStorage;
+		}
+		
+		template<typename T>
+		const FAttributeKey& FAttributeStorage::TAttributeHandle<T, std::enable_if_t<TIsNonEnumType<T>>>::GetKey() const
+		{
+			return Key;
+		}
+		
+		template<typename T>
+		FAttributeStorage::TAttributeHandle<T, std::enable_if_t<TIsNonEnumType<T>>>::TAttributeHandle(const FAttributeKey& InKey, const FAttributeStorage* InAttributeStorage)
+		{
+			AttributeStorage = const_cast<FAttributeStorage*>(InAttributeStorage);
+			Key = InKey;
+		
+			//Look for storage validity
+			if (AttributeStorage == nullptr)
+			{
+				//Storage is null
+				LogAttributeStorageErrors(EAttributeStorageResult::Operation_Error_InvalidStorage, TEXT("GetAttributeHandle"), Key);
+			}
+			else
+			{
+
+				if (!AttributeStorage->ContainAttribute(Key))
+				{
+					//Storage do not contain the key
+					LogAttributeStorageErrors(EAttributeStorageResult::Operation_Error_CannotFoundKey, TEXT("GetAttributeHandle"), Key);
+				}
+				if (AttributeStorage->GetAttributeType(Key) != TAttributeTypeTraitsInternal<T>::GetType())
+				{
+					//Value Type is different from the existing key
+					LogAttributeStorageErrors(EAttributeStorageResult::Operation_Error_WrongType, TEXT("GetAttributeHandle"), Key);
+				}
+			}
+		}
+
+		template<typename T, typename>
+		EAttributeStorageResult FAttributeStorage::RegisterAttribute(const FAttributeKey& ElementAttributeKey, const T& DefaultValue, EAttributeProperty AttributeProperty)
+		{
+			static_assert(TAttributeTypeTraitsInternal<T>::GetType() != EAttributeTypes::None, "T is not a supported type for the attributes. Check EAttributeTypes for the supported types");
+
+			//Lock the storage
+			FScopeLock ScopeLock(&StorageMutex);
+
+
+			const EAttributeTypes ValueType = TAttributeTypeTraitsInternal<T>::GetType();
+			const uint64 ValueSize = GetValueSizeInternal(DefaultValue);
+
+			FAttributeAllocationInfo* AttributeAllocationInfo = AttributeAllocationTable.Find(ElementAttributeKey);
+			if (AttributeAllocationInfo)
+			{
+				if (AttributeAllocationInfo->Type != ValueType)
+				{
+					return EAttributeStorageResult::Operation_Error_WrongType;
+				}
+				if (!AttributeStorage.IsValidIndex(AttributeAllocationInfo->Offset))
+				{
+					return EAttributeStorageResult::Operation_Error_AttributeAllocationCorrupted;
+				}
+			}
+			else
+			{
+				AttributeAllocationInfo = &AttributeAllocationTable.Add(ElementAttributeKey);
+				AttributeAllocationInfo->Type = ValueType;
+				AttributeAllocationInfo->Size = ValueSize;
+				AttributeAllocationInfo->Offset = AttributeStorage.AddZeroed(ValueSize);
+			}
+
+			//Force the specified attribute property
+			AttributeAllocationInfo->Property = AttributeProperty;
+
+			//Use template SetAttribute which support all the types
+			const EAttributeStorageResult Result = SetAttribute(AttributeAllocationInfo, DefaultValue);
+			if (!IsAttributeStorageResultSuccess(Result))
+			{
+				//An error occured, unregister the key from the storage
+				UnregisterAttribute(ElementAttributeKey);
+			}
+			return Result;
+		}
+
 		EAttributeStorageResult FAttributeStorage::UnregisterAttribute(const FAttributeKey& ElementAttributeKey)
 		{
 			//Lock the storage
@@ -196,6 +807,187 @@
 			return EAttributeStorageResult::Operation_Success;
 		}
 
+		template<typename T, typename>
+		INTERCHANGECORE_API EAttributeStorageResult FAttributeStorage::GetAttribute(const FAttributeKey& ElementAttributeKey, T& OutValue) const
+		{
+			static_assert(TAttributeTypeTraitsInternal<T>::GetType() != EAttributeTypes::None, "T is not a supported type for the attributes. Check EAttributeTypes for the supported types");
+
+			//Lock the storage
+			FScopeLock ScopeLock(&StorageMutex);
+
+			const EAttributeTypes ValueType = TAttributeTypeTraitsInternal<T>::GetType();
+			const uint64 ValueSize = GetValueSizeInternal(OutValue);
+
+			const FAttributeAllocationInfo* AttributeAllocationInfo = AttributeAllocationTable.Find(ElementAttributeKey);
+			if (AttributeAllocationInfo)
+			{
+				if (AttributeAllocationInfo->Type != ValueType && !((ValueType == EAttributeTypes::UInt8) && (ValueType == EAttributeTypes::Enum)))
+				{
+					return EAttributeStorageResult::Operation_Error_WrongType;
+				}
+				if (AttributeAllocationInfo->Size != ValueSize)
+				{
+					return EAttributeStorageResult::Operation_Error_WrongSize;
+				}
+				if (!AttributeStorage.IsValidIndex(AttributeAllocationInfo->Offset))
+				{
+					return EAttributeStorageResult::Operation_Error_AttributeAllocationCorrupted;
+				}
+			}
+			else
+			{
+				//The key do not exist
+				return EAttributeStorageResult::Operation_Error_CannotFoundKey;
+			}
+
+			const uint8* StorageData = AttributeStorage.GetData();
+			FMemory::Memcpy(&OutValue, &StorageData[AttributeAllocationInfo->Offset], AttributeAllocationInfo->Size);
+			return EAttributeStorageResult::Operation_Success;
+		}
+
+		template<>
+		INTERCHANGECORE_API EAttributeStorageResult FAttributeStorage::GetAttribute(const FAttributeKey& ElementAttributeKey, TArray<uint8>& OutValue) const
+		{
+			return GenericArrayGetAttribute(ElementAttributeKey, OutValue);
+		}
+
+		template<>
+		INTERCHANGECORE_API EAttributeStorageResult FAttributeStorage::GetAttribute(const FAttributeKey& ElementAttributeKey, TArray64<uint8>& OutValue) const
+		{
+			return GenericArrayGetAttribute(ElementAttributeKey, OutValue);
+		}
+
+		template<>
+		INTERCHANGECORE_API EAttributeStorageResult FAttributeStorage::GetAttribute(const FAttributeKey& ElementAttributeKey, FString& OutValue) const
+		{
+			static_assert(TAttributeTypeTraitsInternal<FString>::GetType() != EAttributeTypes::None, "Not a supported type for the attributes. Check EAttributeTypes for the supported types.");
+
+			//Lock the storage
+			FScopeLock ScopeLock(&StorageMutex);
+
+			const EAttributeTypes ValueType = TAttributeTypeTraitsInternal<FString>::GetType();
+			const FAttributeAllocationInfo* AttributeAllocationInfo = AttributeAllocationTable.Find(ElementAttributeKey);
+			if (AttributeAllocationInfo)
+			{
+				if (AttributeAllocationInfo->Type != ValueType)
+				{
+					return EAttributeStorageResult::Operation_Error_WrongType;
+				}
+				if (!AttributeStorage.IsValidIndex(AttributeAllocationInfo->Offset))
+				{
+					return EAttributeStorageResult::Operation_Error_AttributeAllocationCorrupted;
+				}
+			}
+			else
+			{
+				//The key do not exist
+				return EAttributeStorageResult::Operation_Error_CannotFoundKey;
+			}
+
+			if (AttributeAllocationInfo->Size <= sizeof(TCHAR)) // Account for the null-terminator
+			{
+				OutValue = TEXT("");
+				return EAttributeStorageResult::Operation_Success;
+			}
+
+			const FStringView ValueView = GetFStringViewAttributeFromStorage(AttributeStorage.GetData(), AttributeAllocationInfo);
+			OutValue = ValueView;
+
+			return EAttributeStorageResult::Operation_Success;
+		}
+
+		template<>
+		INTERCHANGECORE_API EAttributeStorageResult FAttributeStorage::GetAttribute(const FAttributeKey& ElementAttributeKey, FName& OutValue) const
+		{
+			static_assert(TAttributeTypeTraitsInternal<FName>::GetType() != EAttributeTypes::None, "Not a supported type for the attributes. Check EAttributeTypes for supported types.");
+
+			//Lock the storage
+			FScopeLock ScopeLock(&StorageMutex);
+
+			const EAttributeTypes ValueType = TAttributeTypeTraitsInternal<FName>::GetType();
+
+
+			const FAttributeAllocationInfo* AttributeAllocationInfo = AttributeAllocationTable.Find(ElementAttributeKey);
+			if (AttributeAllocationInfo)
+			{
+				if (AttributeAllocationInfo->Type != ValueType)
+				{
+					return EAttributeStorageResult::Operation_Error_WrongType;
+				}
+				if (!AttributeStorage.IsValidIndex(AttributeAllocationInfo->Offset))
+				{
+					return EAttributeStorageResult::Operation_Error_AttributeAllocationCorrupted;
+				}
+			}
+			else
+			{
+				//The key do not exist
+				return EAttributeStorageResult::Operation_Error_CannotFoundKey;
+			}
+
+
+			if (AttributeAllocationInfo->Size == 0)
+			{
+				OutValue = NAME_None;
+				return EAttributeStorageResult::Operation_Success;
+			}
+
+			const FStringView ValueStr = GetFStringViewAttributeFromStorage(AttributeStorage.GetData(), AttributeAllocationInfo);
+			OutValue = FName(ValueStr);
+
+			return EAttributeStorageResult::Operation_Success;
+		}
+
+		template<>
+		INTERCHANGECORE_API EAttributeStorageResult FAttributeStorage::GetAttribute(const FAttributeKey& ElementAttributeKey, FSoftObjectPath& OutValue) const
+		{
+			static_assert(TAttributeTypeTraitsInternal<FSoftObjectPath>::GetType() != EAttributeTypes::None, "Not a supported type for the attributes. Check EAttributeTypes for supported types.");
+
+			//Lock the storage
+			FScopeLock ScopeLock(&StorageMutex);
+
+			const EAttributeTypes ValueType = TAttributeTypeTraitsInternal<FSoftObjectPath>::GetType();
+
+
+			const FAttributeAllocationInfo* AttributeAllocationInfo = AttributeAllocationTable.Find(ElementAttributeKey);
+			if (AttributeAllocationInfo)
+			{
+				if (AttributeAllocationInfo->Type != ValueType)
+				{
+					return EAttributeStorageResult::Operation_Error_WrongType;
+				}
+				if (!AttributeStorage.IsValidIndex(AttributeAllocationInfo->Offset))
+				{
+					return EAttributeStorageResult::Operation_Error_AttributeAllocationCorrupted;
+				}
+			}
+			else
+			{
+				//The key do not exist
+				return EAttributeStorageResult::Operation_Error_CannotFoundKey;
+			}
+
+
+			if (AttributeAllocationInfo->Size == 0)
+			{
+				OutValue = TEXT("");
+				return EAttributeStorageResult::Operation_Success;
+			}
+
+			const FStringView ValueStr = GetFStringViewAttributeFromStorage(AttributeStorage.GetData(), AttributeAllocationInfo);
+			OutValue = FSoftObjectPath(ValueStr);
+
+			return EAttributeStorageResult::Operation_Success;
+		}
+
+
+		template<typename T, typename>
+		FAttributeStorage::TAttributeHandle<T> FAttributeStorage::GetAttributeHandle(const FAttributeKey& ElementAttributeKey) const
+		{
+			TAttributeHandle<T> AttributeHandle(ElementAttributeKey, this);
+			return AttributeHandle;
+		}
+
 		EAttributeTypes FAttributeStorage::GetAttributeType(const FAttributeKey& ElementAttributeKey) const
 		{
 			//Lock the storage
@@ -436,6 +1228,13 @@
 			}
 		}
 
+		template<typename T>
+		EAttributeStorageResult FAttributeStorage::SetAttribute(const FAttributeKey& ElementAttributeKey, const T& Value)
+		{
+			FAttributeAllocationInfo* AttributeAllocationInfo = AttributeAllocationTable.Find(ElementAttributeKey);
+			return SetAttribute(AttributeAllocationInfo, Value);
+		}
+
 		void FAttributeStorage::ExtractFStringAttributeFromStorage(const uint8* StorageData, const FAttributeAllocationInfo* AttributeAllocationInfo, FString& OutValue) const
 		{
 			//Allocate the ByteArray
@@ -478,99 +1277,92 @@
 			return FStringView(StringContents, int32(NumberOfChar - 1));
 		}
 
-		EAttributeStorageResult FAttributeStorage::GetAttribute(const FAttributeKey& ElementAttributeKey, FString& OutValue, TSpecializeType<FString >) const
+		/** Set a multisize(TArray<uint8>, FString ...) attribute value into the storage. Return success if the attribute was properly set. */
+		template<typename MultiSizeType>
+		EAttributeStorageResult FAttributeStorage::MultiSizeSetAttribute(FAttributeAllocationInfo* AttributeAllocationInfo, const MultiSizeType& Value, TFunctionRef<uint8* ()> GetValuePointer)
 		{
-			static_assert(TAttributeTypeTraits<FString>::GetType() != EAttributeTypes::None, "Not a supported type for the attributes. Check EAttributeTypes for the supported types.");
-
 			//Lock the storage
 			FScopeLock ScopeLock(&StorageMutex);
 
-			const EAttributeTypes ValueType = TAttributeTypeTraits<FString>::GetType();
-			const FAttributeAllocationInfo* AttributeAllocationInfo = AttributeAllocationTable.Find(ElementAttributeKey);
-			if (AttributeAllocationInfo)
+			const EAttributeTypes ValueType = TAttributeTypeTraitsInternal<MultiSizeType>::GetType();
+
+			if (!AttributeAllocationInfo)
 			{
-				if (AttributeAllocationInfo->Type != ValueType)
+				return EAttributeStorageResult::Operation_Error_CannotFoundKey;
+			}
+
+			if (AttributeAllocationInfo->Type != ValueType)
+			{
+				return EAttributeStorageResult::Operation_Error_WrongType;
+			}
+
+			const uint64 ValueSize = GetValueSizeInternal(Value);
+
+			uint8* DataPtr = GetValuePointer();
+			//We must have a valid data pointer if we have something to copy into the storage.
+			//This case support empty array or empty string because the ValueSize will be 0.
+			if (!DataPtr && ValueSize > 0)
+			{
+				return EAttributeStorageResult::Operation_Error_InvalidMultiSizeValueData;
+			}
+
+			//If the new size is greater then the old size we have to create a new allocation info entry and delete the old one
+			if (ValueSize > AttributeAllocationInfo->Size)
+			{
+				//Find the key to remove it
+				const FAttributeKey* ElementAttributeKeyPtr = AttributeAllocationTable.FindKey(*AttributeAllocationInfo);
+				if (!ElementAttributeKeyPtr)
 				{
-					return EAttributeStorageResult::Operation_Error_WrongType;
+					return EAttributeStorageResult::Operation_Error_CannotFoundKey;
 				}
-				if (!AttributeStorage.IsValidIndex(AttributeAllocationInfo->Offset))
+				const FAttributeKey OriginalKey = *ElementAttributeKeyPtr;
+				const EAttributeProperty AttributePropertyBackup = AttributeAllocationInfo->Property;
+				const EAttributeStorageResult RemoveResult = UnregisterAttribute(OriginalKey);
+				if (!IsAttributeStorageResultSuccess(RemoveResult))
 				{
-					return EAttributeStorageResult::Operation_Error_AttributeAllocationCorrupted;
+					return RemoveResult;
 				}
+				AttributeAllocationInfo = nullptr;
+				//Reallocate a new one
+				AttributeAllocationInfo = &AttributeAllocationTable.Add(OriginalKey);
+				AttributeAllocationInfo->Type = ValueType;
+				AttributeAllocationInfo->Offset = AttributeStorage.AddZeroed(ValueSize);
+				//Force the specified attribute property
+				AttributeAllocationInfo->Property = AttributePropertyBackup;
 			}
 			else
 			{
-				//The key do not exist
-				return EAttributeStorageResult::Operation_Error_CannotFoundKey;
+				//In case we reuse the allocation table, simply adjust the waste memory counter
+				FragmentedMemoryCost += AttributeAllocationInfo->Size - ValueSize;
 			}
 
-			if (AttributeAllocationInfo->Size <= sizeof(TCHAR)) // Account for the null-terminator
+			//We can recycle allocation table if the size is equal or smaller, so we have to always set it
+			AttributeAllocationInfo->Size = ValueSize;
+
+			//Set the hash from the value
+			AttributeAllocationInfo->Hash = GetValueHash(DataPtr, ValueSize);
+
+			if (ValueSize > 0)
 			{
-				OutValue = TEXT("");
-				return EAttributeStorageResult::Operation_Success;
+				uint8* StorageData = AttributeStorage.GetData();
+				FMemory::Memcpy(&StorageData[AttributeAllocationInfo->Offset], DataPtr, ValueSize);
 			}
-
-			const FStringView ValueView = GetFStringViewAttributeFromStorage(AttributeStorage.GetData(), AttributeAllocationInfo);
-			OutValue = ValueView;
 
 			return EAttributeStorageResult::Operation_Success;
 		}
 
-		EAttributeStorageResult FAttributeStorage::GetAttribute(const FAttributeKey& ElementAttributeKey, FName& OutValue, TSpecializeType<FName >) const
+		template<typename ArrayType>
+		EAttributeStorageResult FAttributeStorage::GenericArrayGetAttribute(const FAttributeKey& ElementAttributeKey, ArrayType& OutValue) const
 		{
-			static_assert(TAttributeTypeTraits<FName>::GetType() != EAttributeTypes::None, "Not a supported type for the attributes. Check EAttributeTypes for supported types.");
-
 			//Lock the storage
 			FScopeLock ScopeLock(&StorageMutex);
 
-			const EAttributeTypes ValueType = TAttributeTypeTraits<FName>::GetType();
-
+			static_assert(TIsTArray<ArrayType>::Value);
 
 			const FAttributeAllocationInfo* AttributeAllocationInfo = AttributeAllocationTable.Find(ElementAttributeKey);
 			if (AttributeAllocationInfo)
 			{
-				if (AttributeAllocationInfo->Type != ValueType)
+				if (AttributeAllocationInfo->Type != TAttributeTypeTraitsInternal<ArrayType>::GetType())
 				{
 					return EAttributeStorageResult::Operation_Error_WrongType;
 				}
-				if (!AttributeStorage.IsValidIndex(AttributeAllocationInfo->Offset))
-				{
-					return EAttributeStorageResult::Operation_Error_AttributeAllocationCorrupted;
-				}
-			}
-			else
-			{
-				//The key do not exist
-				return EAttributeStorageResult::Operation_Error_CannotFoundKey;
-			}
-
-
-			if (AttributeAllocationInfo->Size == 0)
-			{
-				OutValue = NAME_None;
-				return EAttributeStorageResult::Operation_Success;
-			}
-
-			const FStringView ValueStr = GetFStringViewAttributeFromStorage(AttributeStorage.GetData(), AttributeAllocationInfo);
-			OutValue = FName(ValueStr);
-
-			return EAttributeStorageResult::Operation_Success;
-		}
-
-		EAttributeStorageResult FAttributeStorage::GetAttribute(const FAttributeKey& ElementAttributeKey, FSoftObjectPath& OutValue, TSpecializeType<FSoftObjectPath >) const
-		{
-			static_assert(TAttributeTypeTraits<FSoftObjectPath>::GetType() != EAttributeTypes::None, "Not a supported type for the attributes. Check EAttributeTypes for supported types.");
-
-			//Lock the storage
-			FScopeLock ScopeLock(&StorageMutex);
-
-			const EAttributeTypes ValueType = TAttributeTypeTraits<FSoftObjectPath>::GetType();
-
-
-			const FAttributeAllocationInfo* AttributeAllocationInfo = AttributeAllocationTable.Find(ElementAttributeKey);
-			if (AttributeAllocationInfo)
-			{
-				if (AttributeAllocationInfo->Type != ValueType)
-				{
-					return EAttributeStorageResult::Operation_Error_WrongType;
-				}
@@ -585,16 +1377,17 @@
 				return EAttributeStorageResult::Operation_Error_CannotFoundKey;
 			}
 
-
 			if (AttributeAllocationInfo->Size == 0)
 			{
-				OutValue = TEXT("");
+				OutValue.Empty();
 				return EAttributeStorageResult::Operation_Success;
 			}
 
-			const FStringView ValueStr = GetFStringViewAttributeFromStorage(AttributeStorage.GetData(), AttributeAllocationInfo);
-			OutValue = FSoftObjectPath(ValueStr);
-
+			//Allocate the ByteArray
+			OutValue.AddZeroed(static_cast<typename ArrayType::SizeType>(AttributeAllocationInfo->Size));
+			//Copy into the buffer
+			const uint8* StorageData = AttributeStorage.GetData();
+			FMemory::Memcpy(OutValue.GetData(), &StorageData[AttributeAllocationInfo->Offset], AttributeAllocationInfo->Size);
 			return EAttributeStorageResult::Operation_Success;
 		}
 
@@ -684,32 +1477,161 @@
 			FragmentedMemoryCost = 0;
 		}
 
-		FGuid FAttributeStorage::GetValueHash(const uint8* Value, uint64 ValueSize)
+		template<typename T>
+		EAttributeStorageResult FAttributeStorage::SetAttribute(FAttributeAllocationInfo* AttributeAllocationInfo, const T& Value)
 		{
-			FSHA1 Sha;
-			Sha.Update(Value, ValueSize);
-			Sha.Final();
-			// Retrieve the hash and use it to construct a pseudo-GUID. 
-			uint32 Hash[5];
-			Sha.GetHash(reinterpret_cast<uint8*>(Hash));
-			return FGuid(Hash[0] ^ Hash[4], Hash[1], Hash[2], Hash[3]);
+			static_assert(TAttributeTypeTraitsInternal<T>::GetType() != EAttributeTypes::None, "T is not a supported type for the attributes. Check EAttributeTypes for the supported types");
+
+			//Lock the storage
+			FScopeLock ScopeLock(&StorageMutex);
+
+			if (!AttributeAllocationInfo)
+			{
+				return EAttributeStorageResult::Operation_Error_CannotFoundKey;
+			}
+			const EAttributeTypes ValueType = TAttributeTypeTraitsInternal<T>::GetType();
+			const uint64 ValueSize = GetValueSizeInternal(Value);
+
+			if (AttributeAllocationInfo->Type != ValueType)
+			{
+				return EAttributeStorageResult::Operation_Error_WrongType;
+			}
+			if (AttributeAllocationInfo->Size != ValueSize)
+			{
+				return EAttributeStorageResult::Operation_Error_WrongSize;
+			}
+			if (!AttributeStorage.IsValidIndex(AttributeAllocationInfo->Offset))
+			{
+				return EAttributeStorageResult::Operation_Error_AttributeAllocationCorrupted;
+			}
+
+			//Set the hash from the value
+			AttributeAllocationInfo->Hash = GetValueHash(reinterpret_cast<const uint8*>(&Value), ValueSize);
+
+			uint8* StorageData = AttributeStorage.GetData();
+			FMemory::Memcpy(&StorageData[AttributeAllocationInfo->Offset], &Value, ValueSize);
+			return EAttributeStorageResult::Operation_Success;
 		}
 
-		uint64 FAttributeStorage::GetValueSize(const FString& Value, TSpecializeType<FString >)
+		template<>
+		EAttributeStorageResult FAttributeStorage::SetAttribute(FAttributeAllocationInfo* AttributeAllocationInfo, const TArray<uint8>& Value)
 		{
-			//We must add the null character '/0' terminate string
-			return (Value.Len() + 1) * sizeof(TCHAR);
+			return MultiSizeSetAttribute(AttributeAllocationInfo, Value, [&Value]()->uint8*
+			{
+				return const_cast<uint8*>(Value.GetData());
+			});
 		}
 
-		uint64 FAttributeStorage::GetValueSize(const TArray<uint8>& Value, TSpecializeType<TArray<uint8> >)
+		template<>
+		EAttributeStorageResult FAttributeStorage::SetAttribute(FAttributeAllocationInfo* AttributeAllocationInfo, const TArray64<uint8>& Value)
 		{
-			return Value.Num();
+			return MultiSizeSetAttribute(AttributeAllocationInfo, Value, [&Value]()->uint8*
+			{
+				return const_cast<uint8*>(Value.GetData());
+			});
 		}
 
-		uint64 FAttributeStorage::GetValueSize(const TArray64<uint8>& Value, TSpecializeType<TArray64<uint8> >)
+		template<>
+		EAttributeStorageResult FAttributeStorage::SetAttribute(FAttributeAllocationInfo* AttributeAllocationInfo, const FString& Value)
 		{
-			return Value.Num();
+			return MultiSizeSetAttribute(AttributeAllocationInfo, Value, [&Value]()->uint8*
+			{
+				return (uint8*)(*Value);
+			});
 		}
 
+		template<>
+		EAttributeStorageResult FAttributeStorage::SetAttribute(FAttributeAllocationInfo* AttributeAllocationInfo, const FName& Value)
+		{
+			//FName must be store has a FString for persistence
+			FString ValueStr = Value.ToString();
+			return MultiSizeSetAttribute(AttributeAllocationInfo, Value, [&ValueStr]()->uint8*
+			{
+				return (uint8*)(*ValueStr);
+			});
+		}
+
+		template<>
+		EAttributeStorageResult FAttributeStorage::SetAttribute(FAttributeAllocationInfo* AttributeAllocationInfo, const FSoftObjectPath& Value)
+		{
+			//FSoftObjectPath can be hold has a FString and restore from it
+			FString ValueStr = Value.ToString();
+			return MultiSizeSetAttribute(AttributeAllocationInfo, Value, [&ValueStr]()->uint8*
+			{
+				return (uint8*)(*ValueStr);
+			});
+		}
+
+		#define SPECIALIZE_SINGLE(X) \
+			template INTERCHANGECORE_API EAttributeStorageResult FAttributeStorage::RegisterAttribute<X>(const FAttributeKey&, const X&, EAttributeProperty);\
+			template INTERCHANGECORE_API EAttributeStorageResult FAttributeStorage::GetAttribute<X>(const FAttributeKey&, X&) const;\
+			template class INTERCHANGECORE_API FAttributeStorage::TAttributeHandle<X>;\
+			template INTERCHANGECORE_API FAttributeStorage::TAttributeHandle<X> FAttributeStorage::GetAttributeHandle<X>(const FAttributeKey&) const;\
+			template INTERCHANGECORE_API FString UE::Interchange::AttributeValueToString(const X&);
+		SPECIALIZE_SINGLE(bool);
+		SPECIALIZE_SINGLE(FColor);
+		SPECIALIZE_SINGLE(FDateTime);
+		SPECIALIZE_SINGLE(double);
+		SPECIALIZE_SINGLE(float);
+		SPECIALIZE_SINGLE(FGuid);
+		SPECIALIZE_SINGLE(int8);
+		SPECIALIZE_SINGLE(int16);
+		SPECIALIZE_SINGLE(int32);
+		SPECIALIZE_SINGLE(int64);
+		SPECIALIZE_SINGLE(FIntRect);
+		SPECIALIZE_SINGLE(FLinearColor);
+		SPECIALIZE_SINGLE(FRandomStream);
+		SPECIALIZE_SINGLE(FTimespan);
+		SPECIALIZE_SINGLE(FTwoVectors);
+		SPECIALIZE_SINGLE(uint8);
+		SPECIALIZE_SINGLE(uint16);
+		SPECIALIZE_SINGLE(uint32);
+		SPECIALIZE_SINGLE(uint64);
+		SPECIALIZE_SINGLE(FVector2D);
+		SPECIALIZE_SINGLE(FIntPoint);
+		SPECIALIZE_SINGLE(FIntVector);
+		SPECIALIZE_SINGLE(FVector2DHalf);
+		SPECIALIZE_SINGLE(FFloat16);
+		SPECIALIZE_SINGLE(FOrientedBox);
+		SPECIALIZE_SINGLE(FFrameNumber);
+		SPECIALIZE_SINGLE(FFrameRate);
+		SPECIALIZE_SINGLE(FFrameTime);
+		SPECIALIZE_SINGLE(FMatrix44f);
+		SPECIALIZE_SINGLE(FMatrix44d);
+		SPECIALIZE_SINGLE(FPlane4f);
+		SPECIALIZE_SINGLE(FPlane4d);
+		SPECIALIZE_SINGLE(FQuat4f);
+		SPECIALIZE_SINGLE(FQuat4d);
+		SPECIALIZE_SINGLE(FRotator3f);
+		SPECIALIZE_SINGLE(FRotator3d);
+		SPECIALIZE_SINGLE(FTransform3f);
+		SPECIALIZE_SINGLE(FTransform3d);
+		SPECIALIZE_SINGLE(FVector3f);
+		SPECIALIZE_SINGLE(FVector3d);
+		SPECIALIZE_SINGLE(FVector2f);
+		SPECIALIZE_SINGLE(FVector4f);
+		SPECIALIZE_SINGLE(FVector4d);
+		SPECIALIZE_SINGLE(FBox2f);
+		SPECIALIZE_SINGLE(FBox2D);
+		SPECIALIZE_SINGLE(FBox3f);
+		SPECIALIZE_SINGLE(FBox3d);
+		SPECIALIZE_SINGLE(FBoxSphereBounds3f);
+		SPECIALIZE_SINGLE(FBoxSphereBounds3d);
+		SPECIALIZE_SINGLE(FSphere3f);
+		SPECIALIZE_SINGLE(FSphere3d);
+		#undef SPECIALIZE_SINGLE
+
+		#define SPECIALIZE_NOGETTER(X) \
+			template INTERCHANGECORE_API EAttributeStorageResult FAttributeStorage::RegisterAttribute<X>(const FAttributeKey&, const X&, EAttributeProperty);\
+			template class INTERCHANGECORE_API FAttributeStorage::TAttributeHandle<X>;\
+			template INTERCHANGECORE_API FAttributeStorage::TAttributeHandle<X> FAttributeStorage::GetAttributeHandle<X>(const FAttributeKey&) const;\
+			template INTERCHANGECORE_API FString UE::Interchange::AttributeValueToString(const X&);
+		SPECIALIZE_NOGETTER(FName);
+		SPECIALIZE_NOGETTER(FString);
+		SPECIALIZE_NOGETTER(FSoftObjectPath);
+		SPECIALIZE_NOGETTER(TArray<uint8>);
+		SPECIALIZE_NOGETTER(TArray64<uint8>);
+		#undef SPECIALIZE_NOGETTER
+
 	} //ns Interchange
 } //ns UE
diff --git a/Engine/Source/Runtime/Interchange/Core/Public/Nodes/InterchangeBaseNodeUtilities.h b/Engine/Source/Runtime/Interchange/Core/Public/Nodes/InterchangeBaseNodeUtilities.h
--- a/Engine/Source/Runtime/Interchange/Core/Public/Nodes/InterchangeBaseNodeUtilities.h
+++ b/Engine/Source/Runtime/Interchange/Core/Public/Nodes/InterchangeBaseNodeUtilities.h
@@ -20,8 +20,6 @@ namespace Interchange
 template<typename ItemType>
 class TArrayAttributeHelper
 {
-	static_assert(TAttributeTypeTraits<ItemType>::GetType() != EAttributeTypes::None, "The value type must be supported by the attribute storage");
-
 public:
 	~TArrayAttributeHelper()
 	{
@@ -294,9 +292,6 @@ private:
 template<typename KeyType, typename ValueType>
 class TMapAttributeHelper
 {
-	static_assert(TAttributeTypeTraits<KeyType>::GetType() != EAttributeTypes::None, "The key type must be supported by the attribute storage");
-	static_assert(TAttributeTypeTraits<ValueType>::GetType() != EAttributeTypes::None, "The value type must be supported by the attribute storage");
-
 	template<typename T>
 	using TAttributeHandle = FAttributeStorage::TAttributeHandle<T>;
 
diff --git a/Engine/Source/Runtime/Interchange/Core/Public/Types/AttributeStorage.h b/Engine/Source/Runtime/Interchange/Core/Public/Types/AttributeStorage.h
--- a/Engine/Source/Runtime/Interchange/Core/Public/Types/AttributeStorage.h
+++ b/Engine/Source/Runtime/Interchange/Core/Public/Types/AttributeStorage.h
@@ -5,40 +5,16 @@
 #include "Containers/Array.h"
 #include "Containers/ContainersFwd.h"
 #include "Containers/Map.h"
-#include "Containers/SortedMap.h"
+
+#include "Containers/EnumAsByte.h"
 #include "Containers/UnrealString.h"
-#include "CoreMinimal.h"
 #include "HAL/CriticalSection.h"
-#include "HAL/UnrealMemory.h"
 #include "Internationalization/Text.h"
-#include "Math/Box.h"
-#include "Math/Box2D.h"
-#include "Math/BoxSphereBounds.h"
-#include "Math/Color.h"
-#include "Math/Float16.h"
-#include "Math/IntPoint.h"
-#include "Math/IntRect.h"
-#include "Math/IntVector.h"
-#include "Math/Matrix.h"
-#include "Math/Plane.h"
-#include "Math/Quat.h"
-#include "Math/RandomStream.h"
-#include "Math/Rotator.h"
-#include "Math/Sphere.h"
 #include "Math/Transform.h"
-#include "Math/TwoVectors.h"
-#include "Math/UnrealMathSSE.h"
-#include "Math/Vector.h"
-#include "Math/Vector2D.h"
-#include "Math/Vector2DHalf.h"
-#include "Math/Vector4.h"
 #include "Misc/DateTime.h"
 #include "Misc/EnumClassFlags.h"
 #include "Misc/FrameNumber.h"
-#include "Misc/FrameRate.h"
 #include "Misc/Guid.h"
-#include "Misc/ScopeLock.h"
-#include "Misc/StringFormatArg.h"
 #include "Misc/Timespan.h"
 #include "Serialization/Archive.h"
 #include "Templates/Function.h"
@@ -46,8 +22,6 @@
 #include "UObject/NameTypes.h"
 #include "UObject/SoftObjectPath.h"
 
-#include <type_traits>
-
 struct FFrameTime;
 struct FOrientedBox;
 template <class TEnum> class TEnumAsByte;
@@ -59,14 +33,11 @@ namespace UE
 	{
 		namespace Private
 		{
-			/**
-			 * Wrapper around std::underlying_type to make sure we don't call it on non enum class types.
-			 */
 			template <typename T, typename Enable = void>
-			struct TUnderlyingType {};
+			struct UE_DEPRECATED(5.7, "No longer used.") TUnderlyingType {};
 
 			template <typename T>
-			struct TUnderlyingType<T, typename std::enable_if<std::is_enum<T>::value>::type>
+			struct UE_DEPRECATED(5.7, "No longer used.") TUnderlyingType<T, typename std::enable_if<std::is_enum<T>::value>::type>
 			{
 				using type = typename std::underlying_type<T>::type;
 			};
@@ -277,24 +248,11 @@ namespace UE
 		 */
 		INTERCHANGECORE_API EAttributeTypes StringToAttributeType(const FString& AttributeTypeString);
 
-		/**
-		 * Stub for attribute type traits.
-		 *
-		 * Actual type traits need to be declared through template specialization for custom
-		 * data types that are to be used internally by FAttributeStorage. Traits for the most commonly used built-in
-		 * types are declared below.
-		 *
-		 * Complex types, such as structures and classes, can be serialized into a byte array
-		 * and then assigned to an attribute. Note that you will be responsible for ensuring
-		 * correct byte ordering when serializing those types.
-		 *
-		 * @param T The type to be used in FAttributeStorage.
-		 */
-		template<typename T, typename Enable = void> struct TAttributeTypeTraits
+		template<typename T, typename Enable = void>
+		struct UE_DEPRECATED(5.7, "No longer used: The AttributeStorage handles all supported type traits internally now.") TAttributeTypeTraits
 		{
 			static constexpr EAttributeTypes GetType()
 			{
-				static_assert(!sizeof(T), "Attribute type trait must be specialized for this type.");
 				return EAttributeTypes::None;
 			}
 
@@ -306,10 +264,7 @@ namespace UE
 		};
 
 		template<typename ValueType>
-		FString AttributeValueToString(const ValueType& Value)
-		{
-			return TAttributeTypeTraits<ValueType>::ToString(Value);
-		}
+		INTERCHANGECORE_API FString AttributeValueToString(const ValueType& Value);
 
 		/**
 		 * Enum to return complete status of a storage operation.
@@ -411,16 +366,15 @@ namespace UE
 			return ((PropertyA & PropertyB) != EAttributeProperty::None);
 		}
 
-		/**
-		 * This is a helper class to specialize template functions inside a class. It allows specialized templates to not be static.
-		 */
 		template<typename T>
-		struct TSpecializeType
+		struct UE_DEPRECATED(5.7, "No longer used.") TSpecializeType
 		{
 			typedef T Type;
 		};
 
-	
+		template<typename T>
+		inline constexpr bool TIsNonEnumType = !TIsEnum<T>::Value && !TIsTEnumAsByte<T>::Value;
+
 		/**
 		 * This class is a Key/Value storage inside a TArray64<uint8>.
 		 * The keys are of type FAttributeKey, which is an FString. Each key is unique and has only one value.
@@ -433,87 +387,101 @@ namespace UE
 		class FAttributeStorage
 		{
 		public:
+			template<typename T, typename Enable = void>
+			class TAttributeHandle;
+
 			/**
 			 * Class to get/set an attribute of the storage.
 			 */
 			template<typename T>
-			class TAttributeHandle
+			class INTERCHANGECORE_API TAttributeHandle<T, std::enable_if_t<TIsNonEnumType<T>>>
 			{
-				static_assert(TAttributeTypeTraits<T>::GetType() != EAttributeTypes::None, "Unsupported attribute type");
-
 			public:
-				TAttributeHandle()
-				: AttributeStorage(nullptr)
-				, Key()
-				{}
+				TAttributeHandle();
 
 				/**
 				* Return true if the storage contains a valid attribute key, or false otherwise. 
 				*/
-				inline bool IsValid() const
-				{
-					const EAttributeTypes ValueType = TAttributeTypeTraits<T>::GetType();
-					return (AttributeStorage && AttributeStorage->ContainAttribute(Key) && AttributeStorage->GetAttributeType(Key) == ValueType);
-				}
+				bool IsValid() const;
+			
+				EAttributeStorageResult Get(T& Value) const;
+
+				EAttributeStorageResult Set(const T& Value);
+
+				const FAttributeKey& GetKey() const;
+
+			protected:
+				class FAttributeStorage* AttributeStorage;
+				FAttributeKey Key;
 			
-				EAttributeStorageResult Get(T& Value) const
+				TAttributeHandle(const FAttributeKey& InKey, const FAttributeStorage* InAttributeStorage);
+
+				friend FAttributeStorage;
+			};
+
+			// Redirect enum instantiations to underlying type
+			template<typename E>
+			class TAttributeHandle<E, std::enable_if_t<TIsEnum<E>::Value>>
+				: public TAttributeHandle<std::underlying_type_t<E>>
+			{
+			public:
+				using Base = TAttributeHandle<std::underlying_type_t<E>>;
+				using UnderlyingType = std::underlying_type_t<E>;
+
+				using Base::Base;
+
+				explicit TAttributeHandle(const TAttributeHandle<UnderlyingType>& Other)
+					: TAttributeHandle(Other.Key, Other.AttributeStorage)
 				{
-					if (AttributeStorage)
-					{
-						return AttributeStorage->GetAttribute<T>(Key, Value);
-					}
-					return EAttributeStorageResult::Operation_Error_InvalidStorage;
 				}
 
-				EAttributeStorageResult Set(const T& Value)
+				// Note: Since we derive TAttributeHandle<std::underlying_type_t<E>> we also get the
+				// base class overloads for free, and can Get()/Set() with the underlying type directly!
+				EAttributeStorageResult Get(E& Value) const
 				{
-					if (AttributeStorage)
-					{
-						return AttributeStorage->SetAttribute<T>(Key, Value);
-					}
-					return EAttributeStorageResult::Operation_Error_InvalidStorage;
+					return Base::Get((UnderlyingType&)Value);
 				}
 
-				const FAttributeKey& GetKey() const
+				EAttributeStorageResult Set(const E& Value)
 				{
-					return Key;
+					return Base::Set((const UnderlyingType&)Value);
 				}
 
 			protected:
-				class FAttributeStorage* AttributeStorage;
-				FAttributeKey Key;
-			
-				TAttributeHandle(const FAttributeKey& InKey, const FAttributeStorage* InAttributeStorage)
+				friend FAttributeStorage;
+			};
+
+			// Redirect TEnumAsByte instantiations to underlying type
+			template<typename B>
+			class TAttributeHandle<B, std::enable_if_t<TIsTEnumAsByte<B>::Value>>
+				: public TAttributeHandle<uint8>
+			{
+			public:
+				using Base = TAttributeHandle<uint8>;
+				using UnderlyingType = uint8;
+
+				using Base::Base;
+
+				explicit TAttributeHandle(const TAttributeHandle<UnderlyingType>& Other)
+					: TAttributeHandle(Other.Key, Other.AttributeStorage)
 				{
-					AttributeStorage = const_cast<FAttributeStorage*>(InAttributeStorage);
-					Key = InKey;
-				
-					//Look for storage validity
-					if (AttributeStorage == nullptr)
-					{
-						//Storage is null
-						LogAttributeStorageErrors(EAttributeStorageResult::Operation_Error_InvalidStorage, TEXT("GetAttributeHandle"), Key);
-					}
-					else
-					{
-
-						if (!AttributeStorage->ContainAttribute(Key))
-						{
-							//Storage do not contain the key
-							LogAttributeStorageErrors(EAttributeStorageResult::Operation_Error_CannotFoundKey, TEXT("GetAttributeHandle"), Key);
-						}
-						if (AttributeStorage->GetAttributeType(Key) != TAttributeTypeTraits<T>::GetType())
-						{
-							//Value Type is different from the existing key
-							LogAttributeStorageErrors(EAttributeStorageResult::Operation_Error_WrongType, TEXT("GetAttributeHandle"), Key);
-						}
-					}
 				}
 
+				EAttributeStorageResult Get(B& Value) const
+				{
+					return Base::Get((UnderlyingType&)Value);
+				}
+
+				EAttributeStorageResult Set(const B& Value)
+				{
+					return Base::Set((const UnderlyingType&)Value);
+				}
+
+			protected:
 				friend FAttributeStorage;
 			};
 
-			FAttributeStorage() = default;
+			INTERCHANGECORE_API FAttributeStorage() = default;
 
 			INTERCHANGECORE_API FAttributeStorage(const FAttributeStorage& Other);
 
@@ -530,49 +498,23 @@ namespace UE
 			 * - Key exists with a different type.
 			 * - Storage is corrupted.
 			 */
-			template<typename T>
-			EAttributeStorageResult RegisterAttribute(const FAttributeKey& ElementAttributeKey, const T& DefaultValue, EAttributeProperty AttributeProperty = EAttributeProperty::None)
-			{
-				static_assert(TAttributeTypeTraits<T>::GetType() != EAttributeTypes::None, "T is not a supported type for the attributes. Check EAttributeTypes for the supported types");
-
-				//Lock the storage
-				FScopeLock ScopeLock(&StorageMutex);
-
-
-				const EAttributeTypes ValueType = TAttributeTypeTraits<T>::GetType();
-				const uint64 ValueSize = GetValueSize(DefaultValue);
-
-				FAttributeAllocationInfo* AttributeAllocationInfo = AttributeAllocationTable.Find(ElementAttributeKey);
-				if (AttributeAllocationInfo)
-				{
-					if (AttributeAllocationInfo->Type != ValueType)
-					{
-						return EAttributeStorageResult::Operation_Error_WrongType;
-					}
-					if (!AttributeStorage.IsValidIndex(AttributeAllocationInfo->Offset))
-					{
-						return EAttributeStorageResult::Operation_Error_AttributeAllocationCorrupted;
-					}
-				}
-				else
-				{
-					AttributeAllocationInfo = &AttributeAllocationTable.Add(ElementAttributeKey);
-					AttributeAllocationInfo->Type = ValueType;
-					AttributeAllocationInfo->Size = ValueSize;
-					AttributeAllocationInfo->Offset = AttributeStorage.AddZeroed(ValueSize);
-				}
+			template<typename T, typename = std::enable_if_t<TIsNonEnumType<T>>>
+			INTERCHANGECORE_API EAttributeStorageResult RegisterAttribute(const FAttributeKey& ElementAttributeKey, const T& DefaultValue, EAttributeProperty AttributeProperty = EAttributeProperty::None);
 
-				//Force the specified attribute property
-				AttributeAllocationInfo->Property = AttributeProperty;
+			// Redirect enum instantiations to underlying type
+			template<typename E, typename = std::enable_if_t<TIsEnum<E>::Value>, int = 0>
+			EAttributeStorageResult RegisterAttribute(const FAttributeKey& ElementAttributeKey, const E& DefaultValue, EAttributeProperty AttributeProperty = EAttributeProperty::None)
+			{
+				using UnderlyingType = std::underlying_type_t<E>;
+				return RegisterAttribute<UnderlyingType>(ElementAttributeKey, (const UnderlyingType&)DefaultValue, AttributeProperty);
+			}
 
-				//Use template SetAttribute which support all the types
-				const EAttributeStorageResult Result = SetAttribute(AttributeAllocationInfo, DefaultValue);
-				if (!IsAttributeStorageResultSuccess(Result))
-				{
-					//An error occured, unregister the key from the storage
-					UnregisterAttribute(ElementAttributeKey);
-				}
-				return Result;
+			// Redirect TEnumAsByte instantiations to underlying type
+			template<typename B, typename = std::enable_if_t<TIsTEnumAsByte<B>::Value>, int = 0, int = 0>
+			EAttributeStorageResult RegisterAttribute(const FAttributeKey& ElementAttributeKey, const B& DefaultValue, EAttributeProperty AttributeProperty = EAttributeProperty::None)
+			{
+				using UnderlyingType = uint8;
+				return RegisterAttribute<UnderlyingType>(ElementAttributeKey, (const UnderlyingType&)DefaultValue, AttributeProperty);
 			}
 
 			/**
@@ -586,19 +528,60 @@ namespace UE
 			 */
 			INTERCHANGECORE_API EAttributeStorageResult UnregisterAttribute(const FAttributeKey& ElementAttributeKey);
 
+			/**
+			 * Retrieve a copy of an attribute value. Return success if the attribute is added.
+			 *
+			 * @param ElementAttributeKey is the storage key (the path) of the attribute.
+			 * @param OutValue the reference value where we copy of the attribute value.
+			 *
+			 * @note Possible errors
+			 * - Key do not exist
+			 * - Key exist with a different type
+			 * - Key exist with a wrong size
+			 */
+			template<typename T, typename = std::enable_if_t<TIsNonEnumType<T>>>
+			INTERCHANGECORE_API EAttributeStorageResult GetAttribute(const FAttributeKey& ElementAttributeKey, T& OutValue) const;
+
+			// Redirect enum instantiations to underlying type
+			template<typename E, typename = std::enable_if_t<TIsEnum<E>::Value>, int = 0>
+			EAttributeStorageResult GetAttribute(const FAttributeKey& ElementAttributeKey, E& OutValue) const
+			{
+				using UnderlyingType = std::underlying_type_t<E>;
+				return GetAttribute<UnderlyingType>(ElementAttributeKey, OutValue);
+			}
+
+			// Redirect TEnumAsByte instantiations to underlying type
+			template<typename B, typename = std::enable_if_t<TIsTEnumAsByte<B>::Value>, int = 0, int = 0>
+			EAttributeStorageResult GetAttribute(const FAttributeKey& ElementAttributeKey, B& OutValue) const
+			{
+				using UnderlyingType = uint8;
+				return GetAttribute<UnderlyingType>(ElementAttributeKey, OutValue);
+			}
+
 			/**
 			 * Return an attribute handle for the specified attribute. This handle is a compile type check and is use to get and set the attribute value type.
 			 * The function will assert if the key is missing or the type doesn't match the specified template type.
 			 *
 			 * @param ElementAttributeKey - the storage key (the path) of the attribute.
 			 */
-			template<typename T>
-			TAttributeHandle<T> GetAttributeHandle(const FAttributeKey& ElementAttributeKey) const
+			template<typename T, typename = std::enable_if_t<TIsNonEnumType<T>>>
+			INTERCHANGECORE_API TAttributeHandle<T> GetAttributeHandle(const FAttributeKey& ElementAttributeKey) const;
+
+			// Redirect enum instantiations to underlying type
+			template<typename E, typename = std::enable_if_t<TIsEnum<E>::Value>, int = 0>
+			TAttributeHandle<E> GetAttributeHandle(const FAttributeKey& ElementAttributeKey) const
 			{
-				TAttributeHandle<T> AttributeHandle(ElementAttributeKey, this);
-				return AttributeHandle;
+				using UnderlyingType = std::underlying_type_t<E>;
+				return (TAttributeHandle<E>)GetAttributeHandle<UnderlyingType>(ElementAttributeKey);
 			}
 
+			// Redirect TEnumAsByte instantiations to underlying type
+			template<typename B, typename = std::enable_if_t<TIsTEnumAsByte<B>::Value>, int = 0, int = 0>
+			TAttributeHandle<B> GetAttributeHandle(const FAttributeKey& ElementAttributeKey) const
+			{
+				using UnderlyingType = uint8;
+				return (TAttributeHandle<B>)GetAttributeHandle<UnderlyingType>(ElementAttributeKey);
+			}
 
 			/**
 			 * Return the attribute type if the key exists, or None if the key is missing.
@@ -764,336 +747,30 @@ namespace UE
 			 * - Key exist with a wrong size
 			 */
 			template<typename T>
-			EAttributeStorageResult SetAttribute(const FAttributeKey& ElementAttributeKey, const T& Value)
-			{
-				FAttributeAllocationInfo* AttributeAllocationInfo = AttributeAllocationTable.Find(ElementAttributeKey);
-				return SetAttribute(AttributeAllocationInfo, Value);
-			}
+			EAttributeStorageResult SetAttribute(const FAttributeKey& ElementAttributeKey, const T& Value);
 
 			template<typename T>
-			EAttributeStorageResult SetAttribute(FAttributeAllocationInfo* AttributeAllocationInfo, const T& Value)
-			{
-				return SetAttribute(AttributeAllocationInfo, Value, TSpecializeType<T>());
-			}
+			EAttributeStorageResult SetAttribute(FAttributeAllocationInfo* AttributeAllocationInfo, const T& Value);
 
-		public:
-			/**
-			 * Retrieve a copy of an attribute value. Return success if the attribute is added.
-			 *
-			 * @param ElementAttributeKey is the storage key (the path) of the attribute.
-			 * @param OutValue the reference value where we copy of the attribute value.
-			 *
-			 * @note Possible errors
-			 * - Key do not exist
-			 * - Key exist with a different type
-			 * - Key exist with a wrong size
-			 */
-			template<typename T>
-			EAttributeStorageResult GetAttribute(const FAttributeKey& ElementAttributeKey, T& OutValue) const
-			{
-				return GetAttribute(ElementAttributeKey, OutValue,TSpecializeType<T>());
-			}
-		protected:
 			/** Defrag the storage by using memmove on the attribute store after a hole in the storage. */
-			INTERCHANGECORE_API void DefragInternal();
+			void DefragInternal();
 
-			static INTERCHANGECORE_API FGuid GetValueHash(const uint8* Value, uint64 ValueSize);
-
-			/** Return the storage size of the value. */
 			template<typename T>
+			UE_DEPRECATED(5.7, "No longer used.")
 			static uint64 GetValueSize(const T& Value)
 			{
-				return GetValueSize(Value, TSpecializeType<T>());
-			}
-	
-			//Begin specialize templates
-	
-			template<typename T>
-			EAttributeStorageResult SetAttribute(FAttributeAllocationInfo* AttributeAllocationInfo, const T& Value, TSpecializeType<T>)
-			{
-				static_assert(TAttributeTypeTraits<T>::GetType() != EAttributeTypes::None, "T is not a supported type for the attributes. Check EAttributeTypes for the supported types");
-
-				//Lock the storage
-				FScopeLock ScopeLock(&StorageMutex);
-
-				if (!AttributeAllocationInfo)
-				{
-					return EAttributeStorageResult::Operation_Error_CannotFoundKey;
-				}
-				const EAttributeTypes ValueType = TAttributeTypeTraits<T>::GetType();
-				const uint64 ValueSize = GetValueSize(Value);
-
-				if (AttributeAllocationInfo->Type != ValueType)
-				{
-					return EAttributeStorageResult::Operation_Error_WrongType;
-				}
-				if (AttributeAllocationInfo->Size != ValueSize)
-				{
-					return EAttributeStorageResult::Operation_Error_WrongSize;
-				}
-				if (!AttributeStorage.IsValidIndex(AttributeAllocationInfo->Offset))
-				{
-					return EAttributeStorageResult::Operation_Error_AttributeAllocationCorrupted;
-				}
-
-				//Set the hash from the value
-				AttributeAllocationInfo->Hash = GetValueHash(reinterpret_cast<const uint8*>(&Value), ValueSize);
-
-				uint8* StorageData = AttributeStorage.GetData();
-				FMemory::Memcpy(&StorageData[AttributeAllocationInfo->Offset], &Value, ValueSize);
-				return EAttributeStorageResult::Operation_Success;
-			}
-
-			EAttributeStorageResult SetAttribute(FAttributeAllocationInfo* AttributeAllocationInfo, const TArray<uint8>& Value, TSpecializeType<TArray<uint8> >)
-			{
-				return MultiSizeSetAttribute(AttributeAllocationInfo, Value, [&Value]()->uint8*
-				{
-					return const_cast<uint8*>(Value.GetData());
-				});
-			}
-
-			EAttributeStorageResult SetAttribute(FAttributeAllocationInfo* AttributeAllocationInfo, const TArray64<uint8>& Value, TSpecializeType<TArray64<uint8> >)
-			{
-				return MultiSizeSetAttribute(AttributeAllocationInfo, Value, [&Value]()->uint8*
-				{
-					return const_cast<uint8*>(Value.GetData());
-				});
-			}
-
-			EAttributeStorageResult SetAttribute(FAttributeAllocationInfo* AttributeAllocationInfo, const FString& Value, TSpecializeType<FString >)
-			{
-				return MultiSizeSetAttribute(AttributeAllocationInfo, Value, [&Value]()->uint8*
-				{
-					return (uint8*)(*Value);
-				});
-			}
-
-			EAttributeStorageResult SetAttribute(FAttributeAllocationInfo* AttributeAllocationInfo, const FName& Value, TSpecializeType<FName >)
-			{
-				//FName must be store has a FString for persistence
-				FString ValueStr = Value.ToString();
-				return MultiSizeSetAttribute(AttributeAllocationInfo, Value, [&ValueStr]()->uint8*
-				{
-					return (uint8*)(*ValueStr);
-				});
-			}
-
-			EAttributeStorageResult SetAttribute(FAttributeAllocationInfo* AttributeAllocationInfo, const FSoftObjectPath& Value, TSpecializeType<FSoftObjectPath >)
-			{
-				//FSoftObjectPath can be hold has a FString and restore from it
-				FString ValueStr = Value.ToString();
-				return MultiSizeSetAttribute(AttributeAllocationInfo, Value, [&ValueStr]()->uint8*
-				{
-					return (uint8*)(*ValueStr);
-				});
-			}
-
-			template<typename T>
-			EAttributeStorageResult GetAttribute(const FAttributeKey& ElementAttributeKey, T& OutValue, TSpecializeType<T>) const
-			{
-				static_assert(TAttributeTypeTraits<T>::GetType() != EAttributeTypes::None, "T is not a supported type for the attributes. Check EAttributeTypes for the supported types");
-
-				//Lock the storage
-				FScopeLock ScopeLock(&StorageMutex);
-
-				const EAttributeTypes ValueType = TAttributeTypeTraits<T>::GetType();
-				const uint64 ValueSize = GetValueSize(OutValue);
-
-				const FAttributeAllocationInfo* AttributeAllocationInfo = AttributeAllocationTable.Find(ElementAttributeKey);
-				if (AttributeAllocationInfo)
-				{
-					if (AttributeAllocationInfo->Type != ValueType && !((ValueType == EAttributeTypes::UInt8) && (ValueType == EAttributeTypes::Enum)))
-					{
-						return EAttributeStorageResult::Operation_Error_WrongType;
-					}
-					if (AttributeAllocationInfo->Size != ValueSize)
-					{
-						return EAttributeStorageResult::Operation_Error_WrongSize;
-					}
-					if (!AttributeStorage.IsValidIndex(AttributeAllocationInfo->Offset))
-					{
-						return EAttributeStorageResult::Operation_Error_AttributeAllocationCorrupted;
-					}
-				}
-				else
-				{
-					//The key do not exist
-					return EAttributeStorageResult::Operation_Error_CannotFoundKey;
-				}
-
-				const uint8* StorageData = AttributeStorage.GetData();
-				FMemory::Memcpy(&OutValue, &StorageData[AttributeAllocationInfo->Offset], AttributeAllocationInfo->Size);
-				return EAttributeStorageResult::Operation_Success;
-			}
-	
-			EAttributeStorageResult GetAttribute(const FAttributeKey& ElementAttributeKey, TArray<uint8>& OutValue, TSpecializeType<TArray<uint8> >) const
-			{
-				return GenericArrayGetAttribute(ElementAttributeKey, OutValue);
-			}
-    
-			EAttributeStorageResult GetAttribute(const FAttributeKey& ElementAttributeKey, TArray64<uint8>& OutValue, TSpecializeType<TArray64<uint8> >) const
-			{
-				return GenericArrayGetAttribute(ElementAttributeKey, OutValue);
-			}
-	
-			INTERCHANGECORE_API void ExtractFStringAttributeFromStorage(const uint8* StorageData, const FAttributeAllocationInfo* AttributeAllocationInfo, FString& OutValue) const;
-
-			INTERCHANGECORE_API const FStringView GetFStringViewAttributeFromStorage(const uint8* StorageData, const FAttributeAllocationInfo* AttributeAllocationInfo) const;
-
-			INTERCHANGECORE_API EAttributeStorageResult GetAttribute(const FAttributeKey& ElementAttributeKey, FString& OutValue, TSpecializeType<FString >) const;
-
-			INTERCHANGECORE_API EAttributeStorageResult GetAttribute(const FAttributeKey& ElementAttributeKey, FName& OutValue, TSpecializeType<FName >) const;
-
-			INTERCHANGECORE_API EAttributeStorageResult GetAttribute(const FAttributeKey& ElementAttributeKey, FSoftObjectPath& OutValue, TSpecializeType<FSoftObjectPath >) const;
-
-			template<typename T>
-			static uint64 GetValueSize(const T& Value, TSpecializeType<T>)
-			{
-				const EAttributeTypes ValueType = TAttributeTypeTraits<T>::GetType();
-				if (ValueType == EAttributeTypes::None)
-				{
-					return 0;
-				}
-				return sizeof(T);
+				return 1;
 			}
-	
-			static INTERCHANGECORE_API uint64 GetValueSize(const FString& Value, TSpecializeType<FString >);
 
-			static uint64 GetValueSize(const FName& Value, TSpecializeType<FName>)
-			{
-				FString ValueStr = Value.ToString();
-				return GetValueSize(ValueStr, TSpecializeType<FString >());
-			}
+			void ExtractFStringAttributeFromStorage(const uint8* StorageData, const FAttributeAllocationInfo* AttributeAllocationInfo, FString& OutValue) const;
 
-			static uint64 GetValueSize(const FSoftObjectPath& Value, TSpecializeType<FSoftObjectPath>)
-			{
-				FString ValueStr = Value.ToString();
-				return GetValueSize(ValueStr, TSpecializeType<FString >());
-			}
+			const FStringView GetFStringViewAttributeFromStorage(const uint8* StorageData, const FAttributeAllocationInfo* AttributeAllocationInfo) const;
 
-			static INTERCHANGECORE_API uint64 GetValueSize(const TArray<uint8>& Value, TSpecializeType<TArray<uint8> >);
-
-			static INTERCHANGECORE_API uint64 GetValueSize(const TArray64<uint8>& Value, TSpecializeType<TArray64<uint8> >);
-	
-			// End specialize templates
-
-			/** Set a multisize(TArray<uint8>, FString ...) attribute value into the storage. Return success if the attribute was properly set. */
 			template<typename MultiSizeType>
-			EAttributeStorageResult MultiSizeSetAttribute(FAttributeAllocationInfo* AttributeAllocationInfo, const MultiSizeType& Value, TFunctionRef<uint8* ()> GetValuePointer)
-			{
-				//Lock the storage
-				FScopeLock ScopeLock(&StorageMutex);
-
-				const EAttributeTypes ValueType = TAttributeTypeTraits<MultiSizeType>::GetType();
-
-				if (!AttributeAllocationInfo)
-				{
-					return EAttributeStorageResult::Operation_Error_CannotFoundKey;
-				}
-
-				if (AttributeAllocationInfo->Type != ValueType)
-				{
-					return EAttributeStorageResult::Operation_Error_WrongType;
-				}
-
-				const uint64 ValueSize = GetValueSize(Value);
-
-				uint8* DataPtr = GetValuePointer();
-				//We must have a valid data pointer if we have something to copy into the storage.
-				//This case support empty array or empty string because the ValueSize will be 0.
-				if (!DataPtr && ValueSize > 0)
-				{
-					return EAttributeStorageResult::Operation_Error_InvalidMultiSizeValueData;
-				}
-
-				//If the new size is greater then the old size we have to create a new allocation info entry and delete the old one
-				if (ValueSize > AttributeAllocationInfo->Size)
-				{
-					//Find the key to remove it
-					const FAttributeKey* ElementAttributeKeyPtr = AttributeAllocationTable.FindKey(*AttributeAllocationInfo);
-					if (!ElementAttributeKeyPtr)
-					{
-						return EAttributeStorageResult::Operation_Error_CannotFoundKey;
-					}
-					const FAttributeKey OriginalKey = *ElementAttributeKeyPtr;
-					const EAttributeProperty AttributePropertyBackup = AttributeAllocationInfo->Property;
-					const EAttributeStorageResult RemoveResult = UnregisterAttribute(OriginalKey);
-					if (!IsAttributeStorageResultSuccess(RemoveResult))
-					{
-						return RemoveResult;
-					}
-					AttributeAllocationInfo = nullptr;
-					//Reallocate a new one
-					AttributeAllocationInfo = &AttributeAllocationTable.Add(OriginalKey);
-					AttributeAllocationInfo->Type = ValueType;
-					AttributeAllocationInfo->Offset = AttributeStorage.AddZeroed(ValueSize);
-					//Force the specified attribute property
-					AttributeAllocationInfo->Property = AttributePropertyBackup;
-				}
-				else
-				{
-					//In case we reuse the allocation table, simply adjust the waste memory counter
-					FragmentedMemoryCost += AttributeAllocationInfo->Size - ValueSize;
-				}
-
-				//We can recycle allocation table if the size is equal or smaller, so we have to always set it
-				AttributeAllocationInfo->Size = ValueSize;
+			EAttributeStorageResult MultiSizeSetAttribute(FAttributeAllocationInfo* AttributeAllocationInfo, const MultiSizeType& Value, TFunctionRef<uint8* ()> GetValuePointer);
 
-				//Set the hash from the value
-				AttributeAllocationInfo->Hash = GetValueHash(DataPtr, ValueSize);
-
-				if (ValueSize > 0)
-				{
-					uint8* StorageData = AttributeStorage.GetData();
-					FMemory::Memcpy(&StorageData[AttributeAllocationInfo->Offset], DataPtr, ValueSize);
-				}
-
-				return EAttributeStorageResult::Operation_Success;
-			}
-	
 			template<typename ArrayType>
-			EAttributeStorageResult GenericArrayGetAttribute(const FAttributeKey& ElementAttributeKey, ArrayType& OutValue) const
-			{
-				//Lock the storage
-				FScopeLock ScopeLock(&StorageMutex);
-
-				const EAttributeTypes ValueType = TAttributeTypeTraits<ArrayType>::GetType();
-				if (ValueType != EAttributeTypes::ByteArray && ValueType != EAttributeTypes::ByteArray64)
-				{
-					return EAttributeStorageResult::Operation_Error_WrongType;
-				}
-
-				const FAttributeAllocationInfo* AttributeAllocationInfo = AttributeAllocationTable.Find(ElementAttributeKey);
-				if (AttributeAllocationInfo)
-				{
-					if (AttributeAllocationInfo->Type != ValueType)
-					{
-						return EAttributeStorageResult::Operation_Error_WrongType;
-					}
-					if (!AttributeStorage.IsValidIndex(AttributeAllocationInfo->Offset))
-					{
-						return EAttributeStorageResult::Operation_Error_AttributeAllocationCorrupted;
-					}
-				}
-				else
-				{
-					//The key do not exist
-					return EAttributeStorageResult::Operation_Error_CannotFoundKey;
-				}
-
-				if (AttributeAllocationInfo->Size == 0)
-				{
-					OutValue.Empty();
-					return EAttributeStorageResult::Operation_Success;
-				}
-
-				//Allocate the ByteArray
-				OutValue.AddZeroed(static_cast<typename ArrayType::SizeType>(AttributeAllocationInfo->Size));
-				//Copy into the buffer
-				const uint8* StorageData = AttributeStorage.GetData();
-				FMemory::Memcpy(OutValue.GetData(), &StorageData[AttributeAllocationInfo->Offset], AttributeAllocationInfo->Size);
-				return EAttributeStorageResult::Operation_Success;
-			}
+			EAttributeStorageResult GenericArrayGetAttribute(const FAttributeKey& ElementAttributeKey, ArrayType& OutValue) const;
 
 			template<typename T>
 			friend void CopyStorageAttributesInternal(const FAttributeStorage& SourceStorage, FAttributeStorage& DestinationStorage, const TArray<T>& AttributeKeys);
@@ -1109,7 +786,7 @@ namespace UE
 			 * A Hole is create each time we remove an attribute.
 			 */
 			uint64 FragmentedMemoryCost = 0;
-	
+
 			/**
 			 * if FragmentedMemoryCost > AttributeStorage.Num*DefragRatio then defrag.
 			 * This is use whenever we remove attribute or change DefragRatio value.
@@ -1121,529 +798,5 @@ namespace UE
 			 */
 			mutable FCriticalSection StorageMutex;
 		};
-
-		/************************************************************************/
-		/* Default FVariant traits for built-in types */
-
-		/** Implements variant type traits for the built-in bool type. */
-		template<> struct TAttributeTypeTraits<bool>
-		{
-			static constexpr EAttributeTypes GetType() { return EAttributeTypes::Bool; }
-			static FString ToString(const bool& Value) { return Value ? TEXT("true") : TEXT("false"); }
-		};
-
-
-		/** See FBox3f/FBox3d variant type traits for the built-in FBox type. */
-
-
-		/** See FBoxSphereBounds3f/FBoxSphereBounds3d variant type traits for the built-in FBoxSphereBounds type. */
-
-
-		/** Implements variant type traits for byte arrays. */
-		template<> struct TAttributeTypeTraits<TArray<uint8> >
-		{
-			static constexpr EAttributeTypes GetType() { return EAttributeTypes::ByteArray; }
-			static FString ToString(const TArray<uint8>& Value) {return TEXT("Array<uint8>"); }
-		};
-
-		/** Implements variant type traits for byte array64s. */
-		template<> struct TAttributeTypeTraits<TArray64<uint8> >
-		{
-			static constexpr EAttributeTypes GetType() { return EAttributeTypes::ByteArray64; }
-			static FString ToString(const TArray64<uint8>& Value) { return TEXT("Array64<uint8>"); }
-		};
-
-
-		/** Implements variant type traits for the built-in FColor type. */
-		template<> struct TAttributeTypeTraits<FColor>
-		{
-			static constexpr EAttributeTypes GetType() { return EAttributeTypes::Color; }
-			static FString ToString(const FColor& Value) { return Value.ToString(); }
-		};
-
-
-		/** Implements variant type traits for the built-in FDateTime type. */
-		template<> struct TAttributeTypeTraits<FDateTime>
-		{
-			static constexpr EAttributeTypes GetType() { return EAttributeTypes::DateTime; }
-			static FString ToString(const FDateTime& Value) { return Value.ToString(); }
-		};
-
-
-		/** Implements variant type traits for the built-in double type. */
-		template<> struct TAttributeTypeTraits<double>
-		{
-			static constexpr EAttributeTypes GetType() { return EAttributeTypes::Double; }
-			static FString ToString(const double& Value)
-			{
-				FStringFormatOrderedArguments OrderedArguments;
-				OrderedArguments.Add(FStringFormatArg(Value));
-				return FString::Format(TEXT("{0}"), OrderedArguments);
-			}
-		};
-
-
-		/** Implements variant type traits for enumeration types. */
-		template<typename EnumType> struct TAttributeTypeTraits<TEnumAsByte<EnumType> >
-		{
-			static constexpr EAttributeTypes GetType() { return EAttributeTypes::Enum; }
-			static FString ToString(const TEnumAsByte<EnumType>& Value)
-			{
-				uint32 ValueConv = static_cast<uint32>(static_cast<uint8>(Value));
-				FStringFormatOrderedArguments OrderedArguments;
-				OrderedArguments.Add(FStringFormatArg(ValueConv));
-				return FString::Format(TEXT("{0}"), OrderedArguments);
-			}
-		};
-
-
-		/** Implements variant type traits for the built-in float type. */
-		template<> struct TAttributeTypeTraits<float>
-		{
-			static constexpr EAttributeTypes GetType() { return EAttributeTypes::Float; }
-			static FString ToString(const float& Value)
-			{
-				FStringFormatOrderedArguments OrderedArguments;
-				OrderedArguments.Add(FStringFormatArg(Value));
-				return FString::Format(TEXT("{0}"), OrderedArguments);
-			}
-		};
-
-
-		/** Implements variant type traits for the built-in FGuid type. */
-		template<> struct TAttributeTypeTraits<FGuid>
-		{
-			static constexpr EAttributeTypes GetType() { return EAttributeTypes::Guid; }
-			static FString ToString(const FGuid& Value) { return Value.ToString(EGuidFormats::DigitsLower); }
-		};
-
-
-		/** Implements variant type traits for the built-in int8 type. */
-		template<> struct TAttributeTypeTraits<int8>
-		{
-			static constexpr EAttributeTypes GetType() { return EAttributeTypes::Int8; }
-			static FString ToString(const int8& Value) { return FString::FromInt(Value); }
-		};
-
-
-		/** Implements variant type traits for the built-in int16 type. */
-		template<> struct TAttributeTypeTraits<int16>
-		{
-			static constexpr EAttributeTypes GetType() { return EAttributeTypes::Int16; }
-			static FString ToString(const int16& Value) { return FString::FromInt(Value); }
-		};
-
-
-		/** Implements variant type traits for the built-in int32 type. */
-		template<> struct TAttributeTypeTraits<int32>
-		{
-			static constexpr EAttributeTypes GetType() { return EAttributeTypes::Int32; }
-			static FString ToString(const int32& Value) { return FString::FromInt(Value); }
-		};
-
-
-		/** Implements variant type traits for the built-in int64 type. */
-		template<> struct TAttributeTypeTraits<int64>
-		{
-			static constexpr EAttributeTypes GetType() { return EAttributeTypes::Int64; }
-			static FString ToString(const int64& Value)
-			{
-				FStringFormatOrderedArguments OrderedArguments;
-				OrderedArguments.Add(FStringFormatArg(Value));
-				return FString::Format(TEXT("{0}"), OrderedArguments);
-			}
-		};
-
-		/** Implements variant type traits for the built-in FIntRect type. */
-		template<> struct TAttributeTypeTraits<FIntRect>
-		{
-			static constexpr EAttributeTypes GetType() { return EAttributeTypes::IntRect; }
-			static FString ToString(const FIntRect& Value) { return Value.ToString(); }
-		};
-
-
-		/** Implements variant type traits for the built-in FLinearColor type. */
-		template<> struct TAttributeTypeTraits<FLinearColor>
-		{
-			static constexpr EAttributeTypes GetType() { return EAttributeTypes::LinearColor; }
-			static FString ToString(const FLinearColor& Value) { return Value.ToString(); }
-		};
-
-
-		/** See FMatrix44f/FMatrix44d variant type traits for the built-in FMatrix type. */
-
-		/** Implements variant type traits for the built-in FName type. */
-		template<> struct TAttributeTypeTraits<FName>
-		{
-			static constexpr EAttributeTypes GetType() { return EAttributeTypes::Name; }
-			static FString ToString(const FName& Value) { return Value.ToString(); }
-		};
-
-		/** See FPlane4f/FPlane4d variant type traits for the built-in FPlane type. */
-
-
-		/** See FQuat4f/FQuat4d variant type traits for the built-in FQuat type. */
-
-		/** Implements variant type traits for the built-in FRandomStream type. */
-		template<> struct TAttributeTypeTraits<FRandomStream>
-		{
-			static constexpr EAttributeTypes GetType() { return EAttributeTypes::RandomStream; }
-			static FString ToString(const FRandomStream& Value) { return Value.ToString(); }
-		};
-
-
-		/** See FRotator3f/FRotator3d variant type traits for the built-in FRotator type. */
-
-
-		/** Implements variant type traits for the built-in FString type. */
-		template<> struct TAttributeTypeTraits<FString>
-		{
-			static constexpr EAttributeTypes GetType() { return EAttributeTypes::String; }
-			static FString ToString(const FString& Value) { return Value; }
-		};
-
-
-		/** Implements variant type traits for the built-in FTimespan type. */
-		template<> struct TAttributeTypeTraits<FTimespan>
-		{
-			static constexpr EAttributeTypes GetType() { return EAttributeTypes::Timespan; }
-			static FString ToString(const FTimespan& Value) { return Value.ToString(); }
-		};
-
-
-		/** See FTransform3f/FTransform3d variant type traits for the built-in FTransform type. */
-
-
-		/** Implements variant type traits for the built-in FTwoVectors type. */
-		template<> struct TAttributeTypeTraits<FTwoVectors>
-		{
-			static constexpr EAttributeTypes GetType() { return EAttributeTypes::TwoVectors; }
-			static FString ToString(const FTwoVectors& Value) { return Value.ToString(); }
-		};
-
-
-		/** Implements variant type traits for the built-in uint8 type. */
-		template<> struct TAttributeTypeTraits<uint8>
-		{
-			static constexpr EAttributeTypes GetType() { return EAttributeTypes::UInt8; }
-			static FString ToString(const uint8& Value)
-			{
-				uint32 ValueConv = Value;
-				FStringFormatOrderedArguments OrderedArguments;
-				OrderedArguments.Add(FStringFormatArg(ValueConv));
-				return FString::Format(TEXT("{0}"), OrderedArguments);
-			}
-		};
-
-
-		/** Implements variant type traits for uint8 enum classes. */
-		template<typename EnumClassType> struct TAttributeTypeTraits
-			<EnumClassType,
-				typename std::enable_if<
-					std::is_same<
-						typename UE::Interchange::Private::TUnderlyingType<EnumClassType>::type,
-						uint8
-					>::value
-				>::type
-			>
-		{
-			static constexpr EAttributeTypes GetType() { return EAttributeTypes::UInt8; }
-			static FString ToString(const EnumClassType& Value)
-			{
-				uint32 ValueConv = (uint8)Value;
-				FStringFormatOrderedArguments OrderedArguments;
-				OrderedArguments.Add(FStringFormatArg(ValueConv));
-				return FString::Format(TEXT("{0}"), OrderedArguments);
-			}
-		};
-
-		/** Implements variant type traits for the built-in uint16 type. */
-		template<> struct TAttributeTypeTraits<uint16>
-		{
-			static constexpr EAttributeTypes GetType() { return EAttributeTypes::UInt16; }
-			static FString ToString(const uint16& Value)
-			{
-				uint32 ValueConv = Value;
-				FStringFormatOrderedArguments OrderedArguments;
-				OrderedArguments.Add(FStringFormatArg(ValueConv));
-				return FString::Format(TEXT("{0}"), OrderedArguments);
-			}
-		};
-
-
-		/** Implements variant type traits for the built-in uint32 type. */
-		template<> struct TAttributeTypeTraits<uint32>
-		{
-			static constexpr EAttributeTypes GetType() { return EAttributeTypes::UInt32; }
-			static FString ToString(const uint32& Value)
-			{
-				FStringFormatOrderedArguments OrderedArguments;
-				OrderedArguments.Add(FStringFormatArg(Value));
-				return FString::Format(TEXT("{0}"), OrderedArguments);
-			}
-		};
-
-
-		/** Implements variant type traits for the built-in uint64 type. */
-		template<> struct TAttributeTypeTraits<uint64>
-		{
-			static constexpr EAttributeTypes GetType() { return EAttributeTypes::UInt64; }
-			static FString ToString(const uint64& Value)
-			{
-				FStringFormatOrderedArguments OrderedArguments;
-				OrderedArguments.Add(FStringFormatArg(Value));
-				return FString::Format(TEXT("{0}"), OrderedArguments);
-			}
-		};
-
-
-		/** See FVector3f/FVector3d variant type traits for the built-in FVector type. */
-
-
-		/** Implements variant type traits for the built-in FVector2D type. */
-		template<> struct TAttributeTypeTraits<FVector2D>
-		{
-			static constexpr EAttributeTypes GetType() { return EAttributeTypes::Vector2d; }
-			static FString ToString(const FVector2D& Value) { return Value.ToString(); }
-		};
-
-		/** See FVector4f/FVector4d variant type traits for the built-in FVector4 type. */
-
-		/** Implements variant type traits for the built-in Vector2DHalf type. */
-		template<> struct TAttributeTypeTraits<FIntPoint>
-		{
-			static constexpr EAttributeTypes GetType() { return EAttributeTypes::IntPoint; }
-			static FString ToString(const FIntPoint& Value) { return Value.ToString(); }
-		};
-
-		/** Implements variant type traits for the built-in Vector2DHalf type. */
-		template<> struct TAttributeTypeTraits<FIntVector>
-		{
-			static constexpr EAttributeTypes GetType() { return EAttributeTypes::IntVector; }
-			static FString ToString(const FIntVector& Value) { return Value.ToString(); }
-		};
-
-		/** Implements variant type traits for the built-in Vector2DHalf type. */
-		template<> struct TAttributeTypeTraits<FVector2DHalf>
-		{
-			static constexpr EAttributeTypes GetType() { return EAttributeTypes::Vector2DHalf; }
-			static FString ToString(const FVector2DHalf& Value) { return Value.ToString(); }
-		};
-
-		/** Implements variant type traits for the built-in Float16 type. */
-		template<> struct TAttributeTypeTraits<FFloat16>
-		{
-			static constexpr EAttributeTypes GetType() { return EAttributeTypes::Float16; }
-			static FString ToString(const FFloat16& Value)
-			{
-				float ValueConvert = Value.GetFloat();
-				FStringFormatOrderedArguments OrderedArguments;
-				OrderedArguments.Add(FStringFormatArg(ValueConvert));
-				return FString::Format(TEXT("{0}"), OrderedArguments);
-			}
-		};
-
-		/** Implements variant type traits for the built-in OrientedBox type. */
-		template<> struct TAttributeTypeTraits<FOrientedBox>
-		{
-			static constexpr EAttributeTypes GetType() { return EAttributeTypes::OrientedBox; }
-			static FString ToString(const FOrientedBox& Value) { return TEXT("FOrientedBox"); }
-		};
-
-		/** See FSphere3f/FSphere3d variant type traits for the built-in Sphere type. */
-
-		/** Implements variant type traits for the built-in FrameNumber type. */
-		template<> struct TAttributeTypeTraits<FFrameNumber>
-		{
-			static constexpr EAttributeTypes GetType() { return EAttributeTypes::FrameNumber; }
-			static FString ToString(const FFrameNumber& Value)
-			{
-				int32 ValueConvert = Value.Value;
-				FStringFormatOrderedArguments OrderedArguments;
-				OrderedArguments.Add(FStringFormatArg(ValueConvert));
-				return FString::Format(TEXT("{0}"), OrderedArguments);
-			}
-		};
-
-		/** Implements variant type traits for the built-in FrameRate type. */
-		template<> struct TAttributeTypeTraits<FFrameRate>
-		{
-			static constexpr EAttributeTypes GetType() { return EAttributeTypes::FrameRate; }
-			static FString ToString(const FFrameRate& Value) { return Value.ToPrettyText().ToString(); }
-		};
-
-		/** Implements variant type traits for the built-in FrameTime type. */
-		template<> struct TAttributeTypeTraits<FFrameTime>
-		{
-			static constexpr EAttributeTypes GetType() { return EAttributeTypes::FrameTime; }
-		};
-
-		/** Implements variant type traits for the built-in SoftObjectPath type. */
-		template<> struct TAttributeTypeTraits<FSoftObjectPath>
-		{
-			static constexpr EAttributeTypes GetType() { return EAttributeTypes::SoftObjectPath; }
-			static FString ToString(const FSoftObjectPath& Value) { return Value.ToString(); }
-		};
-
-		/** Implements variant type traits for the built-in Matrix44f type. */
-		template<> struct TAttributeTypeTraits<FMatrix44f>
-		{
-			static constexpr EAttributeTypes GetType() { return EAttributeTypes::Matrix44f; }
-			static FString ToString(const FMatrix44f& Value) { return Value.ToString(); }
-		};
-
-		/** Implements variant type traits for the built-in Matrix44f type. */
-		template<> struct TAttributeTypeTraits<FMatrix44d>
-		{
-			static constexpr EAttributeTypes GetType() { return EAttributeTypes::Matrix44d; }
-			static FString ToString(const FMatrix44d& Value) { return Value.ToString(); }
-		};
-
-		/** Implements variant type traits for the built-in Plane4f type. */
-		template<> struct TAttributeTypeTraits<FPlane4f>
-		{
-			static constexpr EAttributeTypes GetType() { return EAttributeTypes::Plane4f; }
-			static FString ToString(const FPlane4f& Value) { return Value.ToString(); }
-		};
-
-		/** Implements variant type traits for the built-in Plane4d type. */
-		template<> struct TAttributeTypeTraits<FPlane4d>
-		{
-			static constexpr EAttributeTypes GetType() { return EAttributeTypes::Plane4d; }
-			static FString ToString(const FPlane4d& Value) { return Value.ToString(); }
-		};
-
-		/** Implements variant type traits for the built-in Quat4f type. */
-		template<> struct TAttributeTypeTraits<FQuat4f>
-		{
-			static constexpr EAttributeTypes GetType() { return EAttributeTypes::Quat4f; }
-			static FString ToString(const FQuat4f& Value) { return Value.ToString(); }
-		};
-
-		/** Implements variant type traits for the built-in Quat4d type. */
-		template<> struct TAttributeTypeTraits<FQuat4d>
-		{
-			static constexpr EAttributeTypes GetType() { return EAttributeTypes::Quat4d; }
-			static FString ToString(const FQuat4d& Value) { return Value.ToString(); }
-		};
-
-		/** Implements variant type traits for the built-in Rotator3f type. */
-		template<> struct TAttributeTypeTraits<FRotator3f>
-		{
-			static constexpr EAttributeTypes GetType() { return EAttributeTypes::Rotator3f; }
-			static FString ToString(const FRotator3f& Value) { return Value.ToString(); }
-		};
-
-		/** Implements variant type traits for the built-in Rotator3d type. */
-		template<> struct TAttributeTypeTraits<FRotator3d>
-		{
-			static constexpr EAttributeTypes GetType() { return EAttributeTypes::Rotator3d; }
-			static FString ToString(const FRotator3d& Value) { return Value.ToString(); }
-		};
-
-		/** Implements variant type traits for the built-in Rotator3f type. */
-		template<> struct TAttributeTypeTraits<FTransform3f>
-		{
-			static constexpr EAttributeTypes GetType() { return EAttributeTypes::Transform3f; }
-			static FString ToString(const FTransform3f& Value) { return Value.ToString(); }
-		};
-
-		/** Implements variant type traits for the built-in Transform3d type. */
-		template<> struct TAttributeTypeTraits<FTransform3d>
-		{
-			static constexpr EAttributeTypes GetType() { return EAttributeTypes::Transform3d; }
-			static FString ToString(const FTransform3d& Value) { return Value.ToString(); }
-		};
-
-		/** Implements variant type traits for the built-in Vector3f type. */
-		template<> struct TAttributeTypeTraits<FVector3f>
-		{
-			static constexpr EAttributeTypes GetType() { return EAttributeTypes::Vector3f; }
-			static FString ToString(const FVector3f& Value) { return Value.ToString(); }
-		};
-
-		/** Implements variant type traits for the built-in Vector3d type. */
-		template<> struct TAttributeTypeTraits<FVector3d>
-		{
-			static constexpr EAttributeTypes GetType() { return EAttributeTypes::Vector3d; }
-			static FString ToString(const FVector3d& Value) { return Value.ToString(); }
-		};
-
-		/** Implements variant type traits for the built-in Vector2f type. */
-		template<> struct TAttributeTypeTraits<FVector2f>
-		{
-			static constexpr EAttributeTypes GetType() { return EAttributeTypes::Vector2f; }
-			static FString ToString(const FVector2f& Value) { return Value.ToString(); }
-		};
-
-		/** Implements variant type traits for the built-in Vector4f type. */
-		template<> struct TAttributeTypeTraits<FVector4f>
-		{
-			static constexpr EAttributeTypes GetType() { return EAttributeTypes::Vector4f; }
-			static FString ToString(const FVector4f& Value) { return Value.ToString(); }
-		};
-
-		/** Implements variant type traits for the built-in Vector4d type. */
-		template<> struct TAttributeTypeTraits<FVector4d>
-		{
-			static constexpr EAttributeTypes GetType() { return EAttributeTypes::Vector4d; }
-			static FString ToString(const FVector4d& Value) { return Value.ToString(); }
-		};
-
-		/** Implements variant type traits for the built-in Box2f type. */
-		template<> struct TAttributeTypeTraits<FBox2f>
-		{
-			static constexpr EAttributeTypes GetType() { return EAttributeTypes::Box2f; }
-			static FString ToString(const FBox2f& Value) { return Value.ToString(); }
-		};
-
-		/** Implements variant type traits for the built-in Box2D type. */
-		template<> struct TAttributeTypeTraits<FBox2D>
-		{
-			static constexpr EAttributeTypes GetType() { return EAttributeTypes::Box2D; }
-			static FString ToString(const FBox2D& Value) { return Value.ToString(); }
-		};
-
-		/** Implements variant type traits for the built-in Box3f type. */
-		template<> struct TAttributeTypeTraits<FBox3f>
-		{
-			static constexpr EAttributeTypes GetType() { return EAttributeTypes::Box3f; }
-			static FString ToString(const FBox3f& Value) { return Value.ToString(); }
-		};
-
-		/** Implements variant type traits for the built-in Box3d type. */
-		template<> struct TAttributeTypeTraits<FBox3d>
-		{
-			static constexpr EAttributeTypes GetType() { return EAttributeTypes::Box3d; }
-			static FString ToString(const FBox3d& Value) { return Value.ToString(); }
-		};
-
-		/** Implements variant type traits for the built-in BoxSphereBounds3f type. */
-		template<> struct TAttributeTypeTraits<FBoxSphereBounds3f>
-		{
-			static constexpr EAttributeTypes GetType() { return EAttributeTypes::BoxSphereBounds3f; }
-			static FString ToString(const FBoxSphereBounds3f& Value) { return Value.ToString(); }
-		};
-
-		/** Implements variant type traits for the built-in BoxSphereBounds3d type. */
-		template<> struct TAttributeTypeTraits<FBoxSphereBounds3d>
-		{
-			static constexpr EAttributeTypes GetType() { return EAttributeTypes::BoxSphereBounds3d; }
-			static FString ToString(const FBoxSphereBounds3d& Value) { return Value.ToString(); }
-		};
-
-		/** Implements variant type traits for the built-in Sphere3f type. */
-		template<> struct TAttributeTypeTraits<FSphere3f>
-		{
-			static constexpr EAttributeTypes GetType() { return EAttributeTypes::Sphere3f; }
-			static FString ToString(const FSphere3f& Value) { return TEXT("FSphere3f"); }
-		};
-
-		/** Implements variant type traits for the built-in Sphere3d type. */
-		template<> struct TAttributeTypeTraits<FSphere3d>
-		{
-			static constexpr EAttributeTypes GetType() { return EAttributeTypes::Sphere3d; }
-			static FString ToString(const FSphere3d& Value) { return TEXT("FSphere3d"); }
-		};
 	} //ns interchange
 }
