Interchange Analytics - Renamed Analytics Fields to get the EGPI free Analytics
- Old fields were redacted as older versions of the engine were not sending EGPI free payloads.

#rnx

diff --git a/Engine/Source/Runtime/Interchange/Engine/Private/InterchangeManager.cpp b/Engine/Source/Runtime/Interchange/Engine/Private/InterchangeManager.cpp
--- a/Engine/Source/Runtime/Interchange/Engine/Private/InterchangeManager.cpp
+++ b/Engine/Source/Runtime/Interchange/Engine/Private/InterchangeManager.cpp
@@ -600,9 +600,9 @@ void UE::Interchange::FImportAsyncHelper::SendAnalyticImportEndData()
 	Attribs.Add(FAnalyticsEventAttribute(TEXT("ImportObjectCount"), ImportedObjectCount));
 
 	//Report any warning or error message
-	TMap<FString, int32> WarningMessages;
-	TMap<FString, int32> ErrorMessages;
-	auto CollectResultContainer = [&WarningMessages, &ErrorMessages](const UInterchangeResultsContainer* ResultContainer)
+	TMap<FString, int32> WarningTypes;
+	TMap<FString, int32> ErrorTypes;
+	auto CollectResultContainer = [&WarningTypes, &ErrorTypes](const UInterchangeResultsContainer* ResultContainer)
 	{
 		TArray<UInterchangeResult*> InterchangeResults = ResultContainer->GetResults();
 		
@@ -619,7 +619,7 @@ void UE::Interchange::FImportAsyncHelper::SendAnalyticImportEndData()
 				FString OutWarningAttribValue;
 				if (ExtractNamespace(InterchangeResult->GetText(), OutWarningAttribValue))
 				{
-					int32& Frequency = WarningMessages.FindOrAdd(OutWarningAttribValue);
+					int32& Frequency = WarningTypes.FindOrAdd(OutWarningAttribValue);
 					Frequency++;
 				}
 				else
@@ -633,7 +633,7 @@ void UE::Interchange::FImportAsyncHelper::SendAnalyticImportEndData()
 				FString OutErrorAttribValue;
 				if (ExtractNamespace(InterchangeResult->GetText(), OutErrorAttribValue))
 				{
-					int32& Frequency = ErrorMessages.FindOrAdd(OutErrorAttribValue);
+					int32& Frequency = ErrorTypes.FindOrAdd(OutErrorAttribValue);
 					Frequency++;
 				}
 				else
@@ -657,8 +657,8 @@ void UE::Interchange::FImportAsyncHelper::SendAnalyticImportEndData()
 
 	{
 		using namespace UE::Interchange::Private;
-		Attribs.Add(FAnalyticsEventAttribute(TEXT("WarningMessages"), AnalyticsConvertMapToJsonFragment(WarningMessages)));
-		Attribs.Add(FAnalyticsEventAttribute(TEXT("ErrorMessages"), AnalyticsConvertMapToJsonFragment(ErrorMessages)));
+		Attribs.Add(FAnalyticsEventAttribute(TEXT("WarningTypes"), AnalyticsConvertMapToJsonFragment(WarningTypes)));
+		Attribs.Add(FAnalyticsEventAttribute(TEXT("ErrorTypes"), AnalyticsConvertMapToJsonFragment(ErrorTypes)));
 	}
 
 	FString EventString = TEXT("Interchange.Usage.ImportResult");
