[BUGFIX] Interchange Framework - Missing References on Animation Import with Python Pipelines
- Interchange Python Pipeline when duplicated for asset import data, removes the Editor Only flag from package.

#rnx

diff --git a/Engine/Source/Runtime/Interchange/Engine/Private/InterchangeImportCommon.cpp b/Engine/Source/Runtime/Interchange/Engine/Private/InterchangeImportCommon.cpp
--- a/Engine/Source/Runtime/Interchange/Engine/Private/InterchangeImportCommon.cpp
+++ b/Engine/Source/Runtime/Interchange/Engine/Private/InterchangeImportCommon.cpp
@@ -110,6 +110,11 @@ namespace UE::Interchange
 					UObject* DupPipeline = Cast<UObject>(StaticDuplicateObject(Pipeline, AssetImportData));
 					if (DupPipeline)
 					{
+						// CAUTION: Python Pipeline adds this flag to the Package in PostLoad, 
+						// and StaticDuplicateObject calls PostLoad on the duplicated asset. But here as the Outer is AssetImportData
+						// The package is the imported asset. This will only show references in the reference viewer when the mode is set to
+						// Editor Only References. Hence removing that particular flag from the package.
+						DupPipeline->GetPackage()->ClearPackageFlags(PKG_EditorOnly);
 						NewPipelines.Add(DupPipeline);
 					}
 				}
