Interchange: Add support for array EAttributeTypes.

diff --git a/Engine/Plugins/Interchange/Runtime/Source/Pipelines/Private/InterchangeGenericAssetsPipeline.cpp b/Engine/Plugins/Interchange/Runtime/Source/Pipelines/Private/InterchangeGenericAssetsPipeline.cpp
--- a/Engine/Plugins/Interchange/Runtime/Source/Pipelines/Private/InterchangeGenericAssetsPipeline.cpp
+++ b/Engine/Plugins/Interchange/Runtime/Source/Pipelines/Private/InterchangeGenericAssetsPipeline.cpp
@@ -1339,6 +1339,21 @@ void UInterchangeGenericAssetsPipeline::ExecutePostImportPipeline(const UInterch
 	AddMetaData(CreatedAsset, InBaseNodeContainer->GetNode(NodeKey));
 }
 
+namespace UE::InterchangeGenericAssetsPipeline::Private
+{
+	template<typename T>
+	TOptional<FString> GetValueAsString(const UInterchangeBaseNode* Node, const FString& AttributeName)
+	{
+		TOptional<FString> PayloadKey;
+		T Value{};
+		if (UInterchangeUserDefinedAttributesAPI::GetUserDefinedAttribute(Node, AttributeName, Value, PayloadKey))
+		{
+			return UE::Interchange::AttributeValueToString(Value);
+		}
+		return {};
+	}
+}
+
 void UInterchangeGenericAssetsPipeline::AddMetaData(UObject* CreatedAsset, const UInterchangeBaseNode* Node)
 {
 	if (!CreatedAsset || !Node)
@@ -1372,177 +1387,344 @@ void UInterchangeGenericAssetsPipeline::AddMetaData(UObject* CreatedAsset, const
 			TOptional<FString> PayloadKey;
 			switch (UserAttributeInfo.Type)
 			{
+				using namespace UE::InterchangeGenericAssetsPipeline::Private;
+
 				case UE::Interchange::EAttributeTypes::Bool:
-				{
-					bool Value = false;
-					if (UInterchangeUserDefinedAttributesAPI::GetUserDefinedAttribute(Node, UserAttributeInfo.Name, Value, PayloadKey))
-					{
-						MetaDataValue = UE::Interchange::AttributeValueToString(Value);
-					}
-				}
-				break;
+					MetaDataValue = GetValueAsString<bool>(Node, UserAttributeInfo.Name);
+					break;
+				case UE::Interchange::EAttributeTypes::Color:
+					MetaDataValue = GetValueAsString<FColor>(Node, UserAttributeInfo.Name);
+					break;
+				case UE::Interchange::EAttributeTypes::DateTime:
+					MetaDataValue = GetValueAsString<FDateTime>(Node, UserAttributeInfo.Name);
+					break;
+				case UE::Interchange::EAttributeTypes::Double:
+					MetaDataValue = GetValueAsString<double>(Node, UserAttributeInfo.Name);
+					break;
+				case UE::Interchange::EAttributeTypes::Enum:
+					MetaDataValue = GetValueAsString<uint8>(Node, UserAttributeInfo.Name);
+					break;
+				case UE::Interchange::EAttributeTypes::Float:
+					MetaDataValue = GetValueAsString<float>(Node, UserAttributeInfo.Name);
+					break;
+				case UE::Interchange::EAttributeTypes::Guid:
+					MetaDataValue = GetValueAsString<FGuid>(Node, UserAttributeInfo.Name);
+					break;
 				case UE::Interchange::EAttributeTypes::Int8:
-				{
-					int8 Value = 0;
-					if (UInterchangeUserDefinedAttributesAPI::GetUserDefinedAttribute(Node, UserAttributeInfo.Name, Value, PayloadKey))
-					{
-						MetaDataValue = UE::Interchange::AttributeValueToString(Value);
-					}
-				}
-				break;
+					MetaDataValue = GetValueAsString<int8>(Node, UserAttributeInfo.Name);
+					break;
 				case UE::Interchange::EAttributeTypes::Int16:
-				{
-					int16 Value = 0;
-					if (UInterchangeUserDefinedAttributesAPI::GetUserDefinedAttribute(Node, UserAttributeInfo.Name, Value, PayloadKey))
-					{
-						MetaDataValue = UE::Interchange::AttributeValueToString(Value);
-					}
-				}
-				break;
+					MetaDataValue = GetValueAsString<int16>(Node, UserAttributeInfo.Name);
+					break;
 				case UE::Interchange::EAttributeTypes::Int32:
-				{
-					int32 Value = 0;
-					if (UInterchangeUserDefinedAttributesAPI::GetUserDefinedAttribute(Node, UserAttributeInfo.Name, Value, PayloadKey))
-					{
-						MetaDataValue = UE::Interchange::AttributeValueToString(Value);
-					}
-				}
-				break;
+					MetaDataValue = GetValueAsString<int32>(Node, UserAttributeInfo.Name);
+					break;
 				case UE::Interchange::EAttributeTypes::Int64:
-				{
-					int64 Value = 0;
-					if (UInterchangeUserDefinedAttributesAPI::GetUserDefinedAttribute(Node, UserAttributeInfo.Name, Value, PayloadKey))
-					{
-						MetaDataValue = UE::Interchange::AttributeValueToString(Value);
-					}
-				}
-				break;
+					MetaDataValue = GetValueAsString<int64>(Node, UserAttributeInfo.Name);
+					break;
+				case UE::Interchange::EAttributeTypes::IntRect:
+					MetaDataValue = GetValueAsString<FIntRect>(Node, UserAttributeInfo.Name);
+					break;
+				case UE::Interchange::EAttributeTypes::LinearColor:
+					MetaDataValue = GetValueAsString<FLinearColor>(Node, UserAttributeInfo.Name);
+					break;
+				case UE::Interchange::EAttributeTypes::Name:
+					MetaDataValue = GetValueAsString<FName>(Node, UserAttributeInfo.Name);
+					break;
+				case UE::Interchange::EAttributeTypes::RandomStream:
+					MetaDataValue = GetValueAsString<FRandomStream>(Node, UserAttributeInfo.Name);
+					break;
+				case UE::Interchange::EAttributeTypes::String:
+					MetaDataValue = GetValueAsString<FString>(Node, UserAttributeInfo.Name);
+					break;
+				case UE::Interchange::EAttributeTypes::Timespan:
+					MetaDataValue = GetValueAsString<FTimespan>(Node, UserAttributeInfo.Name);
+					break;
+				case UE::Interchange::EAttributeTypes::TwoVectors:
+					MetaDataValue = GetValueAsString<FTwoVectors>(Node, UserAttributeInfo.Name);
+					break;
 				case UE::Interchange::EAttributeTypes::UInt8:
-				{
-					uint8 Value = 0;
-					if (UInterchangeUserDefinedAttributesAPI::GetUserDefinedAttribute(Node, UserAttributeInfo.Name, Value, PayloadKey))
-					{
-						MetaDataValue = UE::Interchange::AttributeValueToString(Value);
-					}
-				}
-				break;
+					MetaDataValue = GetValueAsString<uint8>(Node, UserAttributeInfo.Name);
+					break;
 				case UE::Interchange::EAttributeTypes::UInt16:
-				{
-					uint16 Value = 0;
-					if (UInterchangeUserDefinedAttributesAPI::GetUserDefinedAttribute(Node, UserAttributeInfo.Name, Value, PayloadKey))
-					{
-						MetaDataValue = UE::Interchange::AttributeValueToString(Value);
-					}
-				}
-				break;
+					MetaDataValue = GetValueAsString<uint16>(Node, UserAttributeInfo.Name);
+					break;
 				case UE::Interchange::EAttributeTypes::UInt32:
-				{
-					uint32 Value = 0;
-					if (UInterchangeUserDefinedAttributesAPI::GetUserDefinedAttribute(Node, UserAttributeInfo.Name, Value, PayloadKey))
-					{
-						MetaDataValue = UE::Interchange::AttributeValueToString(Value);
-					}
-				}
-				break;
+					MetaDataValue = GetValueAsString<uint32>(Node, UserAttributeInfo.Name);
+					break;
 				case UE::Interchange::EAttributeTypes::UInt64:
-				{
-					uint64 Value = 0;
-					if (UInterchangeUserDefinedAttributesAPI::GetUserDefinedAttribute(Node, UserAttributeInfo.Name, Value, PayloadKey))
-					{
-						MetaDataValue = UE::Interchange::AttributeValueToString(Value);
-					}
-				}
-				break;
-				case UE::Interchange::EAttributeTypes::Float:
-				{
-					float Value = 0.0f;
-					if (UInterchangeUserDefinedAttributesAPI::GetUserDefinedAttribute(Node, UserAttributeInfo.Name, Value, PayloadKey))
-					{
-						MetaDataValue = UE::Interchange::AttributeValueToString(Value);
-					}
-				}
-				break;
+					MetaDataValue = GetValueAsString<uint64>(Node, UserAttributeInfo.Name);
+					break;
+				case UE::Interchange::EAttributeTypes::Vector2d:
+					MetaDataValue = GetValueAsString<FVector2d>(Node, UserAttributeInfo.Name);
+					break;
+				case UE::Interchange::EAttributeTypes::IntPoint:
+					MetaDataValue = GetValueAsString<FIntPoint>(Node, UserAttributeInfo.Name);
+					break;
+				case UE::Interchange::EAttributeTypes::IntVector:
+					MetaDataValue = GetValueAsString<FIntVector>(Node, UserAttributeInfo.Name);
+					break;
+				case UE::Interchange::EAttributeTypes::Vector2DHalf:
+					MetaDataValue = GetValueAsString<FVector2DHalf>(Node, UserAttributeInfo.Name);
+					break;
 				case UE::Interchange::EAttributeTypes::Float16:
-				{
-					FFloat16 Value = 0.0f;
-					if (UInterchangeUserDefinedAttributesAPI::GetUserDefinedAttribute(Node, UserAttributeInfo.Name, Value, PayloadKey))
-					{
-						MetaDataValue = UE::Interchange::AttributeValueToString(Value);
-					}
-				}
-				break;
-				case UE::Interchange::EAttributeTypes::Vector2f:
-				{
-					FVector2f Value(0.0f);
-					if (UInterchangeUserDefinedAttributesAPI::GetUserDefinedAttribute(Node, UserAttributeInfo.Name, Value, PayloadKey))
-					{
-						MetaDataValue = UE::Interchange::AttributeValueToString(Value);
-					}
-				}
-				break;
+					MetaDataValue = GetValueAsString<FFloat16>(Node, UserAttributeInfo.Name);
+					break;
+				case UE::Interchange::EAttributeTypes::OrientedBox:
+					MetaDataValue = GetValueAsString<FOrientedBox>(Node, UserAttributeInfo.Name);
+					break;
+				case UE::Interchange::EAttributeTypes::FrameNumber:
+					MetaDataValue = GetValueAsString<FFrameNumber>(Node, UserAttributeInfo.Name);
+					break;
+				case UE::Interchange::EAttributeTypes::FrameRate:
+					MetaDataValue = GetValueAsString<FFrameRate>(Node, UserAttributeInfo.Name);
+					break;
+				case UE::Interchange::EAttributeTypes::FrameTime:
+					MetaDataValue = GetValueAsString<FFrameTime>(Node, UserAttributeInfo.Name);
+					break;
+				case UE::Interchange::EAttributeTypes::SoftObjectPath:
+					MetaDataValue = GetValueAsString<FSoftObjectPath>(Node, UserAttributeInfo.Name);
+					break;
+				case UE::Interchange::EAttributeTypes::Matrix44f:
+					MetaDataValue = GetValueAsString<FMatrix44f>(Node, UserAttributeInfo.Name);
+					break;
+				case UE::Interchange::EAttributeTypes::Matrix44d:
+					MetaDataValue = GetValueAsString<FMatrix44d>(Node, UserAttributeInfo.Name);
+					break;
+				case UE::Interchange::EAttributeTypes::Plane4f:
+					MetaDataValue = GetValueAsString<FPlane4f>(Node, UserAttributeInfo.Name);
+					break;
+				case UE::Interchange::EAttributeTypes::Plane4d:
+					MetaDataValue = GetValueAsString<FPlane4d>(Node, UserAttributeInfo.Name);
+					break;
+				case UE::Interchange::EAttributeTypes::Quat4f:
+					MetaDataValue = GetValueAsString<FQuat4f>(Node, UserAttributeInfo.Name);
+					break;
+				case UE::Interchange::EAttributeTypes::Quat4d:
+					MetaDataValue = GetValueAsString<FQuat4d>(Node, UserAttributeInfo.Name);
+					break;
+				case UE::Interchange::EAttributeTypes::Rotator3f:
+					MetaDataValue = GetValueAsString<FRotator3f>(Node, UserAttributeInfo.Name);
+					break;
+				case UE::Interchange::EAttributeTypes::Rotator3d:
+					MetaDataValue = GetValueAsString<FRotator3d>(Node, UserAttributeInfo.Name);
+					break;
+				case UE::Interchange::EAttributeTypes::Transform3f:
+					MetaDataValue = GetValueAsString<FTransform3f>(Node, UserAttributeInfo.Name);
+					break;
+				case UE::Interchange::EAttributeTypes::Transform3d:
+					MetaDataValue = GetValueAsString<FTransform3d>(Node, UserAttributeInfo.Name);
+					break;
 				case UE::Interchange::EAttributeTypes::Vector3f:
-				{
-					FVector3f Value(0.0f);
-					if (UInterchangeUserDefinedAttributesAPI::GetUserDefinedAttribute(Node, UserAttributeInfo.Name, Value, PayloadKey))
-					{
-						MetaDataValue = UE::Interchange::AttributeValueToString(Value);
-					}
-				}
-				break;
-				case UE::Interchange::EAttributeTypes::Vector4f:
-				{
-					FVector4f Value(0.0f);
-					if (UInterchangeUserDefinedAttributesAPI::GetUserDefinedAttribute(Node, UserAttributeInfo.Name, Value, PayloadKey))
-					{
-						MetaDataValue = UE::Interchange::AttributeValueToString(Value);
-					}
-				}
-				break;
-				case UE::Interchange::EAttributeTypes::Double:
-				{
-					double Value = 0.0;
-					if (UInterchangeUserDefinedAttributesAPI::GetUserDefinedAttribute(Node, UserAttributeInfo.Name, Value, PayloadKey))
-					{
-						MetaDataValue = UE::Interchange::AttributeValueToString(Value);
-					}
-				}
-				break;
-				case UE::Interchange::EAttributeTypes::Vector2d:
-				{
-					FVector2D Value(0.0);
-					if (UInterchangeUserDefinedAttributesAPI::GetUserDefinedAttribute(Node, UserAttributeInfo.Name, Value, PayloadKey))
-					{
-						MetaDataValue = UE::Interchange::AttributeValueToString(Value);
-					}
-				}
-				break;
+					MetaDataValue = GetValueAsString<FVector3f>(Node, UserAttributeInfo.Name);
+					break;
 				case UE::Interchange::EAttributeTypes::Vector3d:
-				{
-					FVector3d Value(0.0);
-					if (UInterchangeUserDefinedAttributesAPI::GetUserDefinedAttribute(Node, UserAttributeInfo.Name, Value, PayloadKey))
-					{
-						MetaDataValue = UE::Interchange::AttributeValueToString(Value);
-					}
-				}
-				break;
+					MetaDataValue = GetValueAsString<FVector3d>(Node, UserAttributeInfo.Name);
+					break;
+				case UE::Interchange::EAttributeTypes::Vector2f:
+					MetaDataValue = GetValueAsString<FVector2f>(Node, UserAttributeInfo.Name);
+					break;
+				case UE::Interchange::EAttributeTypes::Vector4f:
+					MetaDataValue = GetValueAsString<FVector4f>(Node, UserAttributeInfo.Name);
+					break;
 				case UE::Interchange::EAttributeTypes::Vector4d:
-				{
-					FVector4d Value(0.0);
-					if (UInterchangeUserDefinedAttributesAPI::GetUserDefinedAttribute(Node, UserAttributeInfo.Name, Value, PayloadKey))
-					{
-						MetaDataValue = UE::Interchange::AttributeValueToString(Value);
-					}
-				}
-				break;
-				case UE::Interchange::EAttributeTypes::String:
-				{
-					FString Value;
-					if (UInterchangeUserDefinedAttributesAPI::GetUserDefinedAttribute(Node, UserAttributeInfo.Name, Value, PayloadKey))
-					{
-						MetaDataValue = UE::Interchange::AttributeValueToString(Value);
-					}
-				}
-				break;
+					MetaDataValue = GetValueAsString<FVector4d>(Node, UserAttributeInfo.Name);
+					break;
+				case UE::Interchange::EAttributeTypes::Box2f:
+					MetaDataValue = GetValueAsString<FBox2f>(Node, UserAttributeInfo.Name);
+					break;
+				case UE::Interchange::EAttributeTypes::Box2D:
+					MetaDataValue = GetValueAsString<FBox2D>(Node, UserAttributeInfo.Name);
+					break;
+				case UE::Interchange::EAttributeTypes::Box3f:
+					MetaDataValue = GetValueAsString<FBox3f>(Node, UserAttributeInfo.Name);
+					break;
+				case UE::Interchange::EAttributeTypes::Box3d:
+					MetaDataValue = GetValueAsString<FBox3d>(Node, UserAttributeInfo.Name);
+					break;
+				case UE::Interchange::EAttributeTypes::BoxSphereBounds3f:
+					MetaDataValue = GetValueAsString<FBoxSphereBounds3f>(Node, UserAttributeInfo.Name);
+					break;
+				case UE::Interchange::EAttributeTypes::BoxSphereBounds3d:
+					MetaDataValue = GetValueAsString<FBoxSphereBounds3d>(Node, UserAttributeInfo.Name);
+					break;
+				case UE::Interchange::EAttributeTypes::Sphere3f:
+					MetaDataValue = GetValueAsString<FSphere3f>(Node, UserAttributeInfo.Name);
+					break;
+				case UE::Interchange::EAttributeTypes::Sphere3d:
+					MetaDataValue = GetValueAsString<FSphere3d>(Node, UserAttributeInfo.Name);
+					break;
+				case UE::Interchange::EAttributeTypes::BoolArray:
+					MetaDataValue = GetValueAsString<TArray<bool>>(Node, UserAttributeInfo.Name);
+					break;
+				case UE::Interchange::EAttributeTypes::ColorArray:
+					MetaDataValue = GetValueAsString<TArray<FColor>>(Node, UserAttributeInfo.Name);
+					break;
+				case UE::Interchange::EAttributeTypes::DateTimeArray:
+					MetaDataValue = GetValueAsString<TArray<FDateTime>>(Node, UserAttributeInfo.Name);
+					break;
+				case UE::Interchange::EAttributeTypes::DoubleArray:
+					MetaDataValue = GetValueAsString<TArray<double>>(Node, UserAttributeInfo.Name);
+					break;
+				case UE::Interchange::EAttributeTypes::EnumArray:
+					MetaDataValue = GetValueAsString<TArray<uint8>>(Node, UserAttributeInfo.Name);
+					break;
+				case UE::Interchange::EAttributeTypes::FloatArray:
+					MetaDataValue = GetValueAsString<TArray<float>>(Node, UserAttributeInfo.Name);
+					break;
+				case UE::Interchange::EAttributeTypes::GuidArray:
+					MetaDataValue = GetValueAsString<TArray<FGuid>>(Node, UserAttributeInfo.Name);
+					break;
+				case UE::Interchange::EAttributeTypes::Int8Array:
+					MetaDataValue = GetValueAsString<TArray<uint8>>(Node, UserAttributeInfo.Name);
+					break;
+				case UE::Interchange::EAttributeTypes::Int16Array:
+					MetaDataValue = GetValueAsString<TArray<int16>>(Node, UserAttributeInfo.Name);
+					break;
+				case UE::Interchange::EAttributeTypes::Int32Array:
+					MetaDataValue = GetValueAsString<TArray<int32>>(Node, UserAttributeInfo.Name);
+					break;
+				case UE::Interchange::EAttributeTypes::Int64Array:
+					MetaDataValue = GetValueAsString<TArray<int64>>(Node, UserAttributeInfo.Name);
+					break;
+				case UE::Interchange::EAttributeTypes::IntRectArray:
+					MetaDataValue = GetValueAsString<TArray<FIntRect>>(Node, UserAttributeInfo.Name);
+					break;
+				case UE::Interchange::EAttributeTypes::LinearColorArray:
+					MetaDataValue = GetValueAsString<TArray<FLinearColor>>(Node, UserAttributeInfo.Name);
+					break;
+				case UE::Interchange::EAttributeTypes::NameArray:
+					MetaDataValue = GetValueAsString<TArray<FName>>(Node, UserAttributeInfo.Name);
+					break;
+				case UE::Interchange::EAttributeTypes::RandomStreamArray:
+					MetaDataValue = GetValueAsString<TArray<FRandomStream>>(Node, UserAttributeInfo.Name);
+					break;
+				case UE::Interchange::EAttributeTypes::StringArray:
+					MetaDataValue = GetValueAsString<TArray<FString>>(Node, UserAttributeInfo.Name);
+					break;
+				case UE::Interchange::EAttributeTypes::TimespanArray:
+					MetaDataValue = GetValueAsString<TArray<FTimespan>>(Node, UserAttributeInfo.Name);
+					break;
+				case UE::Interchange::EAttributeTypes::TwoVectorsArray:
+					MetaDataValue = GetValueAsString<TArray<FTwoVectors>>(Node, UserAttributeInfo.Name);
+					break;
+				case UE::Interchange::EAttributeTypes::ByteArray:
+					MetaDataValue = GetValueAsString<TArray<uint8>>(Node, UserAttributeInfo.Name);
+					break;
+				case UE::Interchange::EAttributeTypes::ByteArray64:
+					MetaDataValue = GetValueAsString<TArray64<uint8>>(Node, UserAttributeInfo.Name);
+					break;
+				case UE::Interchange::EAttributeTypes::UInt16Array:
+					MetaDataValue = GetValueAsString<TArray<uint16>>(Node, UserAttributeInfo.Name);
+					break;
+				case UE::Interchange::EAttributeTypes::UInt32Array:
+					MetaDataValue = GetValueAsString<TArray<uint32>>(Node, UserAttributeInfo.Name);
+					break;
+				case UE::Interchange::EAttributeTypes::UInt64Array:
+					MetaDataValue = GetValueAsString<TArray<uint64>>(Node, UserAttributeInfo.Name);
+					break;
+				case UE::Interchange::EAttributeTypes::Vector2dArray:
+					MetaDataValue = GetValueAsString<TArray<FVector2d>>(Node, UserAttributeInfo.Name);
+					break;
+				case UE::Interchange::EAttributeTypes::IntPointArray:
+					MetaDataValue = GetValueAsString<TArray<FIntPoint>>(Node, UserAttributeInfo.Name);
+					break;
+				case UE::Interchange::EAttributeTypes::IntVectorArray:
+					MetaDataValue = GetValueAsString<TArray<FIntVector>>(Node, UserAttributeInfo.Name);
+					break;
+				case UE::Interchange::EAttributeTypes::Vector2DHalfArray:
+					MetaDataValue = GetValueAsString<TArray<FVector2DHalf>>(Node, UserAttributeInfo.Name);
+					break;
+				case UE::Interchange::EAttributeTypes::Float16Array:
+					MetaDataValue = GetValueAsString<TArray<FFloat16>>(Node, UserAttributeInfo.Name);
+					break;
+				case UE::Interchange::EAttributeTypes::OrientedBoxArray:
+					MetaDataValue = GetValueAsString<TArray<FOrientedBox>>(Node, UserAttributeInfo.Name);
+					break;
+				case UE::Interchange::EAttributeTypes::FrameNumberArray:
+					MetaDataValue = GetValueAsString<TArray<FFrameNumber>>(Node, UserAttributeInfo.Name);
+					break;
+				case UE::Interchange::EAttributeTypes::FrameRateArray:
+					MetaDataValue = GetValueAsString<TArray<FFrameRate>>(Node, UserAttributeInfo.Name);
+					break;
+				case UE::Interchange::EAttributeTypes::FrameTimeArray:
+					MetaDataValue = GetValueAsString<TArray<FFrameTime>>(Node, UserAttributeInfo.Name);
+					break;
+				case UE::Interchange::EAttributeTypes::SoftObjectPathArray:
+					MetaDataValue = GetValueAsString<TArray<FSoftObjectPath>>(Node, UserAttributeInfo.Name);
+					break;
+				case UE::Interchange::EAttributeTypes::Matrix44fArray:
+					MetaDataValue = GetValueAsString<TArray<FMatrix44f>>(Node, UserAttributeInfo.Name);
+					break;
+				case UE::Interchange::EAttributeTypes::Matrix44dArray:
+					MetaDataValue = GetValueAsString<TArray<FMatrix44d>>(Node, UserAttributeInfo.Name);
+					break;
+				case UE::Interchange::EAttributeTypes::Plane4fArray:
+					MetaDataValue = GetValueAsString<TArray<FPlane4f>>(Node, UserAttributeInfo.Name);
+					break;
+				case UE::Interchange::EAttributeTypes::Plane4dArray:
+					MetaDataValue = GetValueAsString<TArray<FPlane4d>>(Node, UserAttributeInfo.Name);
+					break;
+				case UE::Interchange::EAttributeTypes::Quat4fArray:
+					MetaDataValue = GetValueAsString<TArray<FQuat4f>>(Node, UserAttributeInfo.Name);
+					break;
+				case UE::Interchange::EAttributeTypes::Quat4dArray:
+					MetaDataValue = GetValueAsString<TArray<FQuat4d>>(Node, UserAttributeInfo.Name);
+					break;
+				case UE::Interchange::EAttributeTypes::Rotator3fArray:
+					MetaDataValue = GetValueAsString<TArray<FRotator3f>>(Node, UserAttributeInfo.Name);
+					break;
+				case UE::Interchange::EAttributeTypes::Rotator3dArray:
+					MetaDataValue = GetValueAsString<TArray<FRotator3d>>(Node, UserAttributeInfo.Name);
+					break;
+				case UE::Interchange::EAttributeTypes::Transform3fArray:
+					MetaDataValue = GetValueAsString<TArray<FTransform3f>>(Node, UserAttributeInfo.Name);
+					break;
+				case UE::Interchange::EAttributeTypes::Transform3dArray:
+					MetaDataValue = GetValueAsString<TArray<FTransform3d>>(Node, UserAttributeInfo.Name);
+					break;
+				case UE::Interchange::EAttributeTypes::Vector3fArray:
+					MetaDataValue = GetValueAsString<TArray<FVector3f>>(Node, UserAttributeInfo.Name);
+					break;
+				case UE::Interchange::EAttributeTypes::Vector3dArray:
+					MetaDataValue = GetValueAsString<TArray<FVector3d>>(Node, UserAttributeInfo.Name);
+					break;
+				case UE::Interchange::EAttributeTypes::Vector2fArray:
+					MetaDataValue = GetValueAsString<TArray<FVector2f>>(Node, UserAttributeInfo.Name);
+					break;
+				case UE::Interchange::EAttributeTypes::Vector4fArray:
+					MetaDataValue = GetValueAsString<TArray<FVector4f>>(Node, UserAttributeInfo.Name);
+					break;
+				case UE::Interchange::EAttributeTypes::Vector4dArray:
+					MetaDataValue = GetValueAsString<TArray<FVector4d>>(Node, UserAttributeInfo.Name);
+					break;
+				case UE::Interchange::EAttributeTypes::Box2fArray:
+					MetaDataValue = GetValueAsString<TArray<FBox2f>>(Node, UserAttributeInfo.Name);
+					break;
+				case UE::Interchange::EAttributeTypes::Box2DArray:
+					MetaDataValue = GetValueAsString<TArray<FBox2D>>(Node, UserAttributeInfo.Name);
+					break;
+				case UE::Interchange::EAttributeTypes::Box3fArray:
+					MetaDataValue = GetValueAsString<TArray<FBox3f>>(Node, UserAttributeInfo.Name);
+					break;
+				case UE::Interchange::EAttributeTypes::Box3dArray:
+					MetaDataValue = GetValueAsString<TArray<FBox3d>>(Node, UserAttributeInfo.Name);
+					break;
+				case UE::Interchange::EAttributeTypes::BoxSphereBounds3fArray:
+					MetaDataValue = GetValueAsString<TArray<FBoxSphereBounds3f>>(Node, UserAttributeInfo.Name);
+					break;
+				case UE::Interchange::EAttributeTypes::BoxSphereBounds3dArray:
+					MetaDataValue = GetValueAsString<TArray<FBoxSphereBounds3d>>(Node, UserAttributeInfo.Name);
+					break;
+				case UE::Interchange::EAttributeTypes::Sphere3fArray:
+					MetaDataValue = GetValueAsString<TArray<FSphere3f>>(Node, UserAttributeInfo.Name);
+					break;
+				case UE::Interchange::EAttributeTypes::Sphere3dArray:
+					MetaDataValue = GetValueAsString<TArray<FSphere3d>>(Node, UserAttributeInfo.Name);
+					break;
+				default:
+					ensureMsgf(false, TEXT("Unsupported EAttributeTypes value!"));
+					break;
 			}
 			if (MetaDataValue.IsSet())
 			{
diff --git a/Engine/Source/Runtime/Interchange/Core/Private/Nodes/InterchangeUserDefinedAttribute.cpp b/Engine/Source/Runtime/Interchange/Core/Private/Nodes/InterchangeUserDefinedAttribute.cpp
--- a/Engine/Source/Runtime/Interchange/Core/Private/Nodes/InterchangeUserDefinedAttribute.cpp
+++ b/Engine/Source/Runtime/Interchange/Core/Private/Nodes/InterchangeUserDefinedAttribute.cpp
@@ -4,6 +4,8 @@
 #include "Nodes/InterchangeFactoryBaseNode.h"
 #include "Nodes/InterchangeBaseNode.h"
 
+#include "Misc/FrameRate.h"
+
 #include UE_INLINE_GENERATED_CPP_BY_NAME(InterchangeUserDefinedAttribute)
 
 const FString UInterchangeUserDefinedAttributesAPI::UserDefinedAttributeBaseKey = TEXT("UserDefined_");
@@ -237,6 +239,24 @@ void UInterchangeUserDefinedAttributesAPI::DuplicateAllUserDefinedAttribute(cons
 				case UE::Interchange::EAttributeTypes::Bool:
 					CreateUserDefinedAttribute(InterchangeDestinationNode, DuplicateName, InterchangeSourceNode->GetAttributeChecked<bool>(AttributeKey.Key), PayloadKey, RequiresDelegate);
 					break;
+				case UE::Interchange::EAttributeTypes::Color:
+					CreateUserDefinedAttribute(InterchangeDestinationNode, DuplicateName, InterchangeSourceNode->GetAttributeChecked<FColor>(AttributeKey.Key), PayloadKey, RequiresDelegate);
+					break;
+				case UE::Interchange::EAttributeTypes::DateTime:
+					CreateUserDefinedAttribute(InterchangeDestinationNode, DuplicateName, InterchangeSourceNode->GetAttributeChecked<FDateTime>(AttributeKey.Key), PayloadKey, RequiresDelegate);
+					break;
+				case UE::Interchange::EAttributeTypes::Double:
+					CreateUserDefinedAttribute(InterchangeDestinationNode, DuplicateName, InterchangeSourceNode->GetAttributeChecked<double>(AttributeKey.Key), PayloadKey, RequiresDelegate);
+					break;
+				case UE::Interchange::EAttributeTypes::Enum:
+					CreateUserDefinedAttribute(InterchangeDestinationNode, DuplicateName, InterchangeSourceNode->GetAttributeChecked<uint8>(AttributeKey.Key), PayloadKey, RequiresDelegate);
+					break;
+				case UE::Interchange::EAttributeTypes::Float:
+					CreateUserDefinedAttribute(InterchangeDestinationNode, DuplicateName, InterchangeSourceNode->GetAttributeChecked<float>(AttributeKey.Key), PayloadKey, RequiresDelegate);
+					break;
+				case UE::Interchange::EAttributeTypes::Guid:
+					CreateUserDefinedAttribute(InterchangeDestinationNode, DuplicateName, InterchangeSourceNode->GetAttributeChecked<FGuid>(AttributeKey.Key), PayloadKey, RequiresDelegate);
+					break;
 				case UE::Interchange::EAttributeTypes::Int8:
 					CreateUserDefinedAttribute(InterchangeDestinationNode, DuplicateName, InterchangeSourceNode->GetAttributeChecked<int8>(AttributeKey.Key), PayloadKey, RequiresDelegate);
 					break;
@@ -249,6 +269,27 @@ void UInterchangeUserDefinedAttributesAPI::DuplicateAllUserDefinedAttribute(cons
 				case UE::Interchange::EAttributeTypes::Int64:
 					CreateUserDefinedAttribute(InterchangeDestinationNode, DuplicateName, InterchangeSourceNode->GetAttributeChecked<int64>(AttributeKey.Key), PayloadKey, RequiresDelegate);
 					break;
+				case UE::Interchange::EAttributeTypes::IntRect:
+					CreateUserDefinedAttribute(InterchangeDestinationNode, DuplicateName, InterchangeSourceNode->GetAttributeChecked<FIntRect>(AttributeKey.Key), PayloadKey, RequiresDelegate);
+					break;
+				case UE::Interchange::EAttributeTypes::LinearColor:
+					CreateUserDefinedAttribute(InterchangeDestinationNode, DuplicateName, InterchangeSourceNode->GetAttributeChecked<FLinearColor>(AttributeKey.Key), PayloadKey, RequiresDelegate);
+					break;
+				case UE::Interchange::EAttributeTypes::Name:
+					CreateUserDefinedAttribute(InterchangeDestinationNode, DuplicateName, InterchangeSourceNode->GetAttributeChecked<FName>(AttributeKey.Key), PayloadKey, RequiresDelegate);
+					break;
+				case UE::Interchange::EAttributeTypes::RandomStream:
+					CreateUserDefinedAttribute(InterchangeDestinationNode, DuplicateName, InterchangeSourceNode->GetAttributeChecked<FRandomStream>(AttributeKey.Key), PayloadKey, RequiresDelegate);
+					break;
+				case UE::Interchange::EAttributeTypes::String:
+					CreateUserDefinedAttribute(InterchangeDestinationNode, DuplicateName, InterchangeSourceNode->GetAttributeChecked<FString>(AttributeKey.Key), PayloadKey, RequiresDelegate);
+					break;
+				case UE::Interchange::EAttributeTypes::Timespan:
+					CreateUserDefinedAttribute(InterchangeDestinationNode, DuplicateName, InterchangeSourceNode->GetAttributeChecked<FTimespan>(AttributeKey.Key), PayloadKey, RequiresDelegate);
+					break;
+				case UE::Interchange::EAttributeTypes::TwoVectors:
+					CreateUserDefinedAttribute(InterchangeDestinationNode, DuplicateName, InterchangeSourceNode->GetAttributeChecked<FTwoVectors>(AttributeKey.Key), PayloadKey, RequiresDelegate);
+					break;
 				case UE::Interchange::EAttributeTypes::UInt8:
 					CreateUserDefinedAttribute(InterchangeDestinationNode, DuplicateName, InterchangeSourceNode->GetAttributeChecked<uint8>(AttributeKey.Key), PayloadKey, RequiresDelegate);
 					break;
@@ -261,35 +302,275 @@ void UInterchangeUserDefinedAttributesAPI::DuplicateAllUserDefinedAttribute(cons
 				case UE::Interchange::EAttributeTypes::UInt64:
 					CreateUserDefinedAttribute(InterchangeDestinationNode, DuplicateName, InterchangeSourceNode->GetAttributeChecked<uint64>(AttributeKey.Key), PayloadKey, RequiresDelegate);
 					break;
-				case UE::Interchange::EAttributeTypes::Float:
-					CreateUserDefinedAttribute(InterchangeDestinationNode, DuplicateName, InterchangeSourceNode->GetAttributeChecked<float>(AttributeKey.Key), PayloadKey, RequiresDelegate);
+				case UE::Interchange::EAttributeTypes::Vector2d:
+					CreateUserDefinedAttribute(InterchangeDestinationNode, DuplicateName, InterchangeSourceNode->GetAttributeChecked<FVector2d>(AttributeKey.Key), PayloadKey, RequiresDelegate);
+					break;
+				case UE::Interchange::EAttributeTypes::IntPoint:
+					CreateUserDefinedAttribute(InterchangeDestinationNode, DuplicateName, InterchangeSourceNode->GetAttributeChecked<FIntPoint>(AttributeKey.Key), PayloadKey, RequiresDelegate);
+					break;
+				case UE::Interchange::EAttributeTypes::IntVector:
+					CreateUserDefinedAttribute(InterchangeDestinationNode, DuplicateName, InterchangeSourceNode->GetAttributeChecked<FIntVector>(AttributeKey.Key), PayloadKey, RequiresDelegate);
+					break;
+				case UE::Interchange::EAttributeTypes::Vector2DHalf:
+					CreateUserDefinedAttribute(InterchangeDestinationNode, DuplicateName, InterchangeSourceNode->GetAttributeChecked<FVector2DHalf>(AttributeKey.Key), PayloadKey, RequiresDelegate);
 					break;
 				case UE::Interchange::EAttributeTypes::Float16:
 					CreateUserDefinedAttribute(InterchangeDestinationNode, DuplicateName, InterchangeSourceNode->GetAttributeChecked<FFloat16>(AttributeKey.Key), PayloadKey, RequiresDelegate);
 					break;
-				case UE::Interchange::EAttributeTypes::Vector2f:
-					CreateUserDefinedAttribute(InterchangeDestinationNode, DuplicateName, InterchangeSourceNode->GetAttributeChecked<FVector2f>(AttributeKey.Key), PayloadKey, RequiresDelegate);
+				case UE::Interchange::EAttributeTypes::OrientedBox:
+					CreateUserDefinedAttribute(InterchangeDestinationNode, DuplicateName, InterchangeSourceNode->GetAttributeChecked<FOrientedBox>(AttributeKey.Key), PayloadKey, RequiresDelegate);
 					break;
-				case UE::Interchange::EAttributeTypes::Vector3f:
-					CreateUserDefinedAttribute(InterchangeDestinationNode, DuplicateName, InterchangeSourceNode->GetAttributeChecked<FVector3f>(AttributeKey.Key), PayloadKey, RequiresDelegate);
+				case UE::Interchange::EAttributeTypes::FrameNumber:
+					CreateUserDefinedAttribute(InterchangeDestinationNode, DuplicateName, InterchangeSourceNode->GetAttributeChecked<FFrameNumber>(AttributeKey.Key), PayloadKey, RequiresDelegate);
 					break;
-				case UE::Interchange::EAttributeTypes::Vector4f:
-					CreateUserDefinedAttribute(InterchangeDestinationNode, DuplicateName, InterchangeSourceNode->GetAttributeChecked<FVector4f>(AttributeKey.Key), PayloadKey, RequiresDelegate);
+				case UE::Interchange::EAttributeTypes::FrameRate:
+					CreateUserDefinedAttribute(InterchangeDestinationNode, DuplicateName, InterchangeSourceNode->GetAttributeChecked<FFrameRate>(AttributeKey.Key), PayloadKey, RequiresDelegate);
 					break;
-				case UE::Interchange::EAttributeTypes::Double:
-					CreateUserDefinedAttribute(InterchangeDestinationNode, DuplicateName, InterchangeSourceNode->GetAttributeChecked<double>(AttributeKey.Key), PayloadKey, RequiresDelegate);
+				case UE::Interchange::EAttributeTypes::FrameTime:
+					CreateUserDefinedAttribute(InterchangeDestinationNode, DuplicateName, InterchangeSourceNode->GetAttributeChecked<FFrameTime>(AttributeKey.Key), PayloadKey, RequiresDelegate);
 					break;
-				case UE::Interchange::EAttributeTypes::Vector2d:
-					CreateUserDefinedAttribute(InterchangeDestinationNode, DuplicateName, InterchangeSourceNode->GetAttributeChecked<FVector2D>(AttributeKey.Key), PayloadKey, RequiresDelegate);
+				case UE::Interchange::EAttributeTypes::SoftObjectPath:
+					CreateUserDefinedAttribute(InterchangeDestinationNode, DuplicateName, InterchangeSourceNode->GetAttributeChecked<FSoftObjectPath>(AttributeKey.Key), PayloadKey, RequiresDelegate);
+					break;
+				case UE::Interchange::EAttributeTypes::Matrix44f:
+					CreateUserDefinedAttribute(InterchangeDestinationNode, DuplicateName, InterchangeSourceNode->GetAttributeChecked<FMatrix44f>(AttributeKey.Key), PayloadKey, RequiresDelegate);
+					break;
+				case UE::Interchange::EAttributeTypes::Matrix44d:
+					CreateUserDefinedAttribute(InterchangeDestinationNode, DuplicateName, InterchangeSourceNode->GetAttributeChecked<FMatrix44d>(AttributeKey.Key), PayloadKey, RequiresDelegate);
+					break;
+				case UE::Interchange::EAttributeTypes::Plane4f:
+					CreateUserDefinedAttribute(InterchangeDestinationNode, DuplicateName, InterchangeSourceNode->GetAttributeChecked<FPlane4f>(AttributeKey.Key), PayloadKey, RequiresDelegate);
+					break;
+				case UE::Interchange::EAttributeTypes::Plane4d:
+					CreateUserDefinedAttribute(InterchangeDestinationNode, DuplicateName, InterchangeSourceNode->GetAttributeChecked<FPlane4d>(AttributeKey.Key), PayloadKey, RequiresDelegate);
+					break;
+				case UE::Interchange::EAttributeTypes::Quat4f:
+					CreateUserDefinedAttribute(InterchangeDestinationNode, DuplicateName, InterchangeSourceNode->GetAttributeChecked<FQuat4f>(AttributeKey.Key), PayloadKey, RequiresDelegate);
+					break;
+				case UE::Interchange::EAttributeTypes::Quat4d:
+					CreateUserDefinedAttribute(InterchangeDestinationNode, DuplicateName, InterchangeSourceNode->GetAttributeChecked<FQuat4d>(AttributeKey.Key), PayloadKey, RequiresDelegate);
+					break;
+				case UE::Interchange::EAttributeTypes::Rotator3f:
+					CreateUserDefinedAttribute(InterchangeDestinationNode, DuplicateName, InterchangeSourceNode->GetAttributeChecked<FRotator3f>(AttributeKey.Key), PayloadKey, RequiresDelegate);
+					break;
+				case UE::Interchange::EAttributeTypes::Rotator3d:
+					CreateUserDefinedAttribute(InterchangeDestinationNode, DuplicateName, InterchangeSourceNode->GetAttributeChecked<FRotator3d>(AttributeKey.Key), PayloadKey, RequiresDelegate);
+					break;
+				case UE::Interchange::EAttributeTypes::Transform3f:
+					CreateUserDefinedAttribute(InterchangeDestinationNode, DuplicateName, InterchangeSourceNode->GetAttributeChecked<FTransform3f>(AttributeKey.Key), PayloadKey, RequiresDelegate);
+					break;
+				case UE::Interchange::EAttributeTypes::Transform3d:
+					CreateUserDefinedAttribute(InterchangeDestinationNode, DuplicateName, InterchangeSourceNode->GetAttributeChecked<FTransform3d>(AttributeKey.Key), PayloadKey, RequiresDelegate);
+					break;
+				case UE::Interchange::EAttributeTypes::Vector3f:
+					CreateUserDefinedAttribute(InterchangeDestinationNode, DuplicateName, InterchangeSourceNode->GetAttributeChecked<FVector3f>(AttributeKey.Key), PayloadKey, RequiresDelegate);
 					break;
 				case UE::Interchange::EAttributeTypes::Vector3d:
 					CreateUserDefinedAttribute(InterchangeDestinationNode, DuplicateName, InterchangeSourceNode->GetAttributeChecked<FVector3d>(AttributeKey.Key), PayloadKey, RequiresDelegate);
 					break;
+				case UE::Interchange::EAttributeTypes::Vector2f:
+					CreateUserDefinedAttribute(InterchangeDestinationNode, DuplicateName, InterchangeSourceNode->GetAttributeChecked<FVector2f>(AttributeKey.Key), PayloadKey, RequiresDelegate);
+					break;
+				case UE::Interchange::EAttributeTypes::Vector4f:
+					CreateUserDefinedAttribute(InterchangeDestinationNode, DuplicateName, InterchangeSourceNode->GetAttributeChecked<FVector4f>(AttributeKey.Key), PayloadKey, RequiresDelegate);
+					break;
 				case UE::Interchange::EAttributeTypes::Vector4d:
 					CreateUserDefinedAttribute(InterchangeDestinationNode, DuplicateName, InterchangeSourceNode->GetAttributeChecked<FVector4d>(AttributeKey.Key), PayloadKey, RequiresDelegate);
 					break;
-				case UE::Interchange::EAttributeTypes::String:
-					CreateUserDefinedAttribute(InterchangeDestinationNode, DuplicateName, InterchangeSourceNode->GetAttributeChecked<FString>(AttributeKey.Key), PayloadKey, RequiresDelegate);
+				case UE::Interchange::EAttributeTypes::Box2f:
+					CreateUserDefinedAttribute(InterchangeDestinationNode, DuplicateName, InterchangeSourceNode->GetAttributeChecked<FBox2f>(AttributeKey.Key), PayloadKey, RequiresDelegate);
+					break;
+				case UE::Interchange::EAttributeTypes::Box2D:
+					CreateUserDefinedAttribute(InterchangeDestinationNode, DuplicateName, InterchangeSourceNode->GetAttributeChecked<FBox2D>(AttributeKey.Key), PayloadKey, RequiresDelegate);
+					break;
+				case UE::Interchange::EAttributeTypes::Box3f:
+					CreateUserDefinedAttribute(InterchangeDestinationNode, DuplicateName, InterchangeSourceNode->GetAttributeChecked<FBox3f>(AttributeKey.Key), PayloadKey, RequiresDelegate);
+					break;
+				case UE::Interchange::EAttributeTypes::Box3d:
+					CreateUserDefinedAttribute(InterchangeDestinationNode, DuplicateName, InterchangeSourceNode->GetAttributeChecked<FBox3d>(AttributeKey.Key), PayloadKey, RequiresDelegate);
+					break;
+				case UE::Interchange::EAttributeTypes::BoxSphereBounds3f:
+					CreateUserDefinedAttribute(InterchangeDestinationNode, DuplicateName, InterchangeSourceNode->GetAttributeChecked<FBoxSphereBounds3f>(AttributeKey.Key), PayloadKey, RequiresDelegate);
+					break;
+				case UE::Interchange::EAttributeTypes::BoxSphereBounds3d:
+					CreateUserDefinedAttribute(InterchangeDestinationNode, DuplicateName, InterchangeSourceNode->GetAttributeChecked<FBoxSphereBounds3d>(AttributeKey.Key), PayloadKey, RequiresDelegate);
+					break;
+				case UE::Interchange::EAttributeTypes::Sphere3f:
+					CreateUserDefinedAttribute(InterchangeDestinationNode, DuplicateName, InterchangeSourceNode->GetAttributeChecked<FSphere3f>(AttributeKey.Key), PayloadKey, RequiresDelegate);
+					break;
+				case UE::Interchange::EAttributeTypes::Sphere3d:
+					CreateUserDefinedAttribute(InterchangeDestinationNode, DuplicateName, InterchangeSourceNode->GetAttributeChecked<FSphere3d>(AttributeKey.Key), PayloadKey, RequiresDelegate);
+					break;
+				case UE::Interchange::EAttributeTypes::BoolArray:
+					CreateUserDefinedAttribute(InterchangeDestinationNode, DuplicateName, InterchangeSourceNode->GetAttributeChecked<TArray<bool>>(AttributeKey.Key), PayloadKey, RequiresDelegate);
+					break;
+				case UE::Interchange::EAttributeTypes::ColorArray:
+					CreateUserDefinedAttribute(InterchangeDestinationNode, DuplicateName, InterchangeSourceNode->GetAttributeChecked<TArray<FColor>>(AttributeKey.Key), PayloadKey, RequiresDelegate);
+					break;
+				case UE::Interchange::EAttributeTypes::DateTimeArray:
+					CreateUserDefinedAttribute(InterchangeDestinationNode, DuplicateName, InterchangeSourceNode->GetAttributeChecked<TArray<FDateTime>>(AttributeKey.Key), PayloadKey, RequiresDelegate);
+					break;
+				case UE::Interchange::EAttributeTypes::DoubleArray:
+					CreateUserDefinedAttribute(InterchangeDestinationNode, DuplicateName, InterchangeSourceNode->GetAttributeChecked<TArray<double>>(AttributeKey.Key), PayloadKey, RequiresDelegate);
+					break;
+				case UE::Interchange::EAttributeTypes::EnumArray:
+					CreateUserDefinedAttribute(InterchangeDestinationNode, DuplicateName, InterchangeSourceNode->GetAttributeChecked<TArray<uint8>>(AttributeKey.Key), PayloadKey, RequiresDelegate);
+					break;
+				case UE::Interchange::EAttributeTypes::FloatArray:
+					CreateUserDefinedAttribute(InterchangeDestinationNode, DuplicateName, InterchangeSourceNode->GetAttributeChecked<TArray<float>>(AttributeKey.Key), PayloadKey, RequiresDelegate);
+					break;
+				case UE::Interchange::EAttributeTypes::GuidArray:
+					CreateUserDefinedAttribute(InterchangeDestinationNode, DuplicateName, InterchangeSourceNode->GetAttributeChecked<TArray<FGuid>>(AttributeKey.Key), PayloadKey, RequiresDelegate);
+					break;
+				case UE::Interchange::EAttributeTypes::Int8Array:
+					CreateUserDefinedAttribute(InterchangeDestinationNode, DuplicateName, InterchangeSourceNode->GetAttributeChecked<TArray<uint8>>(AttributeKey.Key), PayloadKey, RequiresDelegate);
+					break;
+				case UE::Interchange::EAttributeTypes::Int16Array:
+					CreateUserDefinedAttribute(InterchangeDestinationNode, DuplicateName, InterchangeSourceNode->GetAttributeChecked<TArray<int16>>(AttributeKey.Key), PayloadKey, RequiresDelegate);
+					break;
+				case UE::Interchange::EAttributeTypes::Int32Array:
+					CreateUserDefinedAttribute(InterchangeDestinationNode, DuplicateName, InterchangeSourceNode->GetAttributeChecked<TArray<int32>>(AttributeKey.Key), PayloadKey, RequiresDelegate);
+					break;
+				case UE::Interchange::EAttributeTypes::Int64Array:
+					CreateUserDefinedAttribute(InterchangeDestinationNode, DuplicateName, InterchangeSourceNode->GetAttributeChecked<TArray<int64>>(AttributeKey.Key), PayloadKey, RequiresDelegate);
+					break;
+				case UE::Interchange::EAttributeTypes::IntRectArray:
+					CreateUserDefinedAttribute(InterchangeDestinationNode, DuplicateName, InterchangeSourceNode->GetAttributeChecked<TArray<FIntRect>>(AttributeKey.Key), PayloadKey, RequiresDelegate);
+					break;
+				case UE::Interchange::EAttributeTypes::LinearColorArray:
+					CreateUserDefinedAttribute(InterchangeDestinationNode, DuplicateName, InterchangeSourceNode->GetAttributeChecked<TArray<FLinearColor>>(AttributeKey.Key), PayloadKey, RequiresDelegate);
+					break;
+				case UE::Interchange::EAttributeTypes::NameArray:
+					CreateUserDefinedAttribute(InterchangeDestinationNode, DuplicateName, InterchangeSourceNode->GetAttributeChecked<TArray<FName>>(AttributeKey.Key), PayloadKey, RequiresDelegate);
+					break;
+				case UE::Interchange::EAttributeTypes::RandomStreamArray:
+					CreateUserDefinedAttribute(InterchangeDestinationNode, DuplicateName, InterchangeSourceNode->GetAttributeChecked<TArray<FRandomStream>>(AttributeKey.Key), PayloadKey, RequiresDelegate);
+					break;
+				case UE::Interchange::EAttributeTypes::StringArray:
+					CreateUserDefinedAttribute(InterchangeDestinationNode, DuplicateName, InterchangeSourceNode->GetAttributeChecked<TArray<FString>>(AttributeKey.Key), PayloadKey, RequiresDelegate);
+					break;
+				case UE::Interchange::EAttributeTypes::TimespanArray:
+					CreateUserDefinedAttribute(InterchangeDestinationNode, DuplicateName, InterchangeSourceNode->GetAttributeChecked<TArray<FTimespan>>(AttributeKey.Key), PayloadKey, RequiresDelegate);
+					break;
+				case UE::Interchange::EAttributeTypes::TwoVectorsArray:
+					CreateUserDefinedAttribute(InterchangeDestinationNode, DuplicateName, InterchangeSourceNode->GetAttributeChecked<TArray<FTwoVectors>>(AttributeKey.Key), PayloadKey, RequiresDelegate);
+					break;
+				case UE::Interchange::EAttributeTypes::ByteArray:
+					CreateUserDefinedAttribute(InterchangeDestinationNode, DuplicateName, InterchangeSourceNode->GetAttributeChecked<TArray<uint8>>(AttributeKey.Key), PayloadKey, RequiresDelegate);
+					break;
+				case UE::Interchange::EAttributeTypes::ByteArray64:
+					CreateUserDefinedAttribute(InterchangeDestinationNode, DuplicateName, InterchangeSourceNode->GetAttributeChecked<TArray64<uint8>>(AttributeKey.Key), PayloadKey, RequiresDelegate);
+					break;
+				case UE::Interchange::EAttributeTypes::UInt16Array:
+					CreateUserDefinedAttribute(InterchangeDestinationNode, DuplicateName, InterchangeSourceNode->GetAttributeChecked<TArray<uint16>>(AttributeKey.Key), PayloadKey, RequiresDelegate);
+					break;
+				case UE::Interchange::EAttributeTypes::UInt32Array:
+					CreateUserDefinedAttribute(InterchangeDestinationNode, DuplicateName, InterchangeSourceNode->GetAttributeChecked<TArray<uint32>>(AttributeKey.Key), PayloadKey, RequiresDelegate);
+					break;
+				case UE::Interchange::EAttributeTypes::UInt64Array:
+					CreateUserDefinedAttribute(InterchangeDestinationNode, DuplicateName, InterchangeSourceNode->GetAttributeChecked<TArray<uint64>>(AttributeKey.Key), PayloadKey, RequiresDelegate);
+					break;
+				case UE::Interchange::EAttributeTypes::Vector2dArray:
+					CreateUserDefinedAttribute(InterchangeDestinationNode, DuplicateName, InterchangeSourceNode->GetAttributeChecked<TArray<FVector2d>>(AttributeKey.Key), PayloadKey, RequiresDelegate);
+					break;
+				case UE::Interchange::EAttributeTypes::IntPointArray:
+					CreateUserDefinedAttribute(InterchangeDestinationNode, DuplicateName, InterchangeSourceNode->GetAttributeChecked<TArray<FIntPoint>>(AttributeKey.Key), PayloadKey, RequiresDelegate);
+					break;
+				case UE::Interchange::EAttributeTypes::IntVectorArray:
+					CreateUserDefinedAttribute(InterchangeDestinationNode, DuplicateName, InterchangeSourceNode->GetAttributeChecked<TArray<FIntVector>>(AttributeKey.Key), PayloadKey, RequiresDelegate);
+					break;
+				case UE::Interchange::EAttributeTypes::Vector2DHalfArray:
+					CreateUserDefinedAttribute(InterchangeDestinationNode, DuplicateName, InterchangeSourceNode->GetAttributeChecked<TArray<FVector2DHalf>>(AttributeKey.Key), PayloadKey, RequiresDelegate);
+					break;
+				case UE::Interchange::EAttributeTypes::Float16Array:
+					CreateUserDefinedAttribute(InterchangeDestinationNode, DuplicateName, InterchangeSourceNode->GetAttributeChecked<TArray<FFloat16>>(AttributeKey.Key), PayloadKey, RequiresDelegate);
+					break;
+				case UE::Interchange::EAttributeTypes::OrientedBoxArray:
+					CreateUserDefinedAttribute(InterchangeDestinationNode, DuplicateName, InterchangeSourceNode->GetAttributeChecked<TArray<FOrientedBox>>(AttributeKey.Key), PayloadKey, RequiresDelegate);
+					break;
+				case UE::Interchange::EAttributeTypes::FrameNumberArray:
+					CreateUserDefinedAttribute(InterchangeDestinationNode, DuplicateName, InterchangeSourceNode->GetAttributeChecked<TArray<FFrameNumber>>(AttributeKey.Key), PayloadKey, RequiresDelegate);
+					break;
+				case UE::Interchange::EAttributeTypes::FrameRateArray:
+					CreateUserDefinedAttribute(InterchangeDestinationNode, DuplicateName, InterchangeSourceNode->GetAttributeChecked<TArray<FFrameRate>>(AttributeKey.Key), PayloadKey, RequiresDelegate);
+					break;
+				case UE::Interchange::EAttributeTypes::FrameTimeArray:
+					CreateUserDefinedAttribute(InterchangeDestinationNode, DuplicateName, InterchangeSourceNode->GetAttributeChecked<TArray<FFrameTime>>(AttributeKey.Key), PayloadKey, RequiresDelegate);
+					break;
+				case UE::Interchange::EAttributeTypes::SoftObjectPathArray:
+					CreateUserDefinedAttribute(InterchangeDestinationNode, DuplicateName, InterchangeSourceNode->GetAttributeChecked<TArray<FSoftObjectPath>>(AttributeKey.Key), PayloadKey, RequiresDelegate);
+					break;
+				case UE::Interchange::EAttributeTypes::Matrix44fArray:
+					CreateUserDefinedAttribute(InterchangeDestinationNode, DuplicateName, InterchangeSourceNode->GetAttributeChecked<TArray<FMatrix44f>>(AttributeKey.Key), PayloadKey, RequiresDelegate);
+					break;
+				case UE::Interchange::EAttributeTypes::Matrix44dArray:
+					CreateUserDefinedAttribute(InterchangeDestinationNode, DuplicateName, InterchangeSourceNode->GetAttributeChecked<TArray<FMatrix44d>>(AttributeKey.Key), PayloadKey, RequiresDelegate);
+					break;
+				case UE::Interchange::EAttributeTypes::Plane4fArray:
+					CreateUserDefinedAttribute(InterchangeDestinationNode, DuplicateName, InterchangeSourceNode->GetAttributeChecked<TArray<FPlane4f>>(AttributeKey.Key), PayloadKey, RequiresDelegate);
+					break;
+				case UE::Interchange::EAttributeTypes::Plane4dArray:
+					CreateUserDefinedAttribute(InterchangeDestinationNode, DuplicateName, InterchangeSourceNode->GetAttributeChecked<TArray<FPlane4d>>(AttributeKey.Key), PayloadKey, RequiresDelegate);
+					break;
+				case UE::Interchange::EAttributeTypes::Quat4fArray:
+					CreateUserDefinedAttribute(InterchangeDestinationNode, DuplicateName, InterchangeSourceNode->GetAttributeChecked<TArray<FQuat4f>>(AttributeKey.Key), PayloadKey, RequiresDelegate);
+					break;
+				case UE::Interchange::EAttributeTypes::Quat4dArray:
+					CreateUserDefinedAttribute(InterchangeDestinationNode, DuplicateName, InterchangeSourceNode->GetAttributeChecked<TArray<FQuat4d>>(AttributeKey.Key), PayloadKey, RequiresDelegate);
+					break;
+				case UE::Interchange::EAttributeTypes::Rotator3fArray:
+					CreateUserDefinedAttribute(InterchangeDestinationNode, DuplicateName, InterchangeSourceNode->GetAttributeChecked<TArray<FRotator3f>>(AttributeKey.Key), PayloadKey, RequiresDelegate);
+					break;
+				case UE::Interchange::EAttributeTypes::Rotator3dArray:
+					CreateUserDefinedAttribute(InterchangeDestinationNode, DuplicateName, InterchangeSourceNode->GetAttributeChecked<TArray<FRotator3d>>(AttributeKey.Key), PayloadKey, RequiresDelegate);
+					break;
+				case UE::Interchange::EAttributeTypes::Transform3fArray:
+					CreateUserDefinedAttribute(InterchangeDestinationNode, DuplicateName, InterchangeSourceNode->GetAttributeChecked<TArray<FTransform3f>>(AttributeKey.Key), PayloadKey, RequiresDelegate);
+					break;
+				case UE::Interchange::EAttributeTypes::Transform3dArray:
+					CreateUserDefinedAttribute(InterchangeDestinationNode, DuplicateName, InterchangeSourceNode->GetAttributeChecked<TArray<FTransform3d>>(AttributeKey.Key), PayloadKey, RequiresDelegate);
+					break;
+				case UE::Interchange::EAttributeTypes::Vector3fArray:
+					CreateUserDefinedAttribute(InterchangeDestinationNode, DuplicateName, InterchangeSourceNode->GetAttributeChecked<TArray<FVector3f>>(AttributeKey.Key), PayloadKey, RequiresDelegate);
+					break;
+				case UE::Interchange::EAttributeTypes::Vector3dArray:
+					CreateUserDefinedAttribute(InterchangeDestinationNode, DuplicateName, InterchangeSourceNode->GetAttributeChecked<TArray<FVector3d>>(AttributeKey.Key), PayloadKey, RequiresDelegate);
+					break;
+				case UE::Interchange::EAttributeTypes::Vector2fArray:
+					CreateUserDefinedAttribute(InterchangeDestinationNode, DuplicateName, InterchangeSourceNode->GetAttributeChecked<TArray<FVector2f>>(AttributeKey.Key), PayloadKey, RequiresDelegate);
+					break;
+				case UE::Interchange::EAttributeTypes::Vector4fArray:
+					CreateUserDefinedAttribute(InterchangeDestinationNode, DuplicateName, InterchangeSourceNode->GetAttributeChecked<TArray<FVector4f>>(AttributeKey.Key), PayloadKey, RequiresDelegate);
+					break;
+				case UE::Interchange::EAttributeTypes::Vector4dArray:
+					CreateUserDefinedAttribute(InterchangeDestinationNode, DuplicateName, InterchangeSourceNode->GetAttributeChecked<TArray<FVector4d>>(AttributeKey.Key), PayloadKey, RequiresDelegate);
+					break;
+				case UE::Interchange::EAttributeTypes::Box2fArray:
+					CreateUserDefinedAttribute(InterchangeDestinationNode, DuplicateName, InterchangeSourceNode->GetAttributeChecked<TArray<FBox2f>>(AttributeKey.Key), PayloadKey, RequiresDelegate);
+					break;
+				case UE::Interchange::EAttributeTypes::Box2DArray:
+					CreateUserDefinedAttribute(InterchangeDestinationNode, DuplicateName, InterchangeSourceNode->GetAttributeChecked<TArray<FBox2D>>(AttributeKey.Key), PayloadKey, RequiresDelegate);
+					break;
+				case UE::Interchange::EAttributeTypes::Box3fArray:
+					CreateUserDefinedAttribute(InterchangeDestinationNode, DuplicateName, InterchangeSourceNode->GetAttributeChecked<TArray<FBox3f>>(AttributeKey.Key), PayloadKey, RequiresDelegate);
+					break;
+				case UE::Interchange::EAttributeTypes::Box3dArray:
+					CreateUserDefinedAttribute(InterchangeDestinationNode, DuplicateName, InterchangeSourceNode->GetAttributeChecked<TArray<FBox3d>>(AttributeKey.Key), PayloadKey, RequiresDelegate);
+					break;
+				case UE::Interchange::EAttributeTypes::BoxSphereBounds3fArray:
+					CreateUserDefinedAttribute(InterchangeDestinationNode, DuplicateName, InterchangeSourceNode->GetAttributeChecked<TArray<FBoxSphereBounds3f>>(AttributeKey.Key), PayloadKey, RequiresDelegate);
+					break;
+				case UE::Interchange::EAttributeTypes::BoxSphereBounds3dArray:
+					CreateUserDefinedAttribute(InterchangeDestinationNode, DuplicateName, InterchangeSourceNode->GetAttributeChecked<TArray<FBoxSphereBounds3d>>(AttributeKey.Key), PayloadKey, RequiresDelegate);
+					break;
+				case UE::Interchange::EAttributeTypes::Sphere3fArray:
+					CreateUserDefinedAttribute(InterchangeDestinationNode, DuplicateName, InterchangeSourceNode->GetAttributeChecked<TArray<FSphere3f>>(AttributeKey.Key), PayloadKey, RequiresDelegate);
+					break;
+				case UE::Interchange::EAttributeTypes::Sphere3dArray:
+					CreateUserDefinedAttribute(InterchangeDestinationNode, DuplicateName, InterchangeSourceNode->GetAttributeChecked<TArray<FSphere3d>>(AttributeKey.Key), PayloadKey, RequiresDelegate);
+					break;
+				default:
+					ensureMsgf(false, TEXT("Unsupported EAttributeTypes value!"));
 					break;
 			}
 		}
@@ -322,63 +603,342 @@ void UInterchangeUserDefinedAttributesAPI::AddApplyAndFillDelegatesToFactory(UIn
 			UE::Interchange::EAttributeTypes Type = InterchangeFactoryNode->GetAttributeType(AttributeKey);
 			switch (Type)
 			{
-			case UE::Interchange::EAttributeTypes::Bool:
-				InterchangeFactoryNode->AddApplyAndFillDelegates<bool>(AttributeKey.Key, ParentClass, *UserDefinedAttributeName);
-				break;
-			case UE::Interchange::EAttributeTypes::Int8:
-				InterchangeFactoryNode->AddApplyAndFillDelegates<int8>(AttributeKey.Key, ParentClass, *UserDefinedAttributeName);
-				break;
-			case UE::Interchange::EAttributeTypes::Int16:
-				InterchangeFactoryNode->AddApplyAndFillDelegates<int16>(AttributeKey.Key, ParentClass, *UserDefinedAttributeName);
-				break;
-			case UE::Interchange::EAttributeTypes::Int32:
-				InterchangeFactoryNode->AddApplyAndFillDelegates<int32>(AttributeKey.Key, ParentClass, *UserDefinedAttributeName);
-				break;
-			case UE::Interchange::EAttributeTypes::Int64:
-				InterchangeFactoryNode->AddApplyAndFillDelegates<int64>(AttributeKey.Key, ParentClass, *UserDefinedAttributeName);
-				break;
-			case UE::Interchange::EAttributeTypes::UInt8:
-				InterchangeFactoryNode->AddApplyAndFillDelegates<uint8>(AttributeKey.Key, ParentClass, *UserDefinedAttributeName);
-				break;
-			case UE::Interchange::EAttributeTypes::UInt16:
-				InterchangeFactoryNode->AddApplyAndFillDelegates<uint16>(AttributeKey.Key, ParentClass, *UserDefinedAttributeName);
-				break;
-			case UE::Interchange::EAttributeTypes::UInt32:
-				InterchangeFactoryNode->AddApplyAndFillDelegates<uint32>(AttributeKey.Key, ParentClass, *UserDefinedAttributeName);
-				break;
-			case UE::Interchange::EAttributeTypes::UInt64:
-				InterchangeFactoryNode->AddApplyAndFillDelegates<uint64>(AttributeKey.Key, ParentClass, *UserDefinedAttributeName);
-				break;
-			case UE::Interchange::EAttributeTypes::Float:
-				InterchangeFactoryNode->AddApplyAndFillDelegates<float>(AttributeKey.Key, ParentClass, *UserDefinedAttributeName);
-				break;
-			case UE::Interchange::EAttributeTypes::Float16:
-				InterchangeFactoryNode->AddApplyAndFillDelegates<FFloat16>(AttributeKey.Key, ParentClass, *UserDefinedAttributeName);
-				break;
-			case UE::Interchange::EAttributeTypes::Vector2f:
-				InterchangeFactoryNode->AddApplyAndFillDelegates<FVector2f>(AttributeKey.Key, ParentClass, *UserDefinedAttributeName);
-				break;
-			case UE::Interchange::EAttributeTypes::Vector3f:
-				InterchangeFactoryNode->AddApplyAndFillDelegates<FVector3f>(AttributeKey.Key, ParentClass, *UserDefinedAttributeName);
-				break;
-			case UE::Interchange::EAttributeTypes::Vector4f:
-				InterchangeFactoryNode->AddApplyAndFillDelegates<FVector4f>(AttributeKey.Key, ParentClass, *UserDefinedAttributeName);
-				break;
-			case UE::Interchange::EAttributeTypes::Double:
-				InterchangeFactoryNode->AddApplyAndFillDelegates<double>(AttributeKey.Key, ParentClass, *UserDefinedAttributeName);
-				break;
-			case UE::Interchange::EAttributeTypes::Vector2d:
-				InterchangeFactoryNode->AddApplyAndFillDelegates<FVector2D>(AttributeKey.Key, ParentClass, *UserDefinedAttributeName);
-				break;
-			case UE::Interchange::EAttributeTypes::Vector3d:
-				InterchangeFactoryNode->AddApplyAndFillDelegates<FVector3d>(AttributeKey.Key, ParentClass, *UserDefinedAttributeName);
-				break;
-			case UE::Interchange::EAttributeTypes::Vector4d:
-				InterchangeFactoryNode->AddApplyAndFillDelegates<FVector4d>(AttributeKey.Key, ParentClass, *UserDefinedAttributeName);
-				break;
-			case UE::Interchange::EAttributeTypes::String:
-				InterchangeFactoryNode->AddApplyAndFillDelegates<FString>(AttributeKey.Key, ParentClass, *UserDefinedAttributeName);
-				break;
+				case UE::Interchange::EAttributeTypes::Bool:
+					InterchangeFactoryNode->AddApplyAndFillDelegates<bool>(AttributeKey.Key, ParentClass, *UserDefinedAttributeName);
+					break;
+				case UE::Interchange::EAttributeTypes::Color:
+					InterchangeFactoryNode->AddApplyAndFillDelegates<FColor>(AttributeKey.Key, ParentClass, *UserDefinedAttributeName);
+					break;
+				case UE::Interchange::EAttributeTypes::DateTime:
+					InterchangeFactoryNode->AddApplyAndFillDelegates<FDateTime>(AttributeKey.Key, ParentClass, *UserDefinedAttributeName);
+					break;
+				case UE::Interchange::EAttributeTypes::Double:
+					InterchangeFactoryNode->AddApplyAndFillDelegates<double>(AttributeKey.Key, ParentClass, *UserDefinedAttributeName);
+					break;
+				case UE::Interchange::EAttributeTypes::Enum:
+					InterchangeFactoryNode->AddApplyAndFillDelegates<bool>(AttributeKey.Key, ParentClass, *UserDefinedAttributeName);
+					break;
+				case UE::Interchange::EAttributeTypes::Float:
+					InterchangeFactoryNode->AddApplyAndFillDelegates<float>(AttributeKey.Key, ParentClass, *UserDefinedAttributeName);
+					break;
+				case UE::Interchange::EAttributeTypes::Guid:
+					InterchangeFactoryNode->AddApplyAndFillDelegates<FGuid>(AttributeKey.Key, ParentClass, *UserDefinedAttributeName);
+					break;
+				case UE::Interchange::EAttributeTypes::Int8:
+					InterchangeFactoryNode->AddApplyAndFillDelegates<int8>(AttributeKey.Key, ParentClass, *UserDefinedAttributeName);
+					break;
+				case UE::Interchange::EAttributeTypes::Int16:
+					InterchangeFactoryNode->AddApplyAndFillDelegates<int16>(AttributeKey.Key, ParentClass, *UserDefinedAttributeName);
+					break;
+				case UE::Interchange::EAttributeTypes::Int32:
+					InterchangeFactoryNode->AddApplyAndFillDelegates<int32>(AttributeKey.Key, ParentClass, *UserDefinedAttributeName);
+					break;
+				case UE::Interchange::EAttributeTypes::Int64:
+					InterchangeFactoryNode->AddApplyAndFillDelegates<int64>(AttributeKey.Key, ParentClass, *UserDefinedAttributeName);
+					break;
+				case UE::Interchange::EAttributeTypes::IntRect:
+					InterchangeFactoryNode->AddApplyAndFillDelegates<FIntRect>(AttributeKey.Key, ParentClass, *UserDefinedAttributeName);
+					break;
+				case UE::Interchange::EAttributeTypes::LinearColor:
+					InterchangeFactoryNode->AddApplyAndFillDelegates<FLinearColor>(AttributeKey.Key, ParentClass, *UserDefinedAttributeName);
+					break;
+				case UE::Interchange::EAttributeTypes::Name:
+					InterchangeFactoryNode->AddApplyAndFillDelegates<FName>(AttributeKey.Key, ParentClass, *UserDefinedAttributeName);
+					break;
+				case UE::Interchange::EAttributeTypes::RandomStream:
+					InterchangeFactoryNode->AddApplyAndFillDelegates<FRandomStream>(AttributeKey.Key, ParentClass, *UserDefinedAttributeName);
+					break;
+				case UE::Interchange::EAttributeTypes::String:
+					InterchangeFactoryNode->AddApplyAndFillDelegates<FString>(AttributeKey.Key, ParentClass, *UserDefinedAttributeName);
+					break;
+				case UE::Interchange::EAttributeTypes::Timespan:
+					InterchangeFactoryNode->AddApplyAndFillDelegates<FTimespan>(AttributeKey.Key, ParentClass, *UserDefinedAttributeName);
+					break;
+				case UE::Interchange::EAttributeTypes::TwoVectors:
+					InterchangeFactoryNode->AddApplyAndFillDelegates<FTwoVectors>(AttributeKey.Key, ParentClass, *UserDefinedAttributeName);
+					break;
+				case UE::Interchange::EAttributeTypes::UInt8:
+					InterchangeFactoryNode->AddApplyAndFillDelegates<uint8>(AttributeKey.Key, ParentClass, *UserDefinedAttributeName);
+					break;
+				case UE::Interchange::EAttributeTypes::UInt16:
+					InterchangeFactoryNode->AddApplyAndFillDelegates<uint16>(AttributeKey.Key, ParentClass, *UserDefinedAttributeName);
+					break;
+				case UE::Interchange::EAttributeTypes::UInt32:
+					InterchangeFactoryNode->AddApplyAndFillDelegates<uint32>(AttributeKey.Key, ParentClass, *UserDefinedAttributeName);
+					break;
+				case UE::Interchange::EAttributeTypes::UInt64:
+					InterchangeFactoryNode->AddApplyAndFillDelegates<uint64>(AttributeKey.Key, ParentClass, *UserDefinedAttributeName);
+					break;
+				case UE::Interchange::EAttributeTypes::Vector2d:
+					InterchangeFactoryNode->AddApplyAndFillDelegates<FVector2d>(AttributeKey.Key, ParentClass, *UserDefinedAttributeName);
+					break;
+				case UE::Interchange::EAttributeTypes::IntPoint:
+					InterchangeFactoryNode->AddApplyAndFillDelegates<FIntPoint>(AttributeKey.Key, ParentClass, *UserDefinedAttributeName);
+					break;
+				case UE::Interchange::EAttributeTypes::IntVector:
+					InterchangeFactoryNode->AddApplyAndFillDelegates<FIntVector>(AttributeKey.Key, ParentClass, *UserDefinedAttributeName);
+					break;
+				case UE::Interchange::EAttributeTypes::Vector2DHalf:
+					InterchangeFactoryNode->AddApplyAndFillDelegates<FVector2DHalf>(AttributeKey.Key, ParentClass, *UserDefinedAttributeName);
+					break;
+				case UE::Interchange::EAttributeTypes::Float16:
+					InterchangeFactoryNode->AddApplyAndFillDelegates<FFloat16>(AttributeKey.Key, ParentClass, *UserDefinedAttributeName);
+					break;
+				case UE::Interchange::EAttributeTypes::OrientedBox:
+					InterchangeFactoryNode->AddApplyAndFillDelegates<FOrientedBox>(AttributeKey.Key, ParentClass, *UserDefinedAttributeName);
+					break;
+				case UE::Interchange::EAttributeTypes::FrameNumber:
+					InterchangeFactoryNode->AddApplyAndFillDelegates<FFrameNumber>(AttributeKey.Key, ParentClass, *UserDefinedAttributeName);
+					break;
+				case UE::Interchange::EAttributeTypes::FrameRate:
+					InterchangeFactoryNode->AddApplyAndFillDelegates<FFrameRate>(AttributeKey.Key, ParentClass, *UserDefinedAttributeName);
+					break;
+				case UE::Interchange::EAttributeTypes::FrameTime:
+					InterchangeFactoryNode->AddApplyAndFillDelegates<FFrameTime>(AttributeKey.Key, ParentClass, *UserDefinedAttributeName);
+					break;
+				case UE::Interchange::EAttributeTypes::SoftObjectPath:
+					InterchangeFactoryNode->AddApplyAndFillDelegates<FSoftObjectPath>(AttributeKey.Key, ParentClass, *UserDefinedAttributeName);
+					break;
+				case UE::Interchange::EAttributeTypes::Matrix44f:
+					InterchangeFactoryNode->AddApplyAndFillDelegates<FMatrix44f>(AttributeKey.Key, ParentClass, *UserDefinedAttributeName);
+					break;
+				case UE::Interchange::EAttributeTypes::Matrix44d:
+					InterchangeFactoryNode->AddApplyAndFillDelegates<FMatrix44d>(AttributeKey.Key, ParentClass, *UserDefinedAttributeName);
+					break;
+				case UE::Interchange::EAttributeTypes::Plane4f:
+					InterchangeFactoryNode->AddApplyAndFillDelegates<FPlane4f>(AttributeKey.Key, ParentClass, *UserDefinedAttributeName);
+					break;
+				case UE::Interchange::EAttributeTypes::Plane4d:
+					InterchangeFactoryNode->AddApplyAndFillDelegates<FPlane4d>(AttributeKey.Key, ParentClass, *UserDefinedAttributeName);
+					break;
+				case UE::Interchange::EAttributeTypes::Quat4f:
+					InterchangeFactoryNode->AddApplyAndFillDelegates<FQuat4f>(AttributeKey.Key, ParentClass, *UserDefinedAttributeName);
+					break;
+				case UE::Interchange::EAttributeTypes::Quat4d:
+					InterchangeFactoryNode->AddApplyAndFillDelegates<FQuat4d>(AttributeKey.Key, ParentClass, *UserDefinedAttributeName);
+					break;
+				case UE::Interchange::EAttributeTypes::Rotator3f:
+					InterchangeFactoryNode->AddApplyAndFillDelegates<FRotator3f>(AttributeKey.Key, ParentClass, *UserDefinedAttributeName);
+					break;
+				case UE::Interchange::EAttributeTypes::Rotator3d:
+					InterchangeFactoryNode->AddApplyAndFillDelegates<FRotator3d>(AttributeKey.Key, ParentClass, *UserDefinedAttributeName);
+					break;
+				case UE::Interchange::EAttributeTypes::Transform3f:
+					InterchangeFactoryNode->AddApplyAndFillDelegates<FTransform3f>(AttributeKey.Key, ParentClass, *UserDefinedAttributeName);
+					break;
+				case UE::Interchange::EAttributeTypes::Transform3d:
+					InterchangeFactoryNode->AddApplyAndFillDelegates<FTransform3d>(AttributeKey.Key, ParentClass, *UserDefinedAttributeName);
+					break;
+				case UE::Interchange::EAttributeTypes::Vector3f:
+					InterchangeFactoryNode->AddApplyAndFillDelegates<FVector3f>(AttributeKey.Key, ParentClass, *UserDefinedAttributeName);
+					break;
+				case UE::Interchange::EAttributeTypes::Vector3d:
+					InterchangeFactoryNode->AddApplyAndFillDelegates<FVector3d>(AttributeKey.Key, ParentClass, *UserDefinedAttributeName);
+					break;
+				case UE::Interchange::EAttributeTypes::Vector2f:
+					InterchangeFactoryNode->AddApplyAndFillDelegates<FVector2f>(AttributeKey.Key, ParentClass, *UserDefinedAttributeName);
+					break;
+				case UE::Interchange::EAttributeTypes::Vector4f:
+					InterchangeFactoryNode->AddApplyAndFillDelegates<FVector4f>(AttributeKey.Key, ParentClass, *UserDefinedAttributeName);
+					break;
+				case UE::Interchange::EAttributeTypes::Vector4d:
+					InterchangeFactoryNode->AddApplyAndFillDelegates<FVector4d>(AttributeKey.Key, ParentClass, *UserDefinedAttributeName);
+					break;
+				case UE::Interchange::EAttributeTypes::Box2f:
+					InterchangeFactoryNode->AddApplyAndFillDelegates<FBox2f>(AttributeKey.Key, ParentClass, *UserDefinedAttributeName);
+					break;
+				case UE::Interchange::EAttributeTypes::Box2D:
+					InterchangeFactoryNode->AddApplyAndFillDelegates<FBox2D>(AttributeKey.Key, ParentClass, *UserDefinedAttributeName);
+					break;
+				case UE::Interchange::EAttributeTypes::Box3f:
+					InterchangeFactoryNode->AddApplyAndFillDelegates<FBox3f>(AttributeKey.Key, ParentClass, *UserDefinedAttributeName);
+					break;
+				case UE::Interchange::EAttributeTypes::Box3d:
+					InterchangeFactoryNode->AddApplyAndFillDelegates<FBox3d>(AttributeKey.Key, ParentClass, *UserDefinedAttributeName);
+					break;
+				case UE::Interchange::EAttributeTypes::BoxSphereBounds3f:
+					InterchangeFactoryNode->AddApplyAndFillDelegates<FBoxSphereBounds3f>(AttributeKey.Key, ParentClass, *UserDefinedAttributeName);
+					break;
+				case UE::Interchange::EAttributeTypes::BoxSphereBounds3d:
+					InterchangeFactoryNode->AddApplyAndFillDelegates<FBoxSphereBounds3d>(AttributeKey.Key, ParentClass, *UserDefinedAttributeName);
+					break;
+				case UE::Interchange::EAttributeTypes::Sphere3f:
+					InterchangeFactoryNode->AddApplyAndFillDelegates<FSphere3f>(AttributeKey.Key, ParentClass, *UserDefinedAttributeName);
+					break;
+				case UE::Interchange::EAttributeTypes::Sphere3d:
+					InterchangeFactoryNode->AddApplyAndFillDelegates<FSphere3d>(AttributeKey.Key, ParentClass, *UserDefinedAttributeName);
+					break;
+				case UE::Interchange::EAttributeTypes::BoolArray:
+					InterchangeFactoryNode->AddApplyAndFillDelegates<TArray<bool>>(AttributeKey.Key, ParentClass, *UserDefinedAttributeName);
+					break;
+				case UE::Interchange::EAttributeTypes::ColorArray:
+					InterchangeFactoryNode->AddApplyAndFillDelegates<TArray<FColor>>(AttributeKey.Key, ParentClass, *UserDefinedAttributeName);
+					break;
+				case UE::Interchange::EAttributeTypes::DateTimeArray:
+					InterchangeFactoryNode->AddApplyAndFillDelegates<TArray<FDateTime>>(AttributeKey.Key, ParentClass, *UserDefinedAttributeName);
+					break;
+				case UE::Interchange::EAttributeTypes::DoubleArray:
+					InterchangeFactoryNode->AddApplyAndFillDelegates<TArray<double>>(AttributeKey.Key, ParentClass, *UserDefinedAttributeName);
+					break;
+				case UE::Interchange::EAttributeTypes::EnumArray:
+					InterchangeFactoryNode->AddApplyAndFillDelegates<TArray<bool>>(AttributeKey.Key, ParentClass, *UserDefinedAttributeName);
+					break;
+				case UE::Interchange::EAttributeTypes::FloatArray:
+					InterchangeFactoryNode->AddApplyAndFillDelegates<TArray<float>>(AttributeKey.Key, ParentClass, *UserDefinedAttributeName);
+					break;
+				case UE::Interchange::EAttributeTypes::GuidArray:
+					InterchangeFactoryNode->AddApplyAndFillDelegates<TArray<FGuid>>(AttributeKey.Key, ParentClass, *UserDefinedAttributeName);
+					break;
+				case UE::Interchange::EAttributeTypes::Int8Array:
+					InterchangeFactoryNode->AddApplyAndFillDelegates<TArray<uint8>>(AttributeKey.Key, ParentClass, *UserDefinedAttributeName);
+					break;
+				case UE::Interchange::EAttributeTypes::Int16Array:
+					InterchangeFactoryNode->AddApplyAndFillDelegates<TArray<int16>>(AttributeKey.Key, ParentClass, *UserDefinedAttributeName);
+					break;
+				case UE::Interchange::EAttributeTypes::Int32Array:
+					InterchangeFactoryNode->AddApplyAndFillDelegates<TArray<int32>>(AttributeKey.Key, ParentClass, *UserDefinedAttributeName);
+					break;
+				case UE::Interchange::EAttributeTypes::Int64Array:
+					InterchangeFactoryNode->AddApplyAndFillDelegates<TArray<int64>>(AttributeKey.Key, ParentClass, *UserDefinedAttributeName);
+					break;
+				case UE::Interchange::EAttributeTypes::IntRectArray:
+					InterchangeFactoryNode->AddApplyAndFillDelegates<TArray<FIntRect>>(AttributeKey.Key, ParentClass, *UserDefinedAttributeName);
+					break;
+				case UE::Interchange::EAttributeTypes::LinearColorArray:
+					InterchangeFactoryNode->AddApplyAndFillDelegates<TArray<FLinearColor>>(AttributeKey.Key, ParentClass, *UserDefinedAttributeName);
+					break;
+				case UE::Interchange::EAttributeTypes::NameArray:
+					InterchangeFactoryNode->AddApplyAndFillDelegates<TArray<FName>>(AttributeKey.Key, ParentClass, *UserDefinedAttributeName);
+					break;
+				case UE::Interchange::EAttributeTypes::RandomStreamArray:
+					InterchangeFactoryNode->AddApplyAndFillDelegates<TArray<FRandomStream>>(AttributeKey.Key, ParentClass, *UserDefinedAttributeName);
+					break;
+				case UE::Interchange::EAttributeTypes::StringArray:
+					InterchangeFactoryNode->AddApplyAndFillDelegates<TArray<FString>>(AttributeKey.Key, ParentClass, *UserDefinedAttributeName);
+					break;
+				case UE::Interchange::EAttributeTypes::TimespanArray:
+					InterchangeFactoryNode->AddApplyAndFillDelegates<TArray<FTimespan>>(AttributeKey.Key, ParentClass, *UserDefinedAttributeName);
+					break;
+				case UE::Interchange::EAttributeTypes::TwoVectorsArray:
+					InterchangeFactoryNode->AddApplyAndFillDelegates<TArray<FTwoVectors>>(AttributeKey.Key, ParentClass, *UserDefinedAttributeName);
+					break;
+				case UE::Interchange::EAttributeTypes::ByteArray:
+					InterchangeFactoryNode->AddApplyAndFillDelegates<TArray<uint8>>(AttributeKey.Key, ParentClass, *UserDefinedAttributeName);
+					break;
+				case UE::Interchange::EAttributeTypes::ByteArray64:
+					InterchangeFactoryNode->AddApplyAndFillDelegates<TArray64<uint8>>(AttributeKey.Key, ParentClass, *UserDefinedAttributeName);
+					break;
+				case UE::Interchange::EAttributeTypes::UInt16Array:
+					InterchangeFactoryNode->AddApplyAndFillDelegates<TArray<uint16>>(AttributeKey.Key, ParentClass, *UserDefinedAttributeName);
+					break;
+				case UE::Interchange::EAttributeTypes::UInt32Array:
+					InterchangeFactoryNode->AddApplyAndFillDelegates<TArray<uint32>>(AttributeKey.Key, ParentClass, *UserDefinedAttributeName);
+					break;
+				case UE::Interchange::EAttributeTypes::UInt64Array:
+					InterchangeFactoryNode->AddApplyAndFillDelegates<TArray<uint64>>(AttributeKey.Key, ParentClass, *UserDefinedAttributeName);
+					break;
+				case UE::Interchange::EAttributeTypes::Vector2dArray:
+					InterchangeFactoryNode->AddApplyAndFillDelegates<TArray<FVector2d>>(AttributeKey.Key, ParentClass, *UserDefinedAttributeName);
+					break;
+				case UE::Interchange::EAttributeTypes::IntPointArray:
+					InterchangeFactoryNode->AddApplyAndFillDelegates<TArray<FIntPoint>>(AttributeKey.Key, ParentClass, *UserDefinedAttributeName);
+					break;
+				case UE::Interchange::EAttributeTypes::IntVectorArray:
+					InterchangeFactoryNode->AddApplyAndFillDelegates<TArray<FIntVector>>(AttributeKey.Key, ParentClass, *UserDefinedAttributeName);
+					break;
+				case UE::Interchange::EAttributeTypes::Vector2DHalfArray:
+					InterchangeFactoryNode->AddApplyAndFillDelegates<TArray<FVector2DHalf>>(AttributeKey.Key, ParentClass, *UserDefinedAttributeName);
+					break;
+				case UE::Interchange::EAttributeTypes::Float16Array:
+					InterchangeFactoryNode->AddApplyAndFillDelegates<TArray<FFloat16>>(AttributeKey.Key, ParentClass, *UserDefinedAttributeName);
+					break;
+				case UE::Interchange::EAttributeTypes::OrientedBoxArray:
+					InterchangeFactoryNode->AddApplyAndFillDelegates<TArray<FOrientedBox>>(AttributeKey.Key, ParentClass, *UserDefinedAttributeName);
+					break;
+				case UE::Interchange::EAttributeTypes::FrameNumberArray:
+					InterchangeFactoryNode->AddApplyAndFillDelegates<TArray<FFrameNumber>>(AttributeKey.Key, ParentClass, *UserDefinedAttributeName);
+					break;
+				case UE::Interchange::EAttributeTypes::FrameRateArray:
+					InterchangeFactoryNode->AddApplyAndFillDelegates<TArray<FFrameRate>>(AttributeKey.Key, ParentClass, *UserDefinedAttributeName);
+					break;
+				case UE::Interchange::EAttributeTypes::FrameTimeArray:
+					InterchangeFactoryNode->AddApplyAndFillDelegates<TArray<FFrameTime>>(AttributeKey.Key, ParentClass, *UserDefinedAttributeName);
+					break;
+				case UE::Interchange::EAttributeTypes::SoftObjectPathArray:
+					InterchangeFactoryNode->AddApplyAndFillDelegates<TArray<FSoftObjectPath>>(AttributeKey.Key, ParentClass, *UserDefinedAttributeName);
+					break;
+				case UE::Interchange::EAttributeTypes::Matrix44fArray:
+					InterchangeFactoryNode->AddApplyAndFillDelegates<TArray<FMatrix44f>>(AttributeKey.Key, ParentClass, *UserDefinedAttributeName);
+					break;
+				case UE::Interchange::EAttributeTypes::Matrix44dArray:
+					InterchangeFactoryNode->AddApplyAndFillDelegates<TArray<FMatrix44d>>(AttributeKey.Key, ParentClass, *UserDefinedAttributeName);
+					break;
+				case UE::Interchange::EAttributeTypes::Plane4fArray:
+					InterchangeFactoryNode->AddApplyAndFillDelegates<TArray<FPlane4f>>(AttributeKey.Key, ParentClass, *UserDefinedAttributeName);
+					break;
+				case UE::Interchange::EAttributeTypes::Plane4dArray:
+					InterchangeFactoryNode->AddApplyAndFillDelegates<TArray<FPlane4d>>(AttributeKey.Key, ParentClass, *UserDefinedAttributeName);
+					break;
+				case UE::Interchange::EAttributeTypes::Quat4fArray:
+					InterchangeFactoryNode->AddApplyAndFillDelegates<TArray<FQuat4f>>(AttributeKey.Key, ParentClass, *UserDefinedAttributeName);
+					break;
+				case UE::Interchange::EAttributeTypes::Quat4dArray:
+					InterchangeFactoryNode->AddApplyAndFillDelegates<TArray<FQuat4d>>(AttributeKey.Key, ParentClass, *UserDefinedAttributeName);
+					break;
+				case UE::Interchange::EAttributeTypes::Rotator3fArray:
+					InterchangeFactoryNode->AddApplyAndFillDelegates<TArray<FRotator3f>>(AttributeKey.Key, ParentClass, *UserDefinedAttributeName);
+					break;
+				case UE::Interchange::EAttributeTypes::Rotator3dArray:
+					InterchangeFactoryNode->AddApplyAndFillDelegates<TArray<FRotator3d>>(AttributeKey.Key, ParentClass, *UserDefinedAttributeName);
+					break;
+				case UE::Interchange::EAttributeTypes::Transform3fArray:
+					InterchangeFactoryNode->AddApplyAndFillDelegates<TArray<FTransform3f>>(AttributeKey.Key, ParentClass, *UserDefinedAttributeName);
+					break;
+				case UE::Interchange::EAttributeTypes::Transform3dArray:
+					InterchangeFactoryNode->AddApplyAndFillDelegates<TArray<FTransform3d>>(AttributeKey.Key, ParentClass, *UserDefinedAttributeName);
+					break;
+				case UE::Interchange::EAttributeTypes::Vector3fArray:
+					InterchangeFactoryNode->AddApplyAndFillDelegates<TArray<FVector3f>>(AttributeKey.Key, ParentClass, *UserDefinedAttributeName);
+					break;
+				case UE::Interchange::EAttributeTypes::Vector3dArray:
+					InterchangeFactoryNode->AddApplyAndFillDelegates<TArray<FVector3d>>(AttributeKey.Key, ParentClass, *UserDefinedAttributeName);
+					break;
+				case UE::Interchange::EAttributeTypes::Vector2fArray:
+					InterchangeFactoryNode->AddApplyAndFillDelegates<TArray<FVector2f>>(AttributeKey.Key, ParentClass, *UserDefinedAttributeName);
+					break;
+				case UE::Interchange::EAttributeTypes::Vector4fArray:
+					InterchangeFactoryNode->AddApplyAndFillDelegates<TArray<FVector4f>>(AttributeKey.Key, ParentClass, *UserDefinedAttributeName);
+					break;
+				case UE::Interchange::EAttributeTypes::Vector4dArray:
+					InterchangeFactoryNode->AddApplyAndFillDelegates<TArray<FVector4d>>(AttributeKey.Key, ParentClass, *UserDefinedAttributeName);
+					break;
+				case UE::Interchange::EAttributeTypes::Box2fArray:
+					InterchangeFactoryNode->AddApplyAndFillDelegates<TArray<FBox2f>>(AttributeKey.Key, ParentClass, *UserDefinedAttributeName);
+					break;
+				case UE::Interchange::EAttributeTypes::Box2DArray:
+					InterchangeFactoryNode->AddApplyAndFillDelegates<TArray<FBox2D>>(AttributeKey.Key, ParentClass, *UserDefinedAttributeName);
+					break;
+				case UE::Interchange::EAttributeTypes::Box3fArray:
+					InterchangeFactoryNode->AddApplyAndFillDelegates<TArray<FBox3f>>(AttributeKey.Key, ParentClass, *UserDefinedAttributeName);
+					break;
+				case UE::Interchange::EAttributeTypes::Box3dArray:
+					InterchangeFactoryNode->AddApplyAndFillDelegates<TArray<FBox3d>>(AttributeKey.Key, ParentClass, *UserDefinedAttributeName);
+					break;
+				case UE::Interchange::EAttributeTypes::BoxSphereBounds3fArray:
+					InterchangeFactoryNode->AddApplyAndFillDelegates<TArray<FBoxSphereBounds3f>>(AttributeKey.Key, ParentClass, *UserDefinedAttributeName);
+					break;
+				case UE::Interchange::EAttributeTypes::BoxSphereBounds3dArray:
+					InterchangeFactoryNode->AddApplyAndFillDelegates<TArray<FBoxSphereBounds3d>>(AttributeKey.Key, ParentClass, *UserDefinedAttributeName);
+					break;
+				case UE::Interchange::EAttributeTypes::Sphere3fArray:
+					InterchangeFactoryNode->AddApplyAndFillDelegates<TArray<FSphere3f>>(AttributeKey.Key, ParentClass, *UserDefinedAttributeName);
+					break;
+				case UE::Interchange::EAttributeTypes::Sphere3dArray:
+					InterchangeFactoryNode->AddApplyAndFillDelegates<TArray<FSphere3d>>(AttributeKey.Key, ParentClass, *UserDefinedAttributeName);
+					break;
+				default:
+					ensureMsgf(false, TEXT("Unsupported EAttributeTypes value!"));
+					break;
 			}
 		}
 	}
diff --git a/Engine/Source/Runtime/Interchange/Core/Private/Tests/StorageTest.cpp b/Engine/Source/Runtime/Interchange/Core/Private/Tests/StorageTest.cpp
--- a/Engine/Source/Runtime/Interchange/Core/Private/Tests/StorageTest.cpp
+++ b/Engine/Source/Runtime/Interchange/Core/Private/Tests/StorageTest.cpp
@@ -1,12 +1,13 @@
 // Copyright Epic Games, Inc. All Rights Reserved.
 
-#include "CoreMinimal.h"
 #include "Async/ParallelFor.h"
+#include "CoreMinimal.h"
 #include "Misc/AutomationTest.h"
 #include "Serialization/MemoryReader.h"
 #include "Serialization/MemoryWriter.h"
 #include "Templates/Atomic.h"
 #include "Types/AttributeStorage.h"
+#include "UObject/SoftObjectPath.h"
 
 #if WITH_DEV_AUTOMATION_TESTS
 
@@ -19,11 +20,20 @@ IMPLEMENT_SIMPLE_AUTOMATION_TEST(FAttributeStorageTest, "System.Runtime.Intercha
 	FAttributeKey Key = CreateUniqueKey();																												\
 	if (!IsAttributeStorageResultSuccess(TestStorage.RegisterAttribute(Key, RefValue)))																	\
 	{																																					\
-		AddError(FString(TEXT("`AttributeStorage` must handle adding ATTRIBUTE_TYPE attribute")));														\
+		AddError(FString(TEXT("`AttributeStorage` must handle adding " #ATTRIBUTE_TYPE " attribute")));													\
 	}																																					\
 	ATTRIBUTE_TYPE StoredValue;																															\
 	TestStorage.GetAttributeHandle<ATTRIBUTE_TYPE>(Key).Get(StoredValue);																				\
-	TestEqual<ATTRIBUTE_TYPE>(TEXT("`AttributeStorage` must handle add and retrieve" #ATTRIBUTE_TYPE "attribute"), StoredValue, RefValue);				\
+	TestEqual<ATTRIBUTE_TYPE>(TEXT("`AttributeStorage` must handle retrieving registered " #ATTRIBUTE_TYPE " attribute"), StoredValue, RefValue);		\
+	ATTRIBUTE_TYPE EmptyValue{};																														\
+	TestStorage.GetAttributeHandle<ATTRIBUTE_TYPE>(Key).Set(EmptyValue);																				\
+	ATTRIBUTE_TYPE StoredEmptyValue;																													\
+	TestStorage.GetAttributeHandle<ATTRIBUTE_TYPE>(Key).Get(StoredEmptyValue);																			\
+	TestEqual<ATTRIBUTE_TYPE>(TEXT("`AttributeStorage` must handle clearing " #ATTRIBUTE_TYPE " attribute value"), StoredEmptyValue, EmptyValue);		\
+	TestStorage.GetAttributeHandle<ATTRIBUTE_TYPE>(Key).Set(RefValue);																					\
+	ATTRIBUTE_TYPE StoredValueAgain;																													\
+	TestStorage.GetAttributeHandle<ATTRIBUTE_TYPE>(Key).Get(StoredValueAgain);																			\
+	TestEqual<ATTRIBUTE_TYPE>(TEXT("`AttributeStorage` must handle overwriting " #ATTRIBUTE_TYPE " attribute"), StoredValueAgain, RefValue);			\
 }
 
 #define ADD_ATTRIBUTE_STORAGE_NOCOMPARE(ATTRIBUTE_TYPE, ATTRIBUTE_VALUE)																				\
@@ -31,7 +41,7 @@ IMPLEMENT_SIMPLE_AUTOMATION_TEST(FAttributeStorageTest, "System.Runtime.Intercha
 	FAttributeKey Key = CreateUniqueKey();																												\
 	if (!IsAttributeStorageResultSuccess(TestStorage.RegisterAttribute(Key, RefValue)))																	\
 	{																																					\
-		AddError(FString(TEXT("`AttributeStorage` must handle adding ATTRIBUTE_TYPE attribute")));														\
+		AddError(FString(TEXT("`AttributeStorage` must handle adding " #ATTRIBUTE_TYPE " attribute")));													\
 	}
 
 
@@ -351,6 +361,242 @@ bool FAttributeStorageTest::RunTest(const FString& Parameters)
 		ADD_ATTRIBUTE_STORAGE(TEnumAsByte<ETestEnum::Type>, ETestEnum::Orthographic);
 	}
 
+	// Store array types
+	{
+		enum class EAttributeStorageEnumTest : uint8
+		{
+			EnumValue0 = 0,
+			EnumValue23 = 23,
+			EnumValue35 = 35,
+			EnumValue255 = 255,
+		};
+
+		enum EBlueprintEnum : int
+		{
+			Opaque,
+			Masked,
+			Translucent,
+			Additive
+		};
+
+		using namespace UE::InterchangeStorageTest::Private;
+
+		ADD_ATTRIBUTE_STORAGE(TArray<bool>, TArray<bool>({true, false, false}));
+		ADD_ATTRIBUTE_STORAGE(TArray<FColor>, TArray<FColor>({FColor::Black, FColor::Green, FColor::Yellow}));
+		ADD_ATTRIBUTE_STORAGE(TArray<FDateTime>, TArray<FDateTime>({FDateTime(), FDateTime(2024, 02, 03), FDateTime(4765417645)}));
+		ADD_ATTRIBUTE_STORAGE(TArray<double>, TArray<double>({0.1234567890123456789, 0.123, -3.5e10}));
+
+		ADD_ATTRIBUTE_STORAGE(TArray<EAttributeStorageEnumTest>, TArray<EAttributeStorageEnumTest>({EAttributeStorageEnumTest::EnumValue0, EAttributeStorageEnumTest::EnumValue255, EAttributeStorageEnumTest::EnumValue35}));
+		ADD_ATTRIBUTE_STORAGE(TArray<EAttributeStorageEnumTest>, TArray<EAttributeStorageEnumTest>({EAttributeStorageEnumTest::EnumValue0}));
+		ADD_ATTRIBUTE_STORAGE(TArray<EAttributeStorageEnumTest>, TArray<EAttributeStorageEnumTest>({}));
+		ADD_ATTRIBUTE_STORAGE(TArray<EBlueprintEnum>, TArray<EBlueprintEnum>({EBlueprintEnum::Opaque, EBlueprintEnum::Masked, EBlueprintEnum::Translucent}));
+		ADD_ATTRIBUTE_STORAGE(TArray<EBlueprintEnum>, TArray<EBlueprintEnum>({EBlueprintEnum::Opaque, EBlueprintEnum::Masked}));
+		ADD_ATTRIBUTE_STORAGE(TArray<EBlueprintEnum>, TArray<EBlueprintEnum>({}));
+		ADD_ATTRIBUTE_STORAGE(TArray<TEnumAsByte<ETestEnum::Type>>, TArray<TEnumAsByte<ETestEnum::Type>>({ETestEnum::Perspective, ETestEnum::Orthographic}));
+		ADD_ATTRIBUTE_STORAGE(TArray<TEnumAsByte<ETestEnum::Type>>, TArray<TEnumAsByte<ETestEnum::Type>>({ETestEnum::Perspective}));
+		ADD_ATTRIBUTE_STORAGE(TArray<TEnumAsByte<ETestEnum::Type>>, TArray<TEnumAsByte<ETestEnum::Type>>({}));
+
+		ADD_ATTRIBUTE_STORAGE(TArray<float>, TArray<float>({0.1234567f, 0.1234567f, 0.0f}));
+		ADD_ATTRIBUTE_STORAGE(TArray<FGuid>, TArray<FGuid>({FGuid(0xA4F0E4CD, 0xF97C4375, 0x9C25C211, 0xFD17B8BF), FGuid()}));
+		ADD_ATTRIBUTE_STORAGE(TArray<int8>, TArray<int8>({-0x7f, 65}));
+		ADD_ATTRIBUTE_STORAGE(TArray<int16>, TArray<int16>({-0x75cc, -2}));
+		ADD_ATTRIBUTE_STORAGE(TArray<int32>, TArray<int32>({0x7fccabcd, 655}));
+		ADD_ATTRIBUTE_STORAGE(TArray<int64>, TArray<int64>({0x7fccabcd12345678, 456154}));
+		ADD_ATTRIBUTE_STORAGE(TArray<FIntRect>, TArray<FIntRect>({FIntRect(-1,-2, 1, 2), FIntRect(-1,-2, 1, 2), FIntRect()}));
+		ADD_ATTRIBUTE_STORAGE(TArray<FLinearColor>, TArray<FLinearColor>({FLinearColor(0.59f, 0.49f, 0.0f), FLinearColor::Red}));
+		ADD_ATTRIBUTE_STORAGE(TArray<FName>, TArray<FName>({FName(TEXT("AAAAAAAAAAAAAAAAAAAAAAAAAAAAAAA")), FName(TEXT("BBBBBBBBBBBBBBBBBBBBBBBBBBBBBBBBBBBBBBBBBBBBBB")), NAME_None}));
+		ADD_ATTRIBUTE_STORAGE(TArray<FName>, TArray<FName>({FName(TEXT("CCCCCCCCCCCCCCCCCCCCCCCCCCCCCCCCCCCCCCCC")), NAME_None}));
+		ADD_ATTRIBUTE_STORAGE(TArray<FName>, TArray<FName>({NAME_None}));
+		ADD_ATTRIBUTE_STORAGE(TArray<FName>, TArray<FName>({}));
+		{
+			ADD_ATTRIBUTE_STORAGE_NOCOMPARE(TArray<FRandomStream>, TArray<FRandomStream>({34, 45, 2}));
+		}
+		ADD_ATTRIBUTE_STORAGE(TArray<FString>, TArray<FString>({TEXT("DDDDDDDDDDDDDDDDDDDDDDDDDDDDDDDDDDDDDDDD"), TEXT("EEEEEEEEEEEEEEEEEEEEEEEEEEEE"), FString{}}));
+		ADD_ATTRIBUTE_STORAGE(TArray<FString>, TArray<FString>({TEXT("FFFFFFFFFFFFF")}));
+		ADD_ATTRIBUTE_STORAGE(TArray<FString>, TArray<FString>({TEXT("")}));
+		ADD_ATTRIBUTE_STORAGE(TArray<FString>, TArray<FString>({}));
+		ADD_ATTRIBUTE_STORAGE(TArray<FSoftObjectPath>, TArray<FSoftObjectPath>({(UClass::StaticClass(), UObject::StaticClass())}));
+		ADD_ATTRIBUTE_STORAGE(TArray<FSoftObjectPath>, TArray<FSoftObjectPath>({(UClass::StaticClass(), FSoftObjectPath())}));
+		ADD_ATTRIBUTE_STORAGE(TArray<FSoftObjectPath>, TArray<FSoftObjectPath>({(FSoftObjectPath())}));
+		ADD_ATTRIBUTE_STORAGE(TArray<FSoftObjectPath>, TArray<FSoftObjectPath>({}));
+		ADD_ATTRIBUTE_STORAGE(TArray<FTimespan>, TArray<FTimespan>({FTimespan{10, 44, 38}, FTimespan{47652}}));
+		ADD_ATTRIBUTE_STORAGE(TArray<FTwoVectors>, TArray<FTwoVectors>({ FTwoVectors(FVector(150.0, -203.0, 4500.7), FVector(1.0, 1.0, 1.0)), FTwoVectors(FVector(0.0, 0.0, 4500.7), FVector(1.0, 0.0, 1.0)) }));
+		ADD_ATTRIBUTE_STORAGE(TArray<uint8>, TArray<uint8>({0x85, 0}));
+		ADD_ATTRIBUTE_STORAGE(TArray<uint16>, TArray<uint16>({0x85cc}));
+		ADD_ATTRIBUTE_STORAGE(TArray<uint32>, TArray<uint32>({0x85ccabcd}));
+		ADD_ATTRIBUTE_STORAGE(TArray<uint64>, TArray<uint64>({0x85ccabcd12345678}));
+		ADD_ATTRIBUTE_STORAGE(TArray<FIntPoint>, TArray<FIntPoint>({FIntPoint{1, 4}, FIntPoint{0, 0}}));
+		ADD_ATTRIBUTE_STORAGE(TArray<FIntVector>, TArray<FIntVector>({FIntVector(1, 4, 3), FIntVector()}));
+
+		{
+			ADD_ATTRIBUTE_STORAGE_NOCOMPARE(TArray<FVector2DHalf>, TArray<FVector2DHalf>({FVector2DHalf(150.0, -203.0)}))
+			TArray<FVector2DHalf> StoredValue;
+			TestStorage.GetAttributeHandle<TArray<FVector2DHalf>>(Key).Get(StoredValue);
+			for (int32 Index = 0; Index < StoredValue.Num(); ++Index)
+			{
+				if (TestEqual(TEXT("`AttributeStorage` must handle add and retrieve TArray<FVector2DHalf> attribute"), StoredValue[Index].X, RefValue[Index].X))
+				{
+					TestEqual(TEXT("`AttributeStorage` must handle add and retrieve TArray<FVector2DHalf> attribute"), StoredValue[Index].Y, RefValue[Index].Y);
+				}
+			}
+		}
+		ADD_ATTRIBUTE_STORAGE(TArray<FFloat16>, TArray<FFloat16>({0.12f, 0.0f, -0.2f}));
+
+		{
+			TArray<FOrientedBox> SourceBoxes{FOrientedBox(), FOrientedBox(), FOrientedBox()};
+			SourceBoxes[0].Center = RandomStream.VRand();
+			SourceBoxes[0].AxisX = RandomStream.VRand();
+			SourceBoxes[0].AxisY = RandomStream.VRand();
+			SourceBoxes[0].AxisZ = RandomStream.VRand();
+			SourceBoxes[0].ExtentX = RandomStream.FRand();
+			SourceBoxes[0].ExtentY = RandomStream.FRand();
+			SourceBoxes[0].ExtentZ = RandomStream.FRand();
+			SourceBoxes[1].Center = RandomStream.VRand();
+			SourceBoxes[1].AxisX = RandomStream.VRand();
+			SourceBoxes[1].AxisY = RandomStream.VRand();
+			SourceBoxes[1].AxisZ = RandomStream.VRand();
+			SourceBoxes[1].ExtentX = RandomStream.FRand();
+			SourceBoxes[1].ExtentY = RandomStream.FRand();
+			SourceBoxes[1].ExtentZ = RandomStream.FRand();
+
+			ADD_ATTRIBUTE_STORAGE_NOCOMPARE(TArray<FOrientedBox>, SourceBoxes)
+			TArray<FOrientedBox> StoredValue;
+			TestStorage.GetAttributeHandle<TArray<FOrientedBox>>(Key).Get(StoredValue);
+
+			for(int32 Index = 0; Index < StoredValue.Num(); ++Index)
+			{
+				TestEqual(TEXT("`AttributeStorage` must handle add and retrieve TArray<FOrientedBox> attribute"), StoredValue[Index].Center, RefValue[Index].Center);
+				TestEqual(TEXT("`AttributeStorage` must handle add and retrieve TArray<FOrientedBox> attribute"), StoredValue[Index].AxisX, RefValue[Index].AxisX);
+				TestEqual(TEXT("`AttributeStorage` must handle add and retrieve TArray<FOrientedBox> attribute"), StoredValue[Index].AxisY, RefValue[Index].AxisY);
+				TestEqual(TEXT("`AttributeStorage` must handle add and retrieve TArray<FOrientedBox> attribute"), StoredValue[Index].AxisZ, RefValue[Index].AxisZ);
+				TestEqual(TEXT("`AttributeStorage` must handle add and retrieve TArray<FOrientedBox> attribute"), StoredValue[Index].ExtentX, RefValue[Index].ExtentX);
+				TestEqual(TEXT("`AttributeStorage` must handle add and retrieve TArray<FOrientedBox> attribute"), StoredValue[Index].ExtentY, RefValue[Index].ExtentY);
+				TestEqual(TEXT("`AttributeStorage` must handle add and retrieve TArray<FOrientedBox> attribute"), StoredValue[Index].ExtentZ, RefValue[Index].ExtentZ);
+			}
+		}
+
+		ADD_ATTRIBUTE_STORAGE(TArray<FFrameNumber>, TArray<FFrameNumber>({FFrameNumber{30}, FFrameNumber{1}, FFrameNumber{}}));
+		ADD_ATTRIBUTE_STORAGE(TArray<FFrameRate>, TArray<FFrameRate>({FFrameRate(1,60), FFrameRate()}));
+		ADD_ATTRIBUTE_STORAGE(TArray<FFrameTime>, TArray<FFrameTime>({FFrameTime(29,0.5f), FFrameTime()}));
+		ADD_ATTRIBUTE_STORAGE(TArray<FMatrix44f>, TArray<FMatrix44f>({
+			(FTransform3f(FRotator3f(25.0f, 2.0f, 3.14159f), FVector3f(150.0f, -203.0f, 4500.7f), FVector3f(1.0f, 1.0f, 1.0f)).ToMatrixWithScale()),
+			(FTransform3f(FRotator3f(0.0f, -2.0f, 3.1), FVector3f(150.0f, 23.0f, 450.7f), FVector3f(2.0f, 2.0f, 2.0f)).ToMatrixWithScale()),
+			FMatrix44f::Identity
+		}));
+		ADD_ATTRIBUTE_STORAGE(TArray<FMatrix44d>, TArray<FMatrix44d>({
+			(FTransform3d(FRotator3d(25.0, 2.0, 3.14159), FVector3d(150.0, -203.0, 4500.7), FVector3d(1.0, 1.0, 1.0)).ToMatrixWithScale()),
+			(FTransform3d(FRotator3d(0.0, -2.0, 5), FVector3d(15.0, -2.0, 40.7), FVector3d(4.0, 4.0, 4.0)).ToMatrixWithScale()),
+			FMatrix44d::Identity
+		}));
+		ADD_ATTRIBUTE_STORAGE(TArray<FPlane4f>, TArray<FPlane4f>({
+			FPlane4f(FVector3f(1.0f, 1.0f, 0.0f), FVector3f(0.0f, 0.0f, 1.0f)),
+			FPlane4f(FVector3f(-1.0f, -1.0f, 0.0f), FVector3f(1.0f, 0.0f, 1.0f)),
+			FPlane4f(),
+		}));
+		ADD_ATTRIBUTE_STORAGE(TArray<FPlane4d>, TArray<FPlane4d>({
+			FPlane4d(FVector3d(1.0, 1.0, 0), FVector3d(0, 0, 1.0)),
+			FPlane4d(FVector3d(1.0, -1.0, 0), FVector3d(2.0, 0, 1.0)),
+			FPlane4d(),
+		}));
+		ADD_ATTRIBUTE_STORAGE(TArray<FQuat4f>, TArray<FQuat4f>({FQuat4f(1.0f, 1.0f, 0.0f, 1.0f), FQuat4f(-1.0f, -1.0f, 0.0f, 1.0f), FQuat4f()}));
+		ADD_ATTRIBUTE_STORAGE(TArray<FQuat4d>, TArray<FQuat4d>({FQuat4d(1.0, 1.0, 0, 1.0), FQuat4d(1.0f, 1.0f, 0.0f, 2.0f), FQuat4d::Identity}));
+		ADD_ATTRIBUTE_STORAGE(TArray<FRotator3f>, TArray<FRotator3f>({FRotator3f(25.0f, 2.0f, 3.14159f), FRotator3f(2.0f, 3.0f, -3.14159f), FRotator3f()}));
+		ADD_ATTRIBUTE_STORAGE(TArray<FRotator3d>, TArray<FRotator3d>({FRotator3d(25.0, 2.0, 3.14159), FRotator3d(55.0, -2.0, 6.14159), FRotator3d()}));
+
+		{
+			ADD_ATTRIBUTE_STORAGE_NOCOMPARE(TArray<FTransform3f>, TArray<FTransform3f>({
+				FTransform3f(FRotator3f(25.0f, 2.0f, 3.14159f), FVector3f(150.0f, -203.0f, 4500.7f), FVector3f(1.0f, 1.0f, 1.0f)),
+				FTransform3f(FRotator3f(0.0f, -2.0f, 3.59f), FVector3f(10.0f, 3.0f, 450.7f), FVector3f(2.0f, 2.0f, 2.0f)),
+				FTransform3f::Identity,
+			}));
+
+			TArray<FTransform3f> StoredValue;
+			TestStorage.GetAttributeHandle<TArray<FTransform3f>>(Key).Get(StoredValue);
+			TestEqual(TEXT("`AttributeStorage` must handle adding TArray<FTransform3f> attribute"), StoredValue.Num(), RefValue.Num());
+			for (int32 Index = 0; Index < StoredValue.Num(); ++Index)
+			{
+				if (!StoredValue[Index].Equals(RefValue[Index]))
+				{
+					AddError(FString(TEXT("`AttributeStorage` must handle adding TArray<FTransform3f> attribute")));
+				}
+			}
+		}
+		{
+			ADD_ATTRIBUTE_STORAGE_NOCOMPARE(TArray<FTransform3d>, TArray<FTransform3d>({
+				FTransform3d(FRotator3d(25.0f, 2.0f, 3.14159f), FVector3d(150.0f, -203.0f, 4500.7f), FVector3d(1.0f, 1.0f, 1.0f)),
+				FTransform3d(FRotator3d(0.0f, -2.0f, 3.59f), FVector3d(10.0f, 3.0f, 450.7f), FVector3d(2.0f, 2.0f, 2.0f)),
+				FTransform3d::Identity,
+			}));
+
+			TArray<FTransform3d> StoredValue;
+			TestStorage.GetAttributeHandle<TArray<FTransform3d>>(Key).Get(StoredValue);
+			TestEqual(TEXT("`AttributeStorage` must handle adding TArray<FTransform3d> attribute"), StoredValue.Num(), RefValue.Num());
+			for (int32 Index = 0; Index < StoredValue.Num(); ++Index)
+			{
+				if (!StoredValue[Index].Equals(RefValue[Index]))
+				{
+					AddError(FString(TEXT("`AttributeStorage` must handle adding TArray<FTransform3d> attribute")));
+				}
+			}
+		}
+
+		ADD_ATTRIBUTE_STORAGE(TArray<FVector2f>, TArray<FVector2f>({FVector2f(150.0f, -203.0f), FVector2f(15065.0f, -03.0f), FVector2f()}));
+		ADD_ATTRIBUTE_STORAGE(TArray<FVector2d>, TArray<FVector2d>({FVector2d(150.0, -203.0), FVector2d(16550.0, 203.0), FVector2d()}));
+		ADD_ATTRIBUTE_STORAGE(TArray<FVector3f>, TArray<FVector3f>({FVector3f(150.0f, -203.0f, 4500.7f), FVector3f()}));
+		ADD_ATTRIBUTE_STORAGE(TArray<FVector3d>, TArray<FVector3d>({FVector3d(150.0, -203.0, 4500.7), FVector3d()}));
+		ADD_ATTRIBUTE_STORAGE(TArray<FVector4f>, TArray<FVector4f>({FVector4f(150.0f, -203.0f, 4500.7f, 2.0f), FVector4f()}));
+		ADD_ATTRIBUTE_STORAGE(TArray<FVector4d>, TArray<FVector4d>({FVector4d(150.0, -203.0, 4500.7, 2.0), FVector4d()}));
+		ADD_ATTRIBUTE_STORAGE(TArray<FBox2f>, TArray<FBox2f>({FBox2f(FVector2f(-1.0f), FVector2f(1.0f)), FBox2f(FVector2f(1.0f), FVector2f(-1.0f))}));
+		ADD_ATTRIBUTE_STORAGE(TArray<FBox2D>, TArray<FBox2D>({FBox2D(FVector2D(-1.0), FVector2D(1.0)), FBox2D(FVector2D(1.0), FVector2D(-1.0))}));
+		ADD_ATTRIBUTE_STORAGE(TArray<FBox3f>, TArray<FBox3f>({FBox3f(FVector3f(-1.0f), FVector3f(1.0f)), FBox3f(FVector3f(1.0f), FVector3f(-1.0f))}));
+		ADD_ATTRIBUTE_STORAGE(TArray<FBox3d>, TArray<FBox3d>({FBox3d(FVector3d(-1.0), FVector3d(1.0)), FBox3d(FVector3d(1.0), FVector3d(-1.0))}));
+		ADD_ATTRIBUTE_STORAGE(TArray<FBoxSphereBounds3f>, TArray<FBoxSphereBounds3f>({FBoxSphereBounds3f(FVector3f(0.0f), FVector3f(2.0f), 2.2f), FBoxSphereBounds3f(FVector3f(-1.0f), FVector3f(-2.0f), 5.2f)}));
+		ADD_ATTRIBUTE_STORAGE(TArray<FBoxSphereBounds3d>, TArray<FBoxSphereBounds3d>({FBoxSphereBounds3d(FVector3d(0.0), FVector3d(2.0), 2.2), FBoxSphereBounds3d(FVector3d(6.0), FVector3d(4.0), 6.2)}));
+
+		{
+			TArray<FSphere3f> Spheres{FSphere3f(), FSphere3f(), FSphere3f()};
+			Spheres[0].Center.X = RandomStream.FRand();
+			Spheres[0].Center.Y = RandomStream.FRand();
+			Spheres[0].Center.Z = RandomStream.FRand();
+			Spheres[0].W = RandomStream.FRand();
+			Spheres[1].Center.X = RandomStream.FRand();
+			Spheres[1].Center.Y = RandomStream.FRand();
+			Spheres[1].Center.Z = RandomStream.FRand();
+			Spheres[1].W = RandomStream.FRand();
+
+			ADD_ATTRIBUTE_STORAGE_NOCOMPARE(TArray<FSphere3f>, Spheres)
+			TArray<FSphere3f> StoredValue;
+			TestStorage.GetAttributeHandle<TArray<FSphere3f>>(Key).Get(StoredValue);
+
+			for(int32 Index = 0; Index < StoredValue.Num(); ++Index)
+			{
+				TestEqual(TEXT("`AttributeStorage` must handle add and retrieve TArray<FSphere3f> attribute"), StoredValue[Index].Center, RefValue[Index].Center);
+				TestEqual(TEXT("`AttributeStorage` must handle add and retrieve TArray<FSphere3f> attribute"), StoredValue[Index].W, RefValue[Index].W);
+			}
+		}
+		{
+			TArray<FSphere3d> Spheres{FSphere3d(), FSphere3d(), FSphere3d()};
+			Spheres[0].Center.X = RandomStream.FRand();
+			Spheres[0].Center.Y = RandomStream.FRand();
+			Spheres[0].Center.Z = RandomStream.FRand();
+			Spheres[0].W = RandomStream.FRand();
+			Spheres[1].Center.X = RandomStream.FRand();
+			Spheres[1].Center.Y = RandomStream.FRand();
+			Spheres[1].Center.Z = RandomStream.FRand();
+			Spheres[1].W = RandomStream.FRand();
+
+			ADD_ATTRIBUTE_STORAGE_NOCOMPARE(TArray<FSphere3d>, Spheres)
+			TArray<FSphere3d> StoredValue;
+			TestStorage.GetAttributeHandle<TArray<FSphere3d>>(Key).Get(StoredValue);
+
+			for(int32 Index = 0; Index < StoredValue.Num(); ++Index)
+			{
+				TestEqual(TEXT("`AttributeStorage` must handle add and retrieve TArray<FSphere3d> attribute"), StoredValue[Index].Center, RefValue[Index].Center);
+				TestEqual(TEXT("`AttributeStorage` must handle add and retrieve TArray<FSphere3d> attribute"), StoredValue[Index].W, RefValue[Index].W);
+			}
+		}
+	}
+
 	//Get the Hash we will use it later to do some test
 	FGuid HashGuidRef = TestStorage.GetStorageHash();
 
diff --git a/Engine/Source/Runtime/Interchange/Core/Private/Types/AttributeStorage.cpp b/Engine/Source/Runtime/Interchange/Core/Private/Types/AttributeStorage.cpp
--- a/Engine/Source/Runtime/Interchange/Core/Private/Types/AttributeStorage.cpp
+++ b/Engine/Source/Runtime/Interchange/Core/Private/Types/AttributeStorage.cpp
@@ -1,6 +1,7 @@
 // Copyright Epic Games, Inc. All Rights Reserved.
 
 #include "Types/AttributeStorage.h"
+
 #include "InterchangeLogPrivate.h"
 
 #include "CoreMinimal.h"
@@ -28,6 +29,7 @@
 #include "Misc/ScopeLock.h"
 #include "Misc/ScopeLock.h"
 #include "Misc/SecureHash.h"
+#include "UObject/SoftObjectPath.h"
 
 #include <type_traits>
 
@@ -54,6 +56,21 @@ namespace UE
 		{
 		};
 
+		template<typename T>
+		FString StringifyArrayAttribute(const T& Array)
+		{
+			static_assert(TIsTArray<T>::Value);
+
+			FString Result = TEXT("[");
+			for (const auto& Element : Array)
+			{
+				Result += TAttributeTypeTraitsInternal<typename T::ElementType>::ToString(Element) + TEXT(", ");
+			}
+			Result.RemoveFromEnd(TEXT(", "));
+			Result += TEXT("]");
+			return Result;
+		}
+
 		template<> 
 		struct TAttributeTypeTraitsInternal<bool>
 		{
@@ -240,7 +257,27 @@ namespace UE
 		struct TAttributeTypeTraitsInternal<FOrientedBox>
 		{
 			static constexpr EAttributeTypes GetType() { return EAttributeTypes::OrientedBox; }
-			static FString ToString(const FOrientedBox& Value) { return TEXT("FOrientedBox"); }
+			static FString ToString(const FOrientedBox& Value) 
+			{
+				return FString::Printf(
+					TEXT("Center= %f, %f, %f | AxisX= %f, %f, %f | AxisY= %f, %f, %f | AxisZ= %f, %f, %f | Extents= %f, %f, %f"),
+					Value.Center.X,
+					Value.Center.Y,
+					Value.Center.Z,
+					Value.AxisX.X,
+					Value.AxisX.Y,
+					Value.AxisX.Z,
+					Value.AxisY.X,
+					Value.AxisY.Y,
+					Value.AxisY.Z,
+					Value.AxisZ.X,
+					Value.AxisZ.Y,
+					Value.AxisZ.Z,
+					Value.ExtentX,
+					Value.ExtentY,
+					Value.ExtentZ
+				);
+			}
 		};
 
 		template<> 
@@ -432,18 +469,389 @@ namespace UE
 			static FString ToString(const FSphere3d& Value) { return Value.ToString(); }
 		};
 
+		template<> 
+		struct TAttributeTypeTraitsInternal<TArray<bool>>
+		{
+			static constexpr EAttributeTypes GetType() { return EAttributeTypes::BoolArray; }
+			static FString ToString(const TArray<bool>& Value) { return StringifyArrayAttribute(Value); };
+		};
+
+		template<> 
+		struct TAttributeTypeTraitsInternal<TArray<FColor>>
+		{
+			static constexpr EAttributeTypes GetType() { return EAttributeTypes::ColorArray; }
+			static FString ToString(const TArray<FColor>& Value) { return StringifyArrayAttribute(Value); };
+		};
+
+		template<> 
+		struct TAttributeTypeTraitsInternal<TArray<FDateTime>>
+		{
+			static constexpr EAttributeTypes GetType() { return EAttributeTypes::DateTimeArray; }
+			static FString ToString(const TArray<FDateTime>& Value) { return StringifyArrayAttribute(Value); };
+		};
+
+		template<> 
+		struct TAttributeTypeTraitsInternal<TArray<double>>
+		{
+			static constexpr EAttributeTypes GetType() { return EAttributeTypes::DoubleArray; }
+			static FString ToString(const TArray<double>& Value) { return StringifyArrayAttribute(Value); };
+		};
+
+		template<> 
+		struct TAttributeTypeTraitsInternal<TArray<float>>
+		{
+			static constexpr EAttributeTypes GetType() { return EAttributeTypes::FloatArray; }
+			static FString ToString(const TArray<float>& Value) { return StringifyArrayAttribute(Value); };
+		};
+
+		template<> 
+		struct TAttributeTypeTraitsInternal<TArray<FGuid>>
+		{
+			static constexpr EAttributeTypes GetType() { return EAttributeTypes::GuidArray; }
+			static FString ToString(const TArray<FGuid>& Value) { return StringifyArrayAttribute(Value); };
+		};
+
+		template<> 
+		struct TAttributeTypeTraitsInternal<TArray<int8>>
+		{
+			static constexpr EAttributeTypes GetType() { return EAttributeTypes::Int8Array; }
+			static FString ToString(const TArray<int8>& Value) { return StringifyArrayAttribute(Value); };
+		};
+
+		template<> 
+		struct TAttributeTypeTraitsInternal<TArray<int16>>
+		{
+			static constexpr EAttributeTypes GetType() { return EAttributeTypes::Int16Array; }
+			static FString ToString(const TArray<int16>& Value) { return StringifyArrayAttribute(Value); };
+		};
+
+		template<> 
+		struct TAttributeTypeTraitsInternal<TArray<int32>>
+		{
+			static constexpr EAttributeTypes GetType() { return EAttributeTypes::Int32Array; }
+			static FString ToString(const TArray<int32>& Value) { return StringifyArrayAttribute(Value); };
+		};
+
+		template<> 
+		struct TAttributeTypeTraitsInternal<TArray<int64>>
+		{
+			static constexpr EAttributeTypes GetType() { return EAttributeTypes::Int64Array; }
+			static FString ToString(const TArray<int64>& Value) { return StringifyArrayAttribute(Value); };
+		};
+
+		template<> 
+		struct TAttributeTypeTraitsInternal<TArray<FIntRect>>
+		{
+			static constexpr EAttributeTypes GetType() { return EAttributeTypes::IntRectArray; }
+			static FString ToString(const TArray<FIntRect>& Value) { return StringifyArrayAttribute(Value); };
+		};
+
+		template<> 
+		struct TAttributeTypeTraitsInternal<TArray<FLinearColor>>
+		{
+			static constexpr EAttributeTypes GetType() { return EAttributeTypes::LinearColorArray; }
+			static FString ToString(const TArray<FLinearColor>& Value) { return StringifyArrayAttribute(Value); };
+		};
+
+		template<> 
+		struct TAttributeTypeTraitsInternal<TArray<FName>>
+		{
+			static constexpr EAttributeTypes GetType() { return EAttributeTypes::NameArray; }
+			static FString ToString(const TArray<FName>& Value) { return StringifyArrayAttribute(Value); };
+		};
+
+		template<> 
+		struct TAttributeTypeTraitsInternal<TArray<FRandomStream>>
+		{
+			static constexpr EAttributeTypes GetType() { return EAttributeTypes::RandomStreamArray; }
+			static FString ToString(const TArray<FRandomStream>& Value) { return StringifyArrayAttribute(Value); };
+		};
+
+		template<> 
+		struct TAttributeTypeTraitsInternal<TArray<FString>>
+		{
+			static constexpr EAttributeTypes GetType() { return EAttributeTypes::StringArray; }
+			static FString ToString(const TArray<FString>& Value) { return StringifyArrayAttribute(Value); };
+		};
+
+		template<> 
+		struct TAttributeTypeTraitsInternal<TArray<FTimespan>>
+		{
+			static constexpr EAttributeTypes GetType() { return EAttributeTypes::TimespanArray; }
+			static FString ToString(const TArray<FTimespan>& Value) { return StringifyArrayAttribute(Value); };
+		};
+
+		template<> 
+		struct TAttributeTypeTraitsInternal<TArray<FTwoVectors>>
+		{
+			static constexpr EAttributeTypes GetType() { return EAttributeTypes::TwoVectorsArray; }
+			static FString ToString(const TArray<FTwoVectors>& Value) { return StringifyArrayAttribute(Value); };
+		};
+
 		template<> 
 		struct TAttributeTypeTraitsInternal<TArray<uint8>>
 		{
 			static constexpr EAttributeTypes GetType() { return EAttributeTypes::ByteArray; }
-			static FString ToString(const TArray<uint8>& Value) {return TEXT("Array<uint8>"); }
+			static FString ToString(const TArray<uint8>& Value) { return StringifyArrayAttribute(Value); };
 		};
 
 		template<> 
 		struct TAttributeTypeTraitsInternal<TArray64<uint8>>
 		{
 			static constexpr EAttributeTypes GetType() { return EAttributeTypes::ByteArray64; }
-			static FString ToString(const TArray64<uint8>& Value) { return TEXT("Array64<uint8>"); }
+			static FString ToString(const TArray64<uint8>& Value) { return StringifyArrayAttribute(Value); };
+		};
+
+		template<> 
+		struct TAttributeTypeTraitsInternal<TArray<uint16>>
+		{
+			static constexpr EAttributeTypes GetType() { return EAttributeTypes::UInt16Array; }
+			static FString ToString(const TArray<uint16>& Value) { return StringifyArrayAttribute(Value); };
+		};
+
+		template<> 
+		struct TAttributeTypeTraitsInternal<TArray<uint32>>
+		{
+			static constexpr EAttributeTypes GetType() { return EAttributeTypes::UInt32Array; }
+			static FString ToString(const TArray<uint32>& Value) { return StringifyArrayAttribute(Value); };
+		};
+
+		template<> 
+		struct TAttributeTypeTraitsInternal<TArray<uint64>>
+		{
+			static constexpr EAttributeTypes GetType() { return EAttributeTypes::UInt64Array; }
+			static FString ToString(const TArray<uint64>& Value) { return StringifyArrayAttribute(Value); };
+		};
+
+		template<> 
+		struct TAttributeTypeTraitsInternal<TArray<FVector2D>>
+		{
+			static constexpr EAttributeTypes GetType() { return EAttributeTypes::Vector2dArray; }
+			static FString ToString(const TArray<FVector2D>& Value) { return StringifyArrayAttribute(Value); };
+		};
+
+		template<> 
+		struct TAttributeTypeTraitsInternal<TArray<FIntPoint>>
+		{
+			static constexpr EAttributeTypes GetType() { return EAttributeTypes::IntPointArray; }
+			static FString ToString(const TArray<FIntPoint>& Value) { return StringifyArrayAttribute(Value); };
+		};
+
+		template<> 
+		struct TAttributeTypeTraitsInternal<TArray<FIntVector>>
+		{
+			static constexpr EAttributeTypes GetType() { return EAttributeTypes::IntVectorArray; }
+			static FString ToString(const TArray<FIntVector>& Value) { return StringifyArrayAttribute(Value); };
+		};
+
+		template<> 
+		struct TAttributeTypeTraitsInternal<TArray<FVector2DHalf>>
+		{
+			static constexpr EAttributeTypes GetType() { return EAttributeTypes::Vector2DHalfArray; }
+			static FString ToString(const TArray<FVector2DHalf>& Value) { return StringifyArrayAttribute(Value); };
+		};
+
+		template<> 
+		struct TAttributeTypeTraitsInternal<TArray<FFloat16>>
+		{
+			static constexpr EAttributeTypes GetType() { return EAttributeTypes::Float16Array; }
+			static FString ToString(const TArray<FFloat16>& Value) { return StringifyArrayAttribute(Value); };
+		};
+
+		template<> 
+		struct TAttributeTypeTraitsInternal<TArray<FOrientedBox>>
+		{
+			static constexpr EAttributeTypes GetType() { return EAttributeTypes::OrientedBoxArray; }
+			static FString ToString(const TArray<FOrientedBox>& Value) { return StringifyArrayAttribute(Value); };
+		};
+
+		template<> 
+		struct TAttributeTypeTraitsInternal<TArray<FFrameNumber>>
+		{
+			static constexpr EAttributeTypes GetType() { return EAttributeTypes::FrameNumberArray; }
+			static FString ToString(const TArray<FFrameNumber>& Value) { return StringifyArrayAttribute(Value); };
+		};
+
+		template<> 
+		struct TAttributeTypeTraitsInternal<TArray<FFrameRate>>
+		{
+			static constexpr EAttributeTypes GetType() { return EAttributeTypes::FrameRateArray; }
+			static FString ToString(const TArray<FFrameRate>& Value) { return StringifyArrayAttribute(Value); };
+		};
+
+		template<> 
+		struct TAttributeTypeTraitsInternal<TArray<FFrameTime>>
+		{
+			static constexpr EAttributeTypes GetType() { return EAttributeTypes::FrameTimeArray; }
+			static FString ToString(const TArray<FFrameTime>& Value) { return StringifyArrayAttribute(Value); };
+		};
+
+		template<> 
+		struct TAttributeTypeTraitsInternal<TArray<FSoftObjectPath>>
+		{
+			static constexpr EAttributeTypes GetType() { return EAttributeTypes::SoftObjectPathArray; }
+			static FString ToString(const TArray<FSoftObjectPath>& Value) { return StringifyArrayAttribute(Value); };
+		};
+
+		template<> 
+		struct TAttributeTypeTraitsInternal<TArray<FMatrix44f>>
+		{
+			static constexpr EAttributeTypes GetType() { return EAttributeTypes::Matrix44fArray; }
+			static FString ToString(const TArray<FMatrix44f>& Value) { return StringifyArrayAttribute(Value); };
+		};
+
+		template<> 
+		struct TAttributeTypeTraitsInternal<TArray<FMatrix44d>>
+		{
+			static constexpr EAttributeTypes GetType() { return EAttributeTypes::Matrix44dArray; }
+			static FString ToString(const TArray<FMatrix44d>& Value) { return StringifyArrayAttribute(Value); };
+		};
+
+		template<> 
+		struct TAttributeTypeTraitsInternal<TArray<FPlane4f>>
+		{
+			static constexpr EAttributeTypes GetType() { return EAttributeTypes::Plane4fArray; }
+			static FString ToString(const TArray<FPlane4f>& Value) { return StringifyArrayAttribute(Value); };
+		};
+
+		template<> 
+		struct TAttributeTypeTraitsInternal<TArray<FPlane4d>>
+		{
+			static constexpr EAttributeTypes GetType() { return EAttributeTypes::Plane4dArray; }
+			static FString ToString(const TArray<FPlane4d>& Value) { return StringifyArrayAttribute(Value); };
+		};
+
+		template<> 
+		struct TAttributeTypeTraitsInternal<TArray<FQuat4f>>
+		{
+			static constexpr EAttributeTypes GetType() { return EAttributeTypes::Quat4fArray; }
+			static FString ToString(const TArray<FQuat4f>& Value) { return StringifyArrayAttribute(Value); };
+		};
+
+		template<> 
+		struct TAttributeTypeTraitsInternal<TArray<FQuat4d>>
+		{
+			static constexpr EAttributeTypes GetType() { return EAttributeTypes::Quat4dArray; }
+			static FString ToString(const TArray<FQuat4d>& Value) { return StringifyArrayAttribute(Value); };
+		};
+
+		template<> 
+		struct TAttributeTypeTraitsInternal<TArray<FRotator3f>>
+		{
+			static constexpr EAttributeTypes GetType() { return EAttributeTypes::Rotator3fArray; }
+			static FString ToString(const TArray<FRotator3f>& Value) { return StringifyArrayAttribute(Value); };
+		};
+
+		template<> 
+		struct TAttributeTypeTraitsInternal<TArray<FRotator3d>>
+		{
+			static constexpr EAttributeTypes GetType() { return EAttributeTypes::Rotator3dArray; }
+			static FString ToString(const TArray<FRotator3d>& Value) { return StringifyArrayAttribute(Value); };
+		};
+
+		template<> 
+		struct TAttributeTypeTraitsInternal<TArray<FTransform3f>>
+		{
+			static constexpr EAttributeTypes GetType() { return EAttributeTypes::Transform3fArray; }
+			static FString ToString(const TArray<FTransform3f>& Value) { return StringifyArrayAttribute(Value); };
+		};
+
+		template<> 
+		struct TAttributeTypeTraitsInternal<TArray<FTransform3d>>
+		{
+			static constexpr EAttributeTypes GetType() { return EAttributeTypes::Transform3dArray; }
+			static FString ToString(const TArray<FTransform3d>& Value) { return StringifyArrayAttribute(Value); };
+		};
+
+		template<> 
+		struct TAttributeTypeTraitsInternal<TArray<FVector3f>>
+		{
+			static constexpr EAttributeTypes GetType() { return EAttributeTypes::Vector3fArray; }
+			static FString ToString(const TArray<FVector3f>& Value) { return StringifyArrayAttribute(Value); };
+		};
+
+		template<> 
+		struct TAttributeTypeTraitsInternal<TArray<FVector3d>>
+		{
+			static constexpr EAttributeTypes GetType() { return EAttributeTypes::Vector3dArray; }
+			static FString ToString(const TArray<FVector3d>& Value) { return StringifyArrayAttribute(Value); };
+		};
+
+		template<> 
+		struct TAttributeTypeTraitsInternal<TArray<FVector2f>>
+		{
+			static constexpr EAttributeTypes GetType() { return EAttributeTypes::Vector2fArray; }
+			static FString ToString(const TArray<FVector2f>& Value) { return StringifyArrayAttribute(Value); };
+		};
+
+		template<> 
+		struct TAttributeTypeTraitsInternal<TArray<FVector4f>>
+		{
+			static constexpr EAttributeTypes GetType() { return EAttributeTypes::Vector4fArray; }
+			static FString ToString(const TArray<FVector4f>& Value) { return StringifyArrayAttribute(Value); };
+		};
+
+		template<> 
+		struct TAttributeTypeTraitsInternal<TArray<FVector4d>>
+		{
+			static constexpr EAttributeTypes GetType() { return EAttributeTypes::Vector4dArray; }
+			static FString ToString(const TArray<FVector4d>& Value) { return StringifyArrayAttribute(Value); };
+		};
+
+		template<> 
+		struct TAttributeTypeTraitsInternal<TArray<FBox2f>>
+		{
+			static constexpr EAttributeTypes GetType() { return EAttributeTypes::Box2fArray; }
+			static FString ToString(const TArray<FBox2f>& Value) { return StringifyArrayAttribute(Value); };
+		};
+
+		template<> 
+		struct TAttributeTypeTraitsInternal<TArray<FBox2D>>
+		{
+			static constexpr EAttributeTypes GetType() { return EAttributeTypes::Box2DArray; }
+			static FString ToString(const TArray<FBox2D>& Value) { return StringifyArrayAttribute(Value); };
+		};
+
+		template<> 
+		struct TAttributeTypeTraitsInternal<TArray<FBox3f>>
+		{
+			static constexpr EAttributeTypes GetType() { return EAttributeTypes::Box3fArray; }
+			static FString ToString(const TArray<FBox3f>& Value) { return StringifyArrayAttribute(Value); };
+		};
+
+		template<> 
+		struct TAttributeTypeTraitsInternal<TArray<FBox3d>>
+		{
+			static constexpr EAttributeTypes GetType() { return EAttributeTypes::Box3dArray; }
+			static FString ToString(const TArray<FBox3d>& Value) { return StringifyArrayAttribute(Value); };
+		};
+
+		template<> 
+		struct TAttributeTypeTraitsInternal<TArray<FBoxSphereBounds3f>>
+		{
+			static constexpr EAttributeTypes GetType() { return EAttributeTypes::BoxSphereBounds3fArray; }
+			static FString ToString(const TArray<FBoxSphereBounds3f>& Value) { return StringifyArrayAttribute(Value); };
+		};
+
+		template<> 
+		struct TAttributeTypeTraitsInternal<TArray<FBoxSphereBounds3d>>
+		{
+			static constexpr EAttributeTypes GetType() { return EAttributeTypes::BoxSphereBounds3dArray; }
+			static FString ToString(const TArray<FBoxSphereBounds3d>& Value) { return StringifyArrayAttribute(Value); };
+		};
+
+		template<> 
+		struct TAttributeTypeTraitsInternal<TArray<FSphere3f>>
+		{
+			static constexpr EAttributeTypes GetType() { return EAttributeTypes::Sphere3fArray; }
+			static FString ToString(const TArray<FSphere3f>& Value) { return StringifyArrayAttribute(Value); };
+		};
+
+		template<> 
+		struct TAttributeTypeTraitsInternal<TArray<FSphere3d>>
+		{
+			static constexpr EAttributeTypes GetType() { return EAttributeTypes::Sphere3dArray; }
+			static FString ToString(const TArray<FSphere3d>& Value) { return StringifyArrayAttribute(Value); };
 		};
 
 		FString AttributeTypeToString(EAttributeTypes AttributeType)
@@ -452,9 +860,8 @@ namespace UE
 			switch (AttributeType)
 			{
 				case EAttributeTypes::None: { AttributeTypeString = TEXT("None"); } break;
+
 				case EAttributeTypes::Bool: { AttributeTypeString = TEXT("Bool"); } break;
-				case EAttributeTypes::ByteArray: { AttributeTypeString = TEXT("ByteArray"); } break;
-				case EAttributeTypes::ByteArray64: { AttributeTypeString = TEXT("ByteArray64"); } break;
 				case EAttributeTypes::Color: { AttributeTypeString = TEXT("Color"); } break;
 				case EAttributeTypes::DateTime: { AttributeTypeString = TEXT("DateTime"); } break;
 				case EAttributeTypes::Double: { AttributeTypeString = TEXT("Double"); } break;
@@ -509,13 +916,67 @@ namespace UE
 				case EAttributeTypes::BoxSphereBounds3d: { AttributeTypeString = TEXT("BoxSphereBounds3d"); } break;
 				case EAttributeTypes::Sphere3f: { AttributeTypeString = TEXT("Sphere3f"); } break;
 				case EAttributeTypes::Sphere3d: { AttributeTypeString = TEXT("Sphere3d"); } break;
+
+				case EAttributeTypes::BoolArray: { AttributeTypeString = TEXT("BoolArray"); } break;
+				case EAttributeTypes::ColorArray: { AttributeTypeString = TEXT("ColorArray"); } break;
+				case EAttributeTypes::DateTimeArray: { AttributeTypeString = TEXT("DateTimeArray"); } break;
+				case EAttributeTypes::DoubleArray: { AttributeTypeString = TEXT("DoubleArray"); } break;
+				case EAttributeTypes::EnumArray: { AttributeTypeString = TEXT("EnumArray"); } break;
+				case EAttributeTypes::FloatArray: { AttributeTypeString = TEXT("FloatArray"); } break;
+				case EAttributeTypes::GuidArray: { AttributeTypeString = TEXT("GuidArray"); } break;
+				case EAttributeTypes::Int8Array: { AttributeTypeString = TEXT("Int8Array"); } break;
+				case EAttributeTypes::Int16Array: { AttributeTypeString = TEXT("Int16Array"); } break;
+				case EAttributeTypes::Int32Array: { AttributeTypeString = TEXT("Int32Array"); } break;
+				case EAttributeTypes::Int64Array: { AttributeTypeString = TEXT("Int64Array"); } break;
+				case EAttributeTypes::IntRectArray: { AttributeTypeString = TEXT("IntRectArray"); } break;
+				case EAttributeTypes::LinearColorArray: { AttributeTypeString = TEXT("LinearColorArray"); } break;
+				case EAttributeTypes::NameArray: { AttributeTypeString = TEXT("NameArray"); } break;
+				case EAttributeTypes::RandomStreamArray: { AttributeTypeString = TEXT("RandomStreamArray"); } break;
+				case EAttributeTypes::StringArray: { AttributeTypeString = TEXT("StringArray"); } break;
+				case EAttributeTypes::TimespanArray: { AttributeTypeString = TEXT("TimespanArray"); } break;
+				case EAttributeTypes::TwoVectorsArray: { AttributeTypeString = TEXT("TwoVectorsArray"); } break;
+				case EAttributeTypes::ByteArray: { AttributeTypeString = TEXT("ByteArray"); } break;
+				case EAttributeTypes::ByteArray64: { AttributeTypeString = TEXT("ByteArray64"); } break;
+				case EAttributeTypes::UInt16Array: { AttributeTypeString = TEXT("UInt16Array"); } break;
+				case EAttributeTypes::UInt32Array: { AttributeTypeString = TEXT("UInt32Array"); } break;
+				case EAttributeTypes::UInt64Array: { AttributeTypeString = TEXT("UInt64Array"); } break;
+				case EAttributeTypes::Vector2dArray: { AttributeTypeString = TEXT("Vector2dArray"); } break;
+				case EAttributeTypes::IntPointArray: { AttributeTypeString = TEXT("IntPointArray"); } break;
+				case EAttributeTypes::IntVectorArray: { AttributeTypeString = TEXT("IntVectorArray"); } break;
+				case EAttributeTypes::Vector2DHalfArray: { AttributeTypeString = TEXT("Vector2DHalfArray"); } break;
+				case EAttributeTypes::Float16Array: { AttributeTypeString = TEXT("Float16Array"); } break;
+				case EAttributeTypes::OrientedBoxArray: { AttributeTypeString = TEXT("OrientedBoxArray"); } break;
+				case EAttributeTypes::FrameNumberArray: { AttributeTypeString = TEXT("FrameNumberArray"); } break;
+				case EAttributeTypes::FrameRateArray: { AttributeTypeString = TEXT("FrameRateArray"); } break;
+				case EAttributeTypes::FrameTimeArray: { AttributeTypeString = TEXT("FrameTimeArray"); } break;
+				case EAttributeTypes::SoftObjectPathArray: { AttributeTypeString = TEXT("SoftObjectPathArray"); } break;
+				case EAttributeTypes::Matrix44fArray: { AttributeTypeString = TEXT("Matrix44fArray"); } break;
+				case EAttributeTypes::Matrix44dArray: { AttributeTypeString = TEXT("Matrix44dArray"); } break;
+				case EAttributeTypes::Plane4fArray: { AttributeTypeString = TEXT("Plane4fArray"); } break;
+				case EAttributeTypes::Plane4dArray: { AttributeTypeString = TEXT("Plane4dArray"); } break;
+				case EAttributeTypes::Quat4fArray: { AttributeTypeString = TEXT("Quat4fArray"); } break;
+				case EAttributeTypes::Quat4dArray: { AttributeTypeString = TEXT("Quat4dArray"); } break;
+				case EAttributeTypes::Rotator3fArray: { AttributeTypeString = TEXT("Rotator3fArray"); } break;
+				case EAttributeTypes::Rotator3dArray: { AttributeTypeString = TEXT("Rotator3dArray"); } break;
+				case EAttributeTypes::Transform3fArray: { AttributeTypeString = TEXT("Transform3fArray"); } break;
+				case EAttributeTypes::Transform3dArray: { AttributeTypeString = TEXT("Transform3dArray"); } break;
+				case EAttributeTypes::Vector3fArray: { AttributeTypeString = TEXT("Vector3fArray"); } break;
+				case EAttributeTypes::Vector3dArray: { AttributeTypeString = TEXT("Vector3dArray"); } break;
+				case EAttributeTypes::Vector2fArray: { AttributeTypeString = TEXT("Vector2fArray"); } break;
+				case EAttributeTypes::Vector4fArray: { AttributeTypeString = TEXT("Vector4fArray"); } break;
+				case EAttributeTypes::Vector4dArray: { AttributeTypeString = TEXT("Vector4dArray"); } break;
+				case EAttributeTypes::Box2fArray: { AttributeTypeString = TEXT("Box2fArray"); } break;
+				case EAttributeTypes::Box2DArray: { AttributeTypeString = TEXT("Box2DArray"); } break;
+				case EAttributeTypes::Box3fArray: { AttributeTypeString = TEXT("Box3fArray"); } break;
+				case EAttributeTypes::Box3dArray: { AttributeTypeString = TEXT("Box3dArray"); } break;
+				case EAttributeTypes::BoxSphereBounds3fArray: { AttributeTypeString = TEXT("BoxSphereBounds3fArray"); } break;
+				case EAttributeTypes::BoxSphereBounds3dArray: { AttributeTypeString = TEXT("BoxSphereBounds3dArray"); } break;
+				case EAttributeTypes::Sphere3fArray: { AttributeTypeString = TEXT("Sphere3fArray"); } break;
+				case EAttributeTypes::Sphere3dArray: { AttributeTypeString = TEXT("Sphere3dArray"); } break;
+
 				default:
-				{
-					//Ensure if we ask an unknown type
-					const bool bUnknownType = true;
-					ensure(!bUnknownType);
-				}
-				break;
+					ensureMsgf(false, TEXT("Unsupported EAttributeTypes value!"));
+					break;
 			}
 			return AttributeTypeString;
 		}
@@ -540,63 +1001,39 @@ namespace UE
 			return TAttributeTypeTraitsInternal<ValueType>::ToString(Value);
 		}
 
-		FGuid GetValueHash(const uint8* Value, uint64 ValueSize)
-		{
-			FSHA1 Sha;
-			Sha.Update(Value, ValueSize);
-			Sha.Final();
-			// Retrieve the hash and use it to construct a pseudo-GUID. 
-			uint32 Hash[5];
-			Sha.GetHash(reinterpret_cast<uint8*>(Hash));
-			return FGuid(Hash[0] ^ Hash[4], Hash[1], Hash[2], Hash[3]);
-		}
-
-		/** Return the storage size of the value. */
+		/** Return the storage size of the value in bytes */
 		template<typename T>
 		uint64 GetValueSizeInternal(const T& Value)
 		{
-			const EAttributeTypes ValueType = TAttributeTypeTraitsInternal<T>::GetType();
-			if (ValueType == EAttributeTypes::None)
+			if constexpr (TIsTArray<T>::Value)
 			{
-				return 0;
-			}
-			return sizeof(T);
-		}
+				// This should not be used for nested array types (arrays of strings)
+				static_assert(!TIsArrayOfStringType<T>);
 
-		template<>
-		uint64 GetValueSizeInternal(const FString& Value)
-		{
-			//We must add the null character '/0' terminate string
-			return (Value.Len() + 1) * sizeof(TCHAR);
+				return Value.NumBytes();
+			}
+			else if constexpr (std::is_same_v<T, FString>)
+			{
+				// We must add the null character '/0' terminate string
+				return (Value.Len() + 1) * sizeof(TCHAR);
+			}
+			else if constexpr (std::is_same_v<T, FName> || std::is_same_v<T, FSoftObjectPath>)
+			{
+				FString ValueStr = Value.ToString();
+				return GetValueSizeInternal(ValueStr);
+			}
+			else // Common single values (float, FRotator3f, int32, etc.)
+			{
+				const EAttributeTypes ValueType = TAttributeTypeTraitsInternal<T>::GetType();
+				if (ValueType == EAttributeTypes::None)
+				{
+					return 0;
+				}
+				return sizeof(T);
+			}
 		}
 
-		template<>
-		uint64 GetValueSizeInternal(const FName& Value)
-		{
-			FString ValueStr = Value.ToString();
-			return GetValueSizeInternal(ValueStr);
-		}
-
-		template<>
-		uint64 GetValueSizeInternal(const FSoftObjectPath& Value)
-		{
-			FString ValueStr = Value.ToString();
-			return GetValueSizeInternal(ValueStr);
-		}
-
-		template<>
-		uint64 GetValueSizeInternal(const TArray<uint8>& Value)
-		{
-			return Value.Num();
-		}
-
-		template<>
-		uint64 GetValueSizeInternal(const TArray64<uint8>& Value)
-		{
-			return Value.Num();
-		}
-
-		void LogAttributeStorageErrors(const EAttributeStorageResult Result, const FString OperationName, const FAttributeKey AttributeKey)
+		void LogAttributeStorageErrors(const EAttributeStorageResult Result, const FString OperationName, const FAttributeKey AttributeKey)
 		{
 			//////////////////////////////////////////////////////////////////////////
 			//Errors
@@ -654,6 +1091,7 @@ namespace UE
 			AttributeStorage = Other.AttributeStorage;
 			FragmentedMemoryCost = Other.FragmentedMemoryCost;
 			DefragRatio = Other.DefragRatio;
+			AllocationCount = Other.AllocationCount;
 		}
 
 		FAttributeStorage& FAttributeStorage::operator=(const FAttributeStorage& Other)
@@ -667,6 +1105,7 @@ namespace UE
 			AttributeStorage = Other.AttributeStorage;
 			FragmentedMemoryCost = Other.FragmentedMemoryCost;
 			DefragRatio = Other.DefragRatio;
+			AllocationCount = Other.AllocationCount;
 
 			return *this;
 		}
@@ -744,13 +1183,12 @@ namespace UE
 		EAttributeStorageResult FAttributeStorage::RegisterAttribute(const FAttributeKey& ElementAttributeKey, const T& DefaultValue, EAttributeProperty AttributeProperty)
 		{
 			static_assert(TAttributeTypeTraitsInternal<T>::GetType() != EAttributeTypes::None, "T is not a supported type for the attributes. Check EAttributeTypes for the supported types");
+			static_assert(!TIsEnum<T>::Value && !TIsTEnumAsByte<T>::Value, "This function should only be instantiated for enum underlying types.");
 
 			//Lock the storage
 			FScopeLock ScopeLock(&StorageMutex);
 
-
 			const EAttributeTypes ValueType = TAttributeTypeTraitsInternal<T>::GetType();
-			const uint64 ValueSize = GetValueSizeInternal(DefaultValue);
 
 			FAttributeAllocationInfo* AttributeAllocationInfo = AttributeAllocationTable.Find(ElementAttributeKey);
 			if (AttributeAllocationInfo)
@@ -759,17 +1197,45 @@ namespace UE
 				{
 					return EAttributeStorageResult::Operation_Error_WrongType;
 				}
-				if (!AttributeStorage.IsValidIndex(AttributeAllocationInfo->Offset))
+
+				for (uint64 Offset : AttributeAllocationInfo->Offsets)
 				{
-					return EAttributeStorageResult::Operation_Error_AttributeAllocationCorrupted;
+					if (!AttributeStorage.IsValidIndex(Offset))
+					{
+						return EAttributeStorageResult::Operation_Error_AttributeAllocationCorrupted;
+					}
 				}
 			}
 			else
 			{
-				AttributeAllocationInfo = &AttributeAllocationTable.Add(ElementAttributeKey);
-				AttributeAllocationInfo->Type = ValueType;
-				AttributeAllocationInfo->Size = ValueSize;
-				AttributeAllocationInfo->Offset = AttributeStorage.AddZeroed(ValueSize);
+				// Create multiple allocations in case of nested array types
+				if constexpr (TIsArrayOfStringType<T>)
+				{
+					AttributeAllocationInfo = &AttributeAllocationTable.Add(ElementAttributeKey);
+					AttributeAllocationInfo->Type = ValueType;
+					AttributeAllocationInfo->Sizes.SetNumZeroed(DefaultValue.Num());
+					AttributeAllocationInfo->Offsets.SetNumZeroed(DefaultValue.Num());
+					AllocationCount += DefaultValue.Num();
+
+					for (int32 Index = 0; Index < DefaultValue.Num(); ++Index)
+					{
+						const uint64 ElementValueSize = GetValueSizeInternal(DefaultValue[Index]);
+
+						AttributeAllocationInfo->Sizes[Index] = ElementValueSize;
+						AttributeAllocationInfo->Offsets[Index] = (uint64)AttributeStorage.AddZeroed(ElementValueSize);
+					}
+				}
+				// Create single allocation for all other types
+				else
+				{
+					const uint64 ValueSize = GetValueSizeInternal(DefaultValue);
+
+					AttributeAllocationInfo = &AttributeAllocationTable.Add(ElementAttributeKey);
+					AttributeAllocationInfo->Type = ValueType;
+					AttributeAllocationInfo->Sizes = {ValueSize};
+					AttributeAllocationInfo->Offsets = {(uint64)AttributeStorage.AddZeroed(ValueSize)};
+					AllocationCount += 1;
+				}
 			}
 
 			//Force the specified attribute property
@@ -796,7 +1262,12 @@ namespace UE
 				return EAttributeStorageResult::Operation_Error_CannotFoundKey;
 			}
 
-			FragmentedMemoryCost += AttributeAllocationInfo->Size;
+			for (uint64 Size : AttributeAllocationInfo->Sizes)
+			{
+				FragmentedMemoryCost += Size;
+			}
+			AllocationCount -= AttributeAllocationInfo->Sizes.Num();
+
 			if (AttributeAllocationTable.Remove(ElementAttributeKey) == 0)
 			{
 				return EAttributeStorageResult::Operation_Error_CannotRemoveAttribute;
@@ -811,176 +1282,70 @@ namespace UE
 		INTERCHANGECORE_API EAttributeStorageResult FAttributeStorage::GetAttribute(const FAttributeKey& ElementAttributeKey, T& OutValue) const
 		{
 			static_assert(TAttributeTypeTraitsInternal<T>::GetType() != EAttributeTypes::None, "T is not a supported type for the attributes. Check EAttributeTypes for the supported types");
+			static_assert(!TIsEnum<T>::Value && !TIsTEnumAsByte<T>::Value, "This function should only be instantiated for enum underlying types.");
 
-			//Lock the storage
-			FScopeLock ScopeLock(&StorageMutex);
-
-			const EAttributeTypes ValueType = TAttributeTypeTraitsInternal<T>::GetType();
-			const uint64 ValueSize = GetValueSizeInternal(OutValue);
-
-			const FAttributeAllocationInfo* AttributeAllocationInfo = AttributeAllocationTable.Find(ElementAttributeKey);
-			if (AttributeAllocationInfo)
+			if constexpr (TIsTArray<T>::Value)
 			{
-				if (AttributeAllocationInfo->Type != ValueType && !((ValueType == EAttributeTypes::UInt8) && (ValueType == EAttributeTypes::Enum)))
-				{
-					return EAttributeStorageResult::Operation_Error_WrongType;
-				}
-				if (AttributeAllocationInfo->Size != ValueSize)
-				{
-					return EAttributeStorageResult::Operation_Error_WrongSize;
-				}
-				if (!AttributeStorage.IsValidIndex(AttributeAllocationInfo->Offset))
-				{
-					return EAttributeStorageResult::Operation_Error_AttributeAllocationCorrupted;
-				}
+				return GenericArrayGetAttribute(ElementAttributeKey, OutValue);
 			}
-			else
+			else // Non-array types
 			{
-				//The key do not exist
-				return EAttributeStorageResult::Operation_Error_CannotFoundKey;
-			}
-
-			const uint8* StorageData = AttributeStorage.GetData();
-			FMemory::Memcpy(&OutValue, &StorageData[AttributeAllocationInfo->Offset], AttributeAllocationInfo->Size);
-			return EAttributeStorageResult::Operation_Success;
-		}
-
-		template<>
-		INTERCHANGECORE_API EAttributeStorageResult FAttributeStorage::GetAttribute(const FAttributeKey& ElementAttributeKey, TArray<uint8>& OutValue) const
-		{
-			return GenericArrayGetAttribute(ElementAttributeKey, OutValue);
-		}
-
-		template<>
-		INTERCHANGECORE_API EAttributeStorageResult FAttributeStorage::GetAttribute(const FAttributeKey& ElementAttributeKey, TArray64<uint8>& OutValue) const
-		{
-			return GenericArrayGetAttribute(ElementAttributeKey, OutValue);
-		}
+				//Lock the storage
+				FScopeLock ScopeLock(&StorageMutex);
 
-		template<>
-		INTERCHANGECORE_API EAttributeStorageResult FAttributeStorage::GetAttribute(const FAttributeKey& ElementAttributeKey, FString& OutValue) const
-		{
-			static_assert(TAttributeTypeTraitsInternal<FString>::GetType() != EAttributeTypes::None, "Not a supported type for the attributes. Check EAttributeTypes for the supported types.");
+				const EAttributeTypes ValueType = TAttributeTypeTraitsInternal<T>::GetType();
 
-			//Lock the storage
-			FScopeLock ScopeLock(&StorageMutex);
-
-			const EAttributeTypes ValueType = TAttributeTypeTraitsInternal<FString>::GetType();
-			const FAttributeAllocationInfo* AttributeAllocationInfo = AttributeAllocationTable.Find(ElementAttributeKey);
-			if (AttributeAllocationInfo)
-			{
-				if (AttributeAllocationInfo->Type != ValueType)
+				const FAttributeAllocationInfo* AttributeAllocationInfo = AttributeAllocationTable.Find(ElementAttributeKey);
+				if (AttributeAllocationInfo)
 				{
-					return EAttributeStorageResult::Operation_Error_WrongType;
+					if (AttributeAllocationInfo->Type != ValueType)
+					{
+						return EAttributeStorageResult::Operation_Error_WrongType;
+					}
+					if (!ensure(AttributeAllocationInfo->Sizes.Num() == 1))
+					{
+						return EAttributeStorageResult::Operation_Error_WrongSize;
+					}
+					if constexpr (!TIsStringType<T>)
+					{
+						// For non-array/non-string types we should have the exact same size as the type we're retrieving
+						if (AttributeAllocationInfo->Sizes[0] != GetValueSizeInternal(OutValue))
+						{
+							return EAttributeStorageResult::Operation_Error_WrongSize;
+						}
+					}
+					if (!ensure(AttributeAllocationInfo->Offsets.Num() == 1) || !AttributeStorage.IsValidIndex(AttributeAllocationInfo->Offsets[0]))
+					{
+						return EAttributeStorageResult::Operation_Error_AttributeAllocationCorrupted;
+					}
 				}
-				if (!AttributeStorage.IsValidIndex(AttributeAllocationInfo->Offset))
+				else
 				{
-					return EAttributeStorageResult::Operation_Error_AttributeAllocationCorrupted;
+					//The key do not exist
+					return EAttributeStorageResult::Operation_Error_CannotFoundKey;
 				}
-			}
-			else
-			{
-				//The key do not exist
-				return EAttributeStorageResult::Operation_Error_CannotFoundKey;
-			}
-
-			if (AttributeAllocationInfo->Size <= sizeof(TCHAR)) // Account for the null-terminator
-			{
-				OutValue = TEXT("");
-				return EAttributeStorageResult::Operation_Success;
-			}
-
-			const FStringView ValueView = GetFStringViewAttributeFromStorage(AttributeStorage.GetData(), AttributeAllocationInfo);
-			OutValue = ValueView;
-
-			return EAttributeStorageResult::Operation_Success;
-		}
-
-		template<>
-		INTERCHANGECORE_API EAttributeStorageResult FAttributeStorage::GetAttribute(const FAttributeKey& ElementAttributeKey, FName& OutValue) const
-		{
-			static_assert(TAttributeTypeTraitsInternal<FName>::GetType() != EAttributeTypes::None, "Not a supported type for the attributes. Check EAttributeTypes for supported types.");
-
-			//Lock the storage
-			FScopeLock ScopeLock(&StorageMutex);
-
-			const EAttributeTypes ValueType = TAttributeTypeTraitsInternal<FName>::GetType();
 
-
-			const FAttributeAllocationInfo* AttributeAllocationInfo = AttributeAllocationTable.Find(ElementAttributeKey);
-			if (AttributeAllocationInfo)
-			{
-				if (AttributeAllocationInfo->Type != ValueType)
-				{
-					return EAttributeStorageResult::Operation_Error_WrongType;
-				}
-				if (!AttributeStorage.IsValidIndex(AttributeAllocationInfo->Offset))
+				if constexpr (TIsStringType<T>)
 				{
-					return EAttributeStorageResult::Operation_Error_AttributeAllocationCorrupted;
-				}
-			}
-			else
-			{
-				//The key do not exist
-				return EAttributeStorageResult::Operation_Error_CannotFoundKey;
-			}
-
-
-			if (AttributeAllocationInfo->Size == 0)
-			{
-				OutValue = NAME_None;
-				return EAttributeStorageResult::Operation_Success;
-			}
-
-			const FStringView ValueStr = GetFStringViewAttributeFromStorage(AttributeStorage.GetData(), AttributeAllocationInfo);
-			OutValue = FName(ValueStr);
-
-			return EAttributeStorageResult::Operation_Success;
-		}
-
-		template<>
-		INTERCHANGECORE_API EAttributeStorageResult FAttributeStorage::GetAttribute(const FAttributeKey& ElementAttributeKey, FSoftObjectPath& OutValue) const
-		{
-			static_assert(TAttributeTypeTraitsInternal<FSoftObjectPath>::GetType() != EAttributeTypes::None, "Not a supported type for the attributes. Check EAttributeTypes for supported types.");
-
-			//Lock the storage
-			FScopeLock ScopeLock(&StorageMutex);
-
-			const EAttributeTypes ValueType = TAttributeTypeTraitsInternal<FSoftObjectPath>::GetType();
-
+					if (AttributeAllocationInfo->Sizes[0] <= sizeof(TCHAR))	   // Account for the null-terminator
+					{
+						OutValue = T();
+						return EAttributeStorageResult::Operation_Success;
+					}
 
-			const FAttributeAllocationInfo* AttributeAllocationInfo = AttributeAllocationTable.Find(ElementAttributeKey);
-			if (AttributeAllocationInfo)
-			{
-				if (AttributeAllocationInfo->Type != ValueType)
-				{
-					return EAttributeStorageResult::Operation_Error_WrongType;
+					const FStringView ValueView = GetFStringViewAttributeFromStorage(AttributeStorage.GetData(), AttributeAllocationInfo);
+					OutValue = T{ValueView};
 				}
-				if (!AttributeStorage.IsValidIndex(AttributeAllocationInfo->Offset))
+				else
 				{
-					return EAttributeStorageResult::Operation_Error_AttributeAllocationCorrupted;
+					const uint8* StorageData = AttributeStorage.GetData();
+					FMemory::Memcpy(&OutValue, &StorageData[AttributeAllocationInfo->Offsets[0]], AttributeAllocationInfo->Sizes[0]);
 				}
-			}
-			else
-			{
-				//The key do not exist
-				return EAttributeStorageResult::Operation_Error_CannotFoundKey;
-			}
 
-
-			if (AttributeAllocationInfo->Size == 0)
-			{
-				OutValue = TEXT("");
 				return EAttributeStorageResult::Operation_Success;
 			}
-
-			const FStringView ValueStr = GetFStringViewAttributeFromStorage(AttributeStorage.GetData(), AttributeAllocationInfo);
-			OutValue = FSoftObjectPath(ValueStr);
-
-			return EAttributeStorageResult::Operation_Success;
 		}
 
-
 		template<typename T, typename>
 		FAttributeStorage::TAttributeHandle<T> FAttributeStorage::GetAttributeHandle(const FAttributeKey& ElementAttributeKey) const
 		{
@@ -1125,14 +1490,23 @@ namespace UE
 
 				if (SourceAttributeInfo)
 				{
-					if (!SourceStorage.AttributeStorage.IsValidIndex(SourceAttributeInfo->Offset))
+					bool bSkipAttribute = false;
+					for (uint64 Offset : SourceAttributeInfo->Offsets)
+					{
+						if (!SourceStorage.AttributeStorage.IsValidIndex(Offset))
+						{
+							// EAttributeStorageResult::Operation_Error_AttributeAllocationCorrupted;
+							bSkipAttribute = true;
+							break;
+						}
+					}
+					if (bSkipAttribute)
 					{
-						//EAttributeStorageResult::Operation_Error_AttributeAllocationCorrupted;
 						continue;
 					}
 
 					const EAttributeTypes SourceValueType = SourceAttributeInfo->Type;
-					const uint64 SourceValueSize = SourceAttributeInfo->Size;
+					const TArray<uint64>& SourceValueSizes = SourceAttributeInfo->Sizes;
 
 					// Search first for the original key
 					FAttributeStorage::FAttributeAllocationInfo* DestinationAttributeInfo = nullptr;
@@ -1158,14 +1532,23 @@ namespace UE
 							//EAttributeStorageResult::Operation_Error_WrongType;
 							continue;
 						}
-						if (DestinationAttributeInfo->Size != SourceValueSize)
+						if (DestinationAttributeInfo->Sizes != SourceValueSizes)
 						{
 							//EAttributeStorageResult::Operation_Error_WrongSize;
 							continue;
 						}
-						if (!DestinationStorage.AttributeStorage.IsValidIndex(DestinationAttributeInfo->Offset))
+
+						for (uint64 Offset : DestinationAttributeInfo->Offsets)
+						{
+							if (!DestinationStorage.AttributeStorage.IsValidIndex(Offset))
+							{
+								// EAttributeStorageResult::Operation_Error_AttributeAllocationCorrupted;
+								bSkipAttribute = true;
+								break;
+							}
+						}
+						if (bSkipAttribute)
 						{
-							//EAttributeStorageResult::Operation_Error_AttributeAllocationCorrupted;
 							continue;
 						}
 					}
@@ -1180,8 +1563,13 @@ namespace UE
 							DestinationAttributeInfo = &DestinationStorage.AttributeAllocationTable.Add(AttributeKey);
 						}
 						DestinationAttributeInfo->Type = SourceValueType;
-						DestinationAttributeInfo->Size = SourceValueSize;
-						DestinationAttributeInfo->Offset = DestinationStorage.AttributeStorage.AddZeroed(SourceValueSize);
+						DestinationAttributeInfo->Sizes = SourceValueSizes;
+
+						DestinationAttributeInfo->Offsets.Reset(DestinationAttributeInfo->Sizes.Num());
+						for (uint64 Size : DestinationAttributeInfo->Sizes)
+						{
+							DestinationAttributeInfo->Offsets.Add(DestinationStorage.AttributeStorage.AddZeroed(Size));
+						}
 					}
 
 					//Force the specified attribute property
@@ -1190,11 +1578,20 @@ namespace UE
 					DestinationAttributeInfo->Hash = SourceAttributeInfo->Hash;
 
 					//Memcpy source to destination storage
-					const uint8* SourceStorageData = SourceStorage.AttributeStorage.GetData();
-					uint8* DestinationStorageData = DestinationStorage.AttributeStorage.GetData();
-					FMemory::Memcpy(&DestinationStorageData[DestinationAttributeInfo->Offset], &SourceStorageData[SourceAttributeInfo->Offset], SourceValueSize);
+					for (int32 AllocIndex = 0; AllocIndex < SourceValueSizes.Num(); ++AllocIndex)
+					{
+						const uint8* SourceStorageData = SourceStorage.AttributeStorage.GetData();
+						uint8* DestinationStorageData = DestinationStorage.AttributeStorage.GetData();
+						FMemory::Memcpy(
+							&DestinationStorageData[DestinationAttributeInfo->Offsets[AllocIndex]], 
+							&SourceStorageData[SourceAttributeInfo->Offsets[AllocIndex]], 
+							SourceValueSizes[AllocIndex]
+						);
+					}
 				}
 			}
+
+			DestinationStorage.UpdateAllocationCount();
 		}
 
 		void FAttributeStorage::CopyStorageAttributes(const FAttributeStorage& SourceStorage, FAttributeStorage& DestinationStorage, const TArray<TPair<FAttributeKey, FAttributeKey>>& AttributeKeys)
@@ -1235,36 +1632,14 @@ namespace UE
 			return SetAttribute(AttributeAllocationInfo, Value);
 		}
 
-		void FAttributeStorage::ExtractFStringAttributeFromStorage(const uint8* StorageData, const FAttributeAllocationInfo* AttributeAllocationInfo, FString& OutValue) const
+		const FStringView FAttributeStorage::GetFStringViewAttributeFromStorage(const uint8* StorageData, const FAttributeAllocationInfo* AttributeAllocationInfo, int32 ElementIndex) const
 		{
-			//Allocate the ByteArray
-			const uint64 NumberOfChar = AttributeAllocationInfo->Size / sizeof(TCHAR);
-			check(NumberOfChar > 0);
-
-			//Since the null character is always there, a empty String is only one TCHAR
-			if (NumberOfChar <= 1)
-			{
-				OutValue.Empty();
-				return;
-			}
-
-			//Empty and pre allocate the FString buffer (TArray<TCHAR>) and add the null terminator
-			OutValue.Reset(NumberOfChar - 1);
-			//SLOW! We have to add the characters 1 x 1, there is no AddZeroed or similar stuff. We do not have a string to use the operator= which do what we need here
-			//We go up to NumberOfChar -1 because there is always a null terminating character add at the end of the FString
-			for (int32 CharIndex = 0; CharIndex < NumberOfChar - 1; ++CharIndex)
+			if (!ensure(AttributeAllocationInfo->Sizes.IsValidIndex(ElementIndex)) || !ensure(AttributeAllocationInfo->Offsets.IsValidIndex(ElementIndex)))
 			{
-				//Adding spaces, we cannot add null character because the null character are not added
-				OutValue.AppendChar(' ');
+				return {};
 			}
 
-			//Copy into the buffer
-			FMemory::Memcpy((uint8*)(*OutValue), &StorageData[AttributeAllocationInfo->Offset], AttributeAllocationInfo->Size);
-		}
-
-		const FStringView FAttributeStorage::GetFStringViewAttributeFromStorage(const uint8* StorageData, const FAttributeAllocationInfo* AttributeAllocationInfo) const
-		{
-			const uint64 NumberOfChar = AttributeAllocationInfo->Size / sizeof(TCHAR);
+			const uint64 NumberOfChar = AttributeAllocationInfo->Sizes[ElementIndex] / sizeof(TCHAR);
 			check(NumberOfChar > 0);
 
 			// Null terminator is included in the size
@@ -1273,79 +1648,98 @@ namespace UE
 				return FStringView();
 			}
 
-			const TCHAR* StringContents = reinterpret_cast<const TCHAR*>(&StorageData[AttributeAllocationInfo->Offset]);
+			const TCHAR* StringContents = reinterpret_cast<const TCHAR*>(&StorageData[AttributeAllocationInfo->Offsets[ElementIndex]]);
 			return FStringView(StringContents, int32(NumberOfChar - 1));
 		}
 
-		/** Set a multisize(TArray<uint8>, FString ...) attribute value into the storage. Return success if the attribute was properly set. */
+		/**
+		 * Set a multisize(TArray<uint8>, FString ...) attribute value into the storage at the allocation with TargetAllocationIndex.
+		 * Return success if the attribute was properly set.
+		 *
+		 * WARNING: Expects the caller to lock the StorageMutex
+		 */
 		template<typename MultiSizeType>
-		EAttributeStorageResult FAttributeStorage::MultiSizeSetAttribute(FAttributeAllocationInfo* AttributeAllocationInfo, const MultiSizeType& Value, TFunctionRef<uint8* ()> GetValuePointer)
+		EAttributeStorageResult FAttributeStorage::MultiSizeSetAttribute(
+			FAttributeAllocationInfo* AttributeAllocationInfo,
+			int32 TargetAllocationIndex,
+			const MultiSizeType& Value,
+			const uint8* SourceDataPtr,
+			bool& bOutNeedsDefrag
+		)
 		{
-			//Lock the storage
-			FScopeLock ScopeLock(&StorageMutex);
-
 			const EAttributeTypes ValueType = TAttributeTypeTraitsInternal<MultiSizeType>::GetType();
+			const uint64 ValueSize = GetValueSizeInternal(Value);
 
 			if (!AttributeAllocationInfo)
 			{
 				return EAttributeStorageResult::Operation_Error_CannotFoundKey;
 			}
 
-			if (AttributeAllocationInfo->Type != ValueType)
+			//For all string-related types (listed below) we'll be setting their individual values with individual FStrings, whether
+			//they're arrays or just a single value
+			if constexpr (TIsStringType<MultiSizeType>)
+			{
+				//We should never be trying to set with actual FNames / FSoftObjectPaths in here, only with FStrings
+				static_assert(std::is_same_v<MultiSizeType, FString>);
+
+				const static TSet<EAttributeTypes> AllowedTypes = {
+					EAttributeTypes::String,
+					EAttributeTypes::StringArray,
+					EAttributeTypes::Name,
+					EAttributeTypes::NameArray,
+					EAttributeTypes::SoftObjectPath,
+					EAttributeTypes::SoftObjectPathArray
+				};
+				if (!AllowedTypes.Contains(AttributeAllocationInfo->Type))
+				{
+					return EAttributeStorageResult::Operation_Error_WrongType;
+				}
+			}
+			else if (AttributeAllocationInfo->Type != ValueType)
 			{
 				return EAttributeStorageResult::Operation_Error_WrongType;
 			}
 
-			const uint64 ValueSize = GetValueSizeInternal(Value);
-
-			uint8* DataPtr = GetValuePointer();
 			//We must have a valid data pointer if we have something to copy into the storage.
 			//This case support empty array or empty string because the ValueSize will be 0.
-			if (!DataPtr && ValueSize > 0)
+			if (!SourceDataPtr && ValueSize > 0)
 			{
 				return EAttributeStorageResult::Operation_Error_InvalidMultiSizeValueData;
 			}
 
-			//If the new size is greater then the old size we have to create a new allocation info entry and delete the old one
-			if (ValueSize > AttributeAllocationInfo->Size)
+			//Make sure we have the desired allocation index
+			if (TargetAllocationIndex >= AttributeAllocationInfo->Sizes.Num())
 			{
-				//Find the key to remove it
-				const FAttributeKey* ElementAttributeKeyPtr = AttributeAllocationTable.FindKey(*AttributeAllocationInfo);
-				if (!ElementAttributeKeyPtr)
-				{
-					return EAttributeStorageResult::Operation_Error_CannotFoundKey;
-				}
-				const FAttributeKey OriginalKey = *ElementAttributeKeyPtr;
-				const EAttributeProperty AttributePropertyBackup = AttributeAllocationInfo->Property;
-				const EAttributeStorageResult RemoveResult = UnregisterAttribute(OriginalKey);
-				if (!IsAttributeStorageResultSuccess(RemoveResult))
-				{
-					return RemoveResult;
-				}
-				AttributeAllocationInfo = nullptr;
-				//Reallocate a new one
-				AttributeAllocationInfo = &AttributeAllocationTable.Add(OriginalKey);
-				AttributeAllocationInfo->Type = ValueType;
-				AttributeAllocationInfo->Offset = AttributeStorage.AddZeroed(ValueSize);
-				//Force the specified attribute property
-				AttributeAllocationInfo->Property = AttributePropertyBackup;
+				AttributeAllocationInfo->Sizes.SetNumZeroed(TargetAllocationIndex + 1);
+				AttributeAllocationInfo->Offsets.SetNumZeroed(TargetAllocationIndex + 1);
+			}
+
+			//If the new size is greater then the old size for the target allocation index, we have to create a new allocation entry and delete the old one
+			if (ValueSize > AttributeAllocationInfo->Sizes[TargetAllocationIndex])
+			{
+				FragmentedMemoryCost += AttributeAllocationInfo->Sizes[TargetAllocationIndex];
+				bOutNeedsDefrag = true;
+
+				AttributeAllocationInfo->Sizes[TargetAllocationIndex] = ValueSize;
+				AttributeAllocationInfo->Offsets[TargetAllocationIndex] = AttributeStorage.AddZeroed(ValueSize);
 			}
 			else
 			{
 				//In case we reuse the allocation table, simply adjust the waste memory counter
-				FragmentedMemoryCost += AttributeAllocationInfo->Size - ValueSize;
-			}
-
-			//We can recycle allocation table if the size is equal or smaller, so we have to always set it
-			AttributeAllocationInfo->Size = ValueSize;
+				FragmentedMemoryCost += AttributeAllocationInfo->Sizes[TargetAllocationIndex] - ValueSize;
+				bOutNeedsDefrag = true;
 
-			//Set the hash from the value
-			AttributeAllocationInfo->Hash = GetValueHash(DataPtr, ValueSize);
+				AttributeAllocationInfo->Sizes[TargetAllocationIndex] = ValueSize;
+			}
 
 			if (ValueSize > 0)
 			{
 				uint8* StorageData = AttributeStorage.GetData();
-				FMemory::Memcpy(&StorageData[AttributeAllocationInfo->Offset], DataPtr, ValueSize);
+				FMemory::Memcpy(
+					&StorageData[AttributeAllocationInfo->Offsets[TargetAllocationIndex]],
+					SourceDataPtr,
+					ValueSize
+				);
 			}
 
 			return EAttributeStorageResult::Operation_Success;
@@ -1366,51 +1760,124 @@ namespace UE
 				{
 					return EAttributeStorageResult::Operation_Error_WrongType;
 				}
-				if (!AttributeStorage.IsValidIndex(AttributeAllocationInfo->Offset))
+				if (AttributeAllocationInfo->Offsets.Num() != AttributeAllocationInfo->Sizes.Num())
 				{
-					return EAttributeStorageResult::Operation_Error_AttributeAllocationCorrupted;
+					// For now we expect all generic array types (i.e. non-string array types) to have a single allocation
+					return EAttributeStorageResult::Operation_Error_WrongSize;
+				}
+				for (uint64 Offset : AttributeAllocationInfo->Offsets)
+				{
+					if (!AttributeStorage.IsValidIndex(Offset))
+					{
+						return EAttributeStorageResult::Operation_Error_AttributeAllocationCorrupted;
+					}
 				}
 			}
 			else
 			{
-				//The key do not exist
+				//The key does not exist
 				return EAttributeStorageResult::Operation_Error_CannotFoundKey;
 			}
 
-			if (AttributeAllocationInfo->Size == 0)
+			constexpr bool bOneAllocationPerElement = TIsArrayOfStringType<ArrayType>;
+			if constexpr (bOneAllocationPerElement)
 			{
-				OutValue.Empty();
-				return EAttributeStorageResult::Operation_Success;
+				const uint64 NumAllocations = AttributeAllocationInfo->Sizes.Num();
+				const uint64 NumElements = NumAllocations;
+
+				const uint8* StorageData = AttributeStorage.GetData();
+
+				OutValue.SetNum(NumElements);
+				for (int32 AllocIndex = 0; AllocIndex < NumAllocations; ++AllocIndex)
+				{
+					if constexpr (TIsStringType<typename ArrayType::ElementType>)
+					{
+						const FStringView ValueView = GetFStringViewAttributeFromStorage(StorageData, AttributeAllocationInfo, AllocIndex);
+						OutValue[AllocIndex] = typename ArrayType::ElementType{ValueView};
+					}
+					else
+					{
+						// This should never happen for now as we only have "nested arrays" of string types
+						ensureMsgf(false, TEXT("Actual nested TArrays are not supported!"));
+					}
+				}
+			}
+			else // One allocation for the whole array (so most types, like TArray<float>, TArray<FTransform3f>, etc.)
+			{
+				const uint64 NumAllocations = AttributeAllocationInfo->Sizes.Num();
+				ensure(NumAllocations <= 1);
+
+				if (NumAllocations == 0 || AttributeAllocationInfo->Sizes[0] == 0)
+				{
+					OutValue.Empty();
+					return EAttributeStorageResult::Operation_Success;
+				}
+
+				const uint64 AllocationSize = AttributeAllocationInfo->Sizes[0];
+				const uint64 NumElements = AllocationSize / ArrayType::GetTypeSize();
+				OutValue.SetNumUninitialized(NumElements);
+
+				const uint8* StorageData = AttributeStorage.GetData();
+				const uint8* SourceBytes = &StorageData[AttributeAllocationInfo->Offsets[0]];
+				FMemory::Memcpy(OutValue.GetData(), SourceBytes, AllocationSize);
 			}
 
-			//Allocate the ByteArray
-			OutValue.AddZeroed(static_cast<typename ArrayType::SizeType>(AttributeAllocationInfo->Size));
-			//Copy into the buffer
-			const uint8* StorageData = AttributeStorage.GetData();
-			FMemory::Memcpy(OutValue.GetData(), &StorageData[AttributeAllocationInfo->Offset], AttributeAllocationInfo->Size);
 			return EAttributeStorageResult::Operation_Success;
 		}
 
+		// WARNING: This may invalidate references to elements of the AttributeAllocationTable, as it will compact its pairs
 		void FAttributeStorage::DefragInternal()
 		{
 			//Defrag if the fragmented memory cost is bigger then the DefragRatio.
+			//TODO: Why 11? Maybe use a cvar?
 			if (FragmentedMemoryCost < 11 || FragmentedMemoryCost < AttributeStorage.Num() * DefragRatio)
 			{
 				return;
 			}
 
+			//Linearize our allocations, since FAttributeAllocationInfo may now hold multiple allocations each, that could
+			//all be interleaved with other attributes'
+			struct FLinearAllocation
+			{
+				FAttributeKey Key;
+				uint32 Index;
+				uint64 Offset;
+				uint64 Size;
+			};
+			TArray<FLinearAllocation> LinearAllocations;
+			LinearAllocations.Reserve(AttributeAllocationTable.Num());
+			for (const TPair<FAttributeKey, FAttributeAllocationInfo>& Pair : AttributeAllocationTable)
+			{
+				const FAttributeKey& AllocKey = Pair.Key;
+				const FAttributeAllocationInfo& AllocInfo = Pair.Value;
+
+				int32 NumAllocs = AllocInfo.Sizes.Num();
+				if (!ensure(NumAllocs == AllocInfo.Offsets.Num()))
+				{
+					return;
+				}
+
+				for (int32 AllocIndex = 0; AllocIndex < NumAllocs; ++AllocIndex)
+				{
+					FLinearAllocation& LinearAlloc = LinearAllocations.Emplace_GetRef();
+					LinearAlloc.Key = AllocKey;
+					LinearAlloc.Index = AllocIndex;
+					LinearAlloc.Offset = AllocInfo.Offsets[AllocIndex];
+					LinearAlloc.Size = AllocInfo.Sizes[AllocIndex];
+				}
+			}
+
 			//Sorts the allocation table per offset since we want to defrag using memory block
-			//The ValueSort is also calling CompactStable which remove all holes in the TMap
-			struct FAttributeAllocationInfoSmaller
+			LinearAllocations.Sort([](const FLinearAllocation& A, const FLinearAllocation& B)
 			{
-				FORCEINLINE bool operator()(const FAttributeAllocationInfo& A, const FAttributeAllocationInfo& B) const
+				if (A.Offset == B.Offset) 
 				{
-					return A.Offset < B.Offset;
+					return A.Size < B.Size;
 				}
-			};
-			AttributeAllocationTable.ValueSort(FAttributeAllocationInfoSmaller());
 
-			const uint64 AttributeNumber = AttributeAllocationTable.Num();
+				return A.Offset < B.Offset;
+			});
+
 			uint8* StorageData = AttributeStorage.GetData();
 			//Number of attribute we have done
 			uint64 AttributeCount = 0;
@@ -1422,31 +1889,31 @@ namespace UE
 			uint64 MemmoveBlockStart = 0;
 			//Replace CurrentOffset when we are moving a block
 			uint64 MemmoveBlockCurrentOffset = 0;
-			for (auto& Kvp : AttributeAllocationTable)
+			const uint64 NumAllocations = LinearAllocations.Num();
+			for (FLinearAllocation& LinearAllocation : LinearAllocations)
 			{
 				AttributeCount++;
-				const bool bLast = AttributeCount == AttributeNumber;
+				const bool bLast = AttributeCount == NumAllocations;
 
-				FAttributeAllocationInfo& AttributeAllocationInfo = Kvp.Value;
-				if (AttributeAllocationInfo.Offset == CurrentOffset)
+				if (LinearAllocation.Offset == CurrentOffset)
 				{
-					CurrentOffset += AttributeAllocationInfo.Size;
+					CurrentOffset += LinearAllocation.Size;
 					continue;
 				}
-				check(CurrentOffset < AttributeAllocationInfo.Offset);
+				check(CurrentOffset < LinearAllocation.Offset);
 				if (!bBuildBlock)
 				{
 					bBuildBlock = true;
-					MemmoveBlockStart = AttributeAllocationInfo.Offset;
-					AttributeAllocationInfo.Offset = CurrentOffset;
-					MemmoveBlockCurrentOffset = MemmoveBlockStart + AttributeAllocationInfo.Size;
+					MemmoveBlockStart = LinearAllocation.Offset;
+					LinearAllocation.Offset = CurrentOffset;
+					MemmoveBlockCurrentOffset = MemmoveBlockStart + LinearAllocation.Size;
 				}
 				else
 				{
-					if (AttributeAllocationInfo.Offset == MemmoveBlockCurrentOffset)
+					if (LinearAllocation.Offset == MemmoveBlockCurrentOffset)
 					{
-						AttributeAllocationInfo.Offset -= MemmoveBlockStart - CurrentOffset;
-						MemmoveBlockCurrentOffset += AttributeAllocationInfo.Size;
+						LinearAllocation.Offset -= MemmoveBlockStart - CurrentOffset;
+						MemmoveBlockCurrentOffset += LinearAllocation.Size;
 					}
 					else
 					{
@@ -1455,9 +1922,9 @@ namespace UE
 						FMemory::Memmove(&StorageData[CurrentOffset], &StorageData[MemmoveBlockStart], BlockSize);
 						CurrentOffset += BlockSize;
 
-						MemmoveBlockStart = AttributeAllocationInfo.Offset;
-						AttributeAllocationInfo.Offset = CurrentOffset;
-						MemmoveBlockCurrentOffset = MemmoveBlockStart + AttributeAllocationInfo.Size;
+						MemmoveBlockStart = LinearAllocation.Offset;
+						LinearAllocation.Offset = CurrentOffset;
+						MemmoveBlockCurrentOffset = MemmoveBlockStart + LinearAllocation.Size;
 					}
 				}
 
@@ -1473,96 +1940,167 @@ namespace UE
 					AttributeStorage.RemoveAt(RemoveStartIndex, RemoveCount);
 				}
 			}
+
+			//Update our allocation count since it's basically free to do right now
+			AllocationCount = LinearAllocations.Num();
+
 			//Reset the fragmented cost
 			FragmentedMemoryCost = 0;
+
+			//Remove holes in the TMap
+			AttributeAllocationTable.CompactStable();
+
+			//Update our allocation table with the defragged offsets / sizes
+			for (const FLinearAllocation& Alloc : LinearAllocations)
+			{
+				FAttributeAllocationInfo* FoundInfo = AttributeAllocationTable.Find(Alloc.Key);
+				if (ensure(FoundInfo))
+				{
+					if (ensure(FoundInfo->Sizes.IsValidIndex(Alloc.Index) && FoundInfo->Offsets.IsValidIndex(Alloc.Index)))
+					{
+						FoundInfo->Sizes[Alloc.Index] = Alloc.Size;
+						FoundInfo->Offsets[Alloc.Index] = Alloc.Offset;
+					}
+				}
+			}
 		}
 
 		template<typename T>
 		EAttributeStorageResult FAttributeStorage::SetAttribute(FAttributeAllocationInfo* AttributeAllocationInfo, const T& Value)
 		{
 			static_assert(TAttributeTypeTraitsInternal<T>::GetType() != EAttributeTypes::None, "T is not a supported type for the attributes. Check EAttributeTypes for the supported types");
-
-			//Lock the storage
-			FScopeLock ScopeLock(&StorageMutex);
+			static_assert(!TIsEnum<T>::Value && !TIsTEnumAsByte<T>::Value, "This function should only be instantiated for enum underlying types.");
 
 			if (!AttributeAllocationInfo)
 			{
 				return EAttributeStorageResult::Operation_Error_CannotFoundKey;
 			}
-			const EAttributeTypes ValueType = TAttributeTypeTraitsInternal<T>::GetType();
-			const uint64 ValueSize = GetValueSizeInternal(Value);
 
-			if (AttributeAllocationInfo->Type != ValueType)
+			FScopeLock ScopeLock(&StorageMutex);
+
+			bool bNeedsDefrag = false;
+
+			EAttributeStorageResult Result = EAttributeStorageResult::Operation_Success;
+			if constexpr (TIsArrayOfStringType<T>)
 			{
-				return EAttributeStorageResult::Operation_Error_WrongType;
+				AllocationCount -= AttributeAllocationInfo->Sizes.Num();
+				AllocationCount += Value.Num();
+
+				// We only need as many allocations as Value has elements: Discard the extra ones
+				for (int32 Index = AttributeAllocationInfo->Sizes.Num() - 1; Index >= Value.Num(); --Index)
+				{
+					FragmentedMemoryCost += AttributeAllocationInfo->Sizes[Index];
+					bNeedsDefrag = true;
+				}
+				AttributeAllocationInfo->Sizes.SetNum(Value.Num());
+				AttributeAllocationInfo->Offsets.SetNum(Value.Num());
+
+				// Set one allocation for each Value element
+				for (int32 Index = 0; Index < Value.Num(); ++Index)
+				{
+					if constexpr (std::is_same_v<typename T::ElementType, FString>)
+					{
+						const FString& Element = Value[Index];
+						Result = MultiSizeSetAttribute(AttributeAllocationInfo, Index, Element, (uint8*)(*Element), bNeedsDefrag);
+					}
+					else
+					{
+						const FString& Element = Value[Index].ToString();
+						Result = MultiSizeSetAttribute(AttributeAllocationInfo, Index, Element, (uint8*)(*Element), bNeedsDefrag);
+					}
+
+					if (Result != EAttributeStorageResult::Operation_Success)
+					{
+						break;
+					}
+				}
 			}
-			if (AttributeAllocationInfo->Size != ValueSize)
+			else if constexpr (TIsTArray<T>::Value)
 			{
-				return EAttributeStorageResult::Operation_Error_WrongSize;
+				const int32 TargetAllocationIndex = 0;
+				Result = MultiSizeSetAttribute(AttributeAllocationInfo, TargetAllocationIndex, Value, (uint8*)(Value.GetData()), bNeedsDefrag);
 			}
-			if (!AttributeStorage.IsValidIndex(AttributeAllocationInfo->Offset))
+			else if constexpr (std::is_same_v<T, FString>)
 			{
-				return EAttributeStorageResult::Operation_Error_AttributeAllocationCorrupted;
+				const int32 TargetAllocationIndex = 0;
+				Result = MultiSizeSetAttribute(AttributeAllocationInfo, TargetAllocationIndex, Value, (uint8*)(*Value), bNeedsDefrag);
 			}
+			else if constexpr (std::is_same_v<T, FName> || std::is_same_v<T, FSoftObjectPath>)
+			{
+				//Value must be stored as a FString for persistence
+				const int32 TargetAllocationIndex = 0;
+				FString ValueStr = Value.ToString();
+				Result = MultiSizeSetAttribute(AttributeAllocationInfo, TargetAllocationIndex, ValueStr, (uint8*)(*ValueStr), bNeedsDefrag);
+			}
+			else //Other non-string elements (float, FGuid, FTransform3f, etc.)
+			{
+				const EAttributeTypes ValueType = TAttributeTypeTraitsInternal<T>::GetType();
+				const uint64 ValueSize = GetValueSizeInternal(Value);
+
+				if (AttributeAllocationInfo->Type != ValueType)
+				{
+					Result = EAttributeStorageResult::Operation_Error_WrongType;
+				}
+				if (!ensure(AttributeAllocationInfo->Sizes.Num() == 1) || AttributeAllocationInfo->Sizes[0] != ValueSize)
+				{
+					Result = EAttributeStorageResult::Operation_Error_WrongSize;
+				}
+				if (!ensure(AttributeAllocationInfo->Offsets.Num() == 1) || !AttributeStorage.IsValidIndex(AttributeAllocationInfo->Offsets[0]))
+				{
+					Result = EAttributeStorageResult::Operation_Error_AttributeAllocationCorrupted;
+				}
 
-			//Set the hash from the value
-			AttributeAllocationInfo->Hash = GetValueHash(reinterpret_cast<const uint8*>(&Value), ValueSize);
+				if (Result == EAttributeStorageResult::Operation_Success)
+				{
+					uint8* StorageData = AttributeStorage.GetData();
+					FMemory::Memcpy(&StorageData[AttributeAllocationInfo->Offsets[0]], &Value, ValueSize);
+				}
+			}
 
-			uint8* StorageData = AttributeStorage.GetData();
-			FMemory::Memcpy(&StorageData[AttributeAllocationInfo->Offset], &Value, ValueSize);
-			return EAttributeStorageResult::Operation_Success;
-		}
+			UpdateAllocationInfoHash(*AttributeAllocationInfo);
 
-		template<>
-		EAttributeStorageResult FAttributeStorage::SetAttribute(FAttributeAllocationInfo* AttributeAllocationInfo, const TArray<uint8>& Value)
-		{
-			return MultiSizeSetAttribute(AttributeAllocationInfo, Value, [&Value]()->uint8*
+			// We can only defrag after we're done since it may potentially invalidate our AttributeAllocationInfo pointer
+			// as it will internally compact the allocation table
+			if (bNeedsDefrag)
 			{
-				return const_cast<uint8*>(Value.GetData());
-			});
-		}
+				DefragInternal();
+			}
 
-		template<>
-		EAttributeStorageResult FAttributeStorage::SetAttribute(FAttributeAllocationInfo* AttributeAllocationInfo, const TArray64<uint8>& Value)
-		{
-			return MultiSizeSetAttribute(AttributeAllocationInfo, Value, [&Value]()->uint8*
-			{
-				return const_cast<uint8*>(Value.GetData());
-			});
+			return Result;
 		}
 
-		template<>
-		EAttributeStorageResult FAttributeStorage::SetAttribute(FAttributeAllocationInfo* AttributeAllocationInfo, const FString& Value)
+		void FAttributeStorage::UpdateAllocationInfoHash(FAttributeStorage::FAttributeAllocationInfo& AllocationInfo)
 		{
-			return MultiSizeSetAttribute(AttributeAllocationInfo, Value, [&Value]()->uint8*
-			{
-				return (uint8*)(*Value);
-			});
-		}
+			ensure(AllocationInfo.Offsets.Num() == AllocationInfo.Sizes.Num());
 
-		template<>
-		EAttributeStorageResult FAttributeStorage::SetAttribute(FAttributeAllocationInfo* AttributeAllocationInfo, const FName& Value)
-		{
-			//FName must be store has a FString for persistence
-			FString ValueStr = Value.ToString();
-			return MultiSizeSetAttribute(AttributeAllocationInfo, Value, [&ValueStr]()->uint8*
+			const uint8* StorageData = AttributeStorage.GetData();
+
+			FSHA1 Sha;
+			for (int32 AllocIndex = 0; AllocIndex < AllocationInfo.Sizes.Num(); ++AllocIndex)
 			{
-				return (uint8*)(*ValueStr);
-			});
+				Sha.Update(
+					&StorageData[AllocationInfo.Offsets[AllocIndex]], 
+					AllocationInfo.Sizes[AllocIndex]
+				);
+				Sha.Final();
+			}
+
+			// Retrieve the hash and use it to construct a pseudo-GUID.
+			uint32 Hash[5];
+			Sha.GetHash(reinterpret_cast<uint8*>(Hash));
+			AllocationInfo.Hash = FGuid(Hash[0] ^ Hash[4], Hash[1], Hash[2], Hash[3]);
 		}
 
-		template<>
-		EAttributeStorageResult FAttributeStorage::SetAttribute(FAttributeAllocationInfo* AttributeAllocationInfo, const FSoftObjectPath& Value)
+		void FAttributeStorage::UpdateAllocationCount()
 		{
-			//FSoftObjectPath can be hold has a FString and restore from it
-			FString ValueStr = Value.ToString();
-			return MultiSizeSetAttribute(AttributeAllocationInfo, Value, [&ValueStr]()->uint8*
+			AllocationCount = 0;
+			for (const TPair<FAttributeKey, FAttributeAllocationInfo>& Pair : AttributeAllocationTable)
 			{
-				return (uint8*)(*ValueStr);
-			});
+				AllocationCount += Pair.Value.Sizes.Num();
+			}
 		}
 
-		#define SPECIALIZE_SINGLE(X) \
+#define SPECIALIZE(X) \
 			template INTERCHANGECORE_API EAttributeStorageResult FAttributeStorage::RegisterAttribute<X>(const FAttributeKey&, const X&, EAttributeProperty);\
 			template INTERCHANGECORE_API EAttributeStorageResult FAttributeStorage::GetAttribute<X>(const FAttributeKey&, X&) const;\
 			template INTERCHANGECORE_API FAttributeStorage::TAttributeHandle<X>::TAttributeHandle();\
@@ -1573,75 +2111,118 @@ namespace UE
 			template INTERCHANGECORE_API const FAttributeKey& FAttributeStorage::TAttributeHandle<X>::GetKey() const;\
 			template INTERCHANGECORE_API FAttributeStorage::TAttributeHandle<X> FAttributeStorage::GetAttributeHandle<X>(const FAttributeKey&) const;\
 			template INTERCHANGECORE_API FString UE::Interchange::AttributeValueToString(const X&);
-		SPECIALIZE_SINGLE(bool);
-		SPECIALIZE_SINGLE(FColor);
-		SPECIALIZE_SINGLE(FDateTime);
-		SPECIALIZE_SINGLE(double);
-		SPECIALIZE_SINGLE(float);
-		SPECIALIZE_SINGLE(FGuid);
-		SPECIALIZE_SINGLE(int8);
-		SPECIALIZE_SINGLE(int16);
-		SPECIALIZE_SINGLE(int32);
-		SPECIALIZE_SINGLE(int64);
-		SPECIALIZE_SINGLE(FIntRect);
-		SPECIALIZE_SINGLE(FLinearColor);
-		SPECIALIZE_SINGLE(FRandomStream);
-		SPECIALIZE_SINGLE(FTimespan);
-		SPECIALIZE_SINGLE(FTwoVectors);
-		SPECIALIZE_SINGLE(uint8);
-		SPECIALIZE_SINGLE(uint16);
-		SPECIALIZE_SINGLE(uint32);
-		SPECIALIZE_SINGLE(uint64);
-		SPECIALIZE_SINGLE(FVector2D);
-		SPECIALIZE_SINGLE(FIntPoint);
-		SPECIALIZE_SINGLE(FIntVector);
-		SPECIALIZE_SINGLE(FVector2DHalf);
-		SPECIALIZE_SINGLE(FFloat16);
-		SPECIALIZE_SINGLE(FOrientedBox);
-		SPECIALIZE_SINGLE(FFrameNumber);
-		SPECIALIZE_SINGLE(FFrameRate);
-		SPECIALIZE_SINGLE(FFrameTime);
-		SPECIALIZE_SINGLE(FMatrix44f);
-		SPECIALIZE_SINGLE(FMatrix44d);
-		SPECIALIZE_SINGLE(FPlane4f);
-		SPECIALIZE_SINGLE(FPlane4d);
-		SPECIALIZE_SINGLE(FQuat4f);
-		SPECIALIZE_SINGLE(FQuat4d);
-		SPECIALIZE_SINGLE(FRotator3f);
-		SPECIALIZE_SINGLE(FRotator3d);
-		SPECIALIZE_SINGLE(FTransform3f);
-		SPECIALIZE_SINGLE(FTransform3d);
-		SPECIALIZE_SINGLE(FVector3f);
-		SPECIALIZE_SINGLE(FVector3d);
-		SPECIALIZE_SINGLE(FVector2f);
-		SPECIALIZE_SINGLE(FVector4f);
-		SPECIALIZE_SINGLE(FVector4d);
-		SPECIALIZE_SINGLE(FBox2f);
-		SPECIALIZE_SINGLE(FBox2D);
-		SPECIALIZE_SINGLE(FBox3f);
-		SPECIALIZE_SINGLE(FBox3d);
-		SPECIALIZE_SINGLE(FBoxSphereBounds3f);
-		SPECIALIZE_SINGLE(FBoxSphereBounds3d);
-		SPECIALIZE_SINGLE(FSphere3f);
-		SPECIALIZE_SINGLE(FSphere3d);
-		#undef SPECIALIZE_SINGLE
-
-		#define SPECIALIZE_NOGETTER(X) \
-			template INTERCHANGECORE_API EAttributeStorageResult FAttributeStorage::RegisterAttribute<X>(const FAttributeKey&, const X&, EAttributeProperty);\
-			template INTERCHANGECORE_API FAttributeStorage::TAttributeHandle<X>::TAttributeHandle();\
-			template INTERCHANGECORE_API FAttributeStorage::TAttributeHandle<X>::TAttributeHandle(const FAttributeKey& InKey, const FAttributeStorage* InAttributeStorage);\
-			template INTERCHANGECORE_API bool FAttributeStorage::TAttributeHandle<X>::IsValid() const;\
-			template INTERCHANGECORE_API EAttributeStorageResult FAttributeStorage::TAttributeHandle<X>::Get(X& Value) const;\
-			template INTERCHANGECORE_API EAttributeStorageResult FAttributeStorage::TAttributeHandle<X>::Set(const X& Value);\
-			template INTERCHANGECORE_API const FAttributeKey& FAttributeStorage::TAttributeHandle<X>::GetKey() const;\
-			template INTERCHANGECORE_API FAttributeStorage::TAttributeHandle<X> FAttributeStorage::GetAttributeHandle<X>(const FAttributeKey&) const;\
-			template INTERCHANGECORE_API FString UE::Interchange::AttributeValueToString(const X&);
-		SPECIALIZE_NOGETTER(FName);
-		SPECIALIZE_NOGETTER(FString);
-		SPECIALIZE_NOGETTER(FSoftObjectPath);
-		SPECIALIZE_NOGETTER(TArray<uint8>);
-		SPECIALIZE_NOGETTER(TArray64<uint8>);
-		#undef SPECIALIZE_NOGETTER
 
+		SPECIALIZE(bool);
+		SPECIALIZE(FColor);
+		SPECIALIZE(FDateTime);
+		SPECIALIZE(double);
+		SPECIALIZE(float);
+		SPECIALIZE(FGuid);
+		SPECIALIZE(int8);
+		SPECIALIZE(int16);
+		SPECIALIZE(int32);
+		SPECIALIZE(int64);
+		SPECIALIZE(FIntRect);
+		SPECIALIZE(FLinearColor);
+		SPECIALIZE(FName);
+		SPECIALIZE(FRandomStream);
+		SPECIALIZE(FTimespan);
+		SPECIALIZE(FString);
+		SPECIALIZE(FTwoVectors);
+		SPECIALIZE(uint8);
+		SPECIALIZE(uint16);
+		SPECIALIZE(uint32);
+		SPECIALIZE(uint64);
+		SPECIALIZE(FVector2D);
+		SPECIALIZE(FIntPoint);
+		SPECIALIZE(FIntVector);
+		SPECIALIZE(FVector2DHalf);
+		SPECIALIZE(FFloat16);
+		SPECIALIZE(FOrientedBox);
+		SPECIALIZE(FFrameNumber);
+		SPECIALIZE(FFrameRate);
+		SPECIALIZE(FFrameTime);
+		SPECIALIZE(FSoftObjectPath);
+		SPECIALIZE(FMatrix44f);
+		SPECIALIZE(FMatrix44d);
+		SPECIALIZE(FPlane4f);
+		SPECIALIZE(FPlane4d);
+		SPECIALIZE(FQuat4f);
+		SPECIALIZE(FQuat4d);
+		SPECIALIZE(FRotator3f);
+		SPECIALIZE(FRotator3d);
+		SPECIALIZE(FTransform3f);
+		SPECIALIZE(FTransform3d);
+		SPECIALIZE(FVector3f);
+		SPECIALIZE(FVector3d);
+		SPECIALIZE(FVector2f);
+		SPECIALIZE(FVector4f);
+		SPECIALIZE(FVector4d);
+		SPECIALIZE(FBox2f);
+		SPECIALIZE(FBox2D);
+		SPECIALIZE(FBox3f);
+		SPECIALIZE(FBox3d);
+		SPECIALIZE(FBoxSphereBounds3f);
+		SPECIALIZE(FBoxSphereBounds3d);
+		SPECIALIZE(FSphere3f);
+		SPECIALIZE(FSphere3d);
+
+		SPECIALIZE(TArray<bool>);
+		SPECIALIZE(TArray<FColor>);
+		SPECIALIZE(TArray<FDateTime>);
+		SPECIALIZE(TArray<double>);
+		SPECIALIZE(TArray<float>);
+		SPECIALIZE(TArray<FGuid>);
+		SPECIALIZE(TArray<int8>);
+		SPECIALIZE(TArray<int16>);
+		SPECIALIZE(TArray<int32>);
+		SPECIALIZE(TArray<int64>);
+		SPECIALIZE(TArray<FIntRect>);
+		SPECIALIZE(TArray<FLinearColor>);
+		SPECIALIZE(TArray<FName>);
+		SPECIALIZE(TArray<FRandomStream>);
+		SPECIALIZE(TArray<FTimespan>);
+		SPECIALIZE(TArray<FString>);
+		SPECIALIZE(TArray<FTwoVectors>);
+		SPECIALIZE(TArray<uint8>);
+		SPECIALIZE(TArray64<uint8>);
+		SPECIALIZE(TArray<uint16>);
+		SPECIALIZE(TArray<uint32>);
+		SPECIALIZE(TArray<uint64>);
+		SPECIALIZE(TArray<FVector2D>);
+		SPECIALIZE(TArray<FIntPoint>);
+		SPECIALIZE(TArray<FIntVector>);
+		SPECIALIZE(TArray<FVector2DHalf>);
+		SPECIALIZE(TArray<FFloat16>);
+		SPECIALIZE(TArray<FOrientedBox>);
+		SPECIALIZE(TArray<FFrameNumber>);
+		SPECIALIZE(TArray<FFrameRate>);
+		SPECIALIZE(TArray<FFrameTime>);
+		SPECIALIZE(TArray<FSoftObjectPath>);
+		SPECIALIZE(TArray<FMatrix44f>);
+		SPECIALIZE(TArray<FMatrix44d>);
+		SPECIALIZE(TArray<FPlane4f>);
+		SPECIALIZE(TArray<FPlane4d>);
+		SPECIALIZE(TArray<FQuat4f>);
+		SPECIALIZE(TArray<FQuat4d>);
+		SPECIALIZE(TArray<FRotator3f>);
+		SPECIALIZE(TArray<FRotator3d>);
+		SPECIALIZE(TArray<FTransform3f>);
+		SPECIALIZE(TArray<FTransform3d>);
+		SPECIALIZE(TArray<FVector3f>);
+		SPECIALIZE(TArray<FVector3d>);
+		SPECIALIZE(TArray<FVector2f>);
+		SPECIALIZE(TArray<FVector4f>);
+		SPECIALIZE(TArray<FVector4d>);
+		SPECIALIZE(TArray<FBox2f>);
+		SPECIALIZE(TArray<FBox2D>);
+		SPECIALIZE(TArray<FBox3f>);
+		SPECIALIZE(TArray<FBox3d>);
+		SPECIALIZE(TArray<FBoxSphereBounds3f>);
+		SPECIALIZE(TArray<FBoxSphereBounds3d>);
+		SPECIALIZE(TArray<FSphere3f>);
+		SPECIALIZE(TArray<FSphere3d>);
+
+		#undef SPECIALIZE
 	} //ns Interchange
 } //ns UE
\ No newline at end of file
diff --git a/Engine/Source/Runtime/Interchange/Core/Public/InterchangeCustomVersion.h b/Engine/Source/Runtime/Interchange/Core/Public/InterchangeCustomVersion.h
--- a/Engine/Source/Runtime/Interchange/Core/Public/InterchangeCustomVersion.h
+++ b/Engine/Source/Runtime/Interchange/Core/Public/InterchangeCustomVersion.h
@@ -14,6 +14,8 @@ struct FInterchangeCustomVersion
 
 		SerializedInterchangeObjectStoring,
 
+		MultipleAllocationsPerAttributeInStorage,
+
 		// -----<new versions can be added above this line>-------------------------------------------------
 		VersionPlusOne,
 		LatestVersion = VersionPlusOne - 1
diff --git a/Engine/Source/Runtime/Interchange/Core/Public/Types/AttributeStorage.h b/Engine/Source/Runtime/Interchange/Core/Public/Types/AttributeStorage.h
--- a/Engine/Source/Runtime/Interchange/Core/Public/Types/AttributeStorage.h
+++ b/Engine/Source/Runtime/Interchange/Core/Public/Types/AttributeStorage.h
@@ -4,27 +4,47 @@
 
 #include "Containers/Array.h"
 #include "Containers/ContainersFwd.h"
-#include "Containers/Map.h"
-
 #include "Containers/EnumAsByte.h"
+#include "Containers/Map.h"
 #include "Containers/UnrealString.h"
 #include "HAL/CriticalSection.h"
 #include "Internationalization/Text.h"
+#include "Math/Box.h"
+#include "Math/Box2D.h"
+#include "Math/BoxSphereBounds.h"
+#include "Math/Color.h"
+#include "Math/Float16.h"
+#include "Math/IntPoint.h"
+#include "Math/IntRect.h"
+#include "Math/IntVector.h"
+#include "Math/MathFwd.h"
+#include "Math/Matrix.h"
+#include "Math/OrientedBox.h"
+#include "Math/Plane.h"
+#include "Math/Quat.h"
+#include "Math/RandomStream.h"
+#include "Math/Rotator.h"
+#include "Math/Sphere.h"
+#include "Math/Sphere.h"
 #include "Math/Transform.h"
+#include "Math/TwoVectors.h"
+#include "Math/Vector.h"
+#include "Math/Vector2D.h"
+#include "Math/Vector2DHalf.h"
+#include "Math/Vector4.h"
 #include "Misc/DateTime.h"
 #include "Misc/EnumClassFlags.h"
 #include "Misc/FrameNumber.h"
+#include "Misc/FrameRate.h"
+#include "Misc/FrameTime.h"
 #include "Misc/Guid.h"
 #include "Misc/Timespan.h"
 #include "Serialization/Archive.h"
-#include "Templates/Function.h"
 #include "Templates/UnrealTemplate.h"
 #include "UObject/NameTypes.h"
 #include "UObject/SoftObjectPath.h"
 
-struct FFrameTime;
-struct FOrientedBox;
-template <class TEnum> class TEnumAsByte;
+#include "InterchangeCustomVersion.h"
 
 //Interchange namespace
 namespace UE
@@ -174,67 +194,123 @@ namespace UE
 		 */
 		enum class EAttributeTypes : int32
 		{
-			None				= 0,
-			Bool				= 1,
-			ByteArray			= 2,
-			ByteArray64			= 3,
-			Color				= 4,
-			DateTime			= 5,
-			Double				= 6,
-			Enum				= 7,
-			Float				= 8,
-			Guid				= 9,
-			Int8				= 10,
-			Int16				= 11,
-			Int32				= 12,
-			Int64				= 13,
-			IntRect				= 14,
-			LinearColor			= 15,
-			Name				= 16,
-			RandomStream		= 17,
-			String				= 18,
-			Timespan			= 19,
-			TwoVectors			= 20,
-			UInt8				= 21,
-			UInt16				= 22,
-			UInt32				= 23,
-			UInt64				= 24,
-			Vector2d			= 25,
-			IntPoint			= 26,
-			IntVector			= 27,
-			Vector2DHalf		= 28,
-			Float16				= 29,
-			OrientedBox			= 30,
-			FrameNumber			= 31,
-			FrameRate			= 32,
-			FrameTime			= 33,
-			SoftObjectPath		= 34,
-			Matrix44f			= 35,
-			Matrix44d			= 36,
-			Plane4f				= 37,
-			Plane4d				= 38,
-			Quat4f				= 39,
-			Quat4d				= 40,
-			Rotator3f			= 41,
-			Rotator3d			= 42,
-			Transform3f			= 43,
-			Transform3d			= 44,
-			Vector3f			= 45,
-			Vector3d			= 46,
-			Vector2f			= 47,
-			Vector4f			= 48,
-			Vector4d			= 49,
-			Box2f				= 50,
-			Box2D				= 51,
-			Box3f				= 52,
-			Box3d				= 53,
-			BoxSphereBounds3f	= 54,
-			BoxSphereBounds3d	= 55,
-			Sphere3f			= 56,
-			Sphere3d			= 57,
+			None					= 0,
+
+			Bool					= 1,
+			Color					= 4,
+			DateTime				= 5,
+			Double					= 6,
+			Enum					= 7,
+			Float					= 8,
+			Guid					= 9,
+			Int8					= 10,
+			Int16					= 11,
+			Int32					= 12,
+			Int64					= 13,
+			IntRect					= 14,
+			LinearColor				= 15,
+			Name					= 16,
+			RandomStream			= 17,
+			String					= 18,
+			Timespan				= 19,
+			TwoVectors				= 20,
+			UInt8					= 21,
+			UInt16					= 22,
+			UInt32					= 23,
+			UInt64					= 24,
+			Vector2d				= 25,
+			IntPoint				= 26,
+			IntVector				= 27,
+			Vector2DHalf			= 28,
+			Float16					= 29,
+			OrientedBox				= 30,
+			FrameNumber				= 31,
+			FrameRate				= 32,
+			FrameTime				= 33,
+			SoftObjectPath			= 34,
+			Matrix44f				= 35,
+			Matrix44d				= 36,
+			Plane4f					= 37,
+			Plane4d					= 38,
+			Quat4f					= 39,
+			Quat4d					= 40,
+			Rotator3f				= 41,
+			Rotator3d				= 42,
+			Transform3f				= 43,
+			Transform3d				= 44,
+			Vector3f				= 45,
+			Vector3d				= 46,
+			Vector2f				= 47,
+			Vector4f				= 48,
+			Vector4d				= 49,
+			Box2f					= 50,
+			Box2D					= 51,
+			Box3f					= 52,
+			Box3d					= 53,
+			BoxSphereBounds3f		= 54,
+			BoxSphereBounds3d		= 55,
+			Sphere3f				= 56,
+			Sphere3d				= 57,
+
+			BoolArray				= 58,
+			ColorArray				= 59,
+			DateTimeArray			= 60,
+			DoubleArray				= 61,
+			EnumArray				= 62,
+			FloatArray				= 63,
+			GuidArray				= 64,
+			Int8Array				= 65,
+			Int16Array				= 66,
+			Int32Array				= 67,
+			Int64Array				= 68,
+			IntRectArray			= 69,
+			LinearColorArray		= 70,
+			NameArray				= 71,
+			RandomStreamArray		= 72,
+			StringArray				= 73,
+			TimespanArray			= 74,
+			TwoVectorsArray			= 75,
+			ByteArray				= 2,
+			ByteArray64				= 3,
+			UInt16Array				= 76,
+			UInt32Array				= 77,
+			UInt64Array				= 78,
+			Vector2dArray			= 79,
+			IntPointArray			= 80,
+			IntVectorArray			= 81,
+			Vector2DHalfArray		= 82,
+			Float16Array			= 83,
+			OrientedBoxArray		= 84,
+			FrameNumberArray		= 85,
+			FrameRateArray			= 86,
+			FrameTimeArray			= 87,
+			SoftObjectPathArray		= 88,
+			Matrix44fArray			= 89,
+			Matrix44dArray			= 90,
+			Plane4fArray			= 91,
+			Plane4dArray			= 92,
+			Quat4fArray				= 93,
+			Quat4dArray				= 94,
+			Rotator3fArray			= 95,
+			Rotator3dArray			= 96,
+			Transform3fArray		= 97,
+			Transform3dArray		= 98,
+			Vector3fArray			= 99,
+			Vector3dArray			= 100,
+			Vector2fArray			= 101,
+			Vector4fArray			= 102,
+			Vector4dArray			= 103,
+			Box2fArray				= 104,
+			Box2DArray				= 105,
+			Box3fArray				= 106,
+			Box3dArray				= 107,
+			BoxSphereBounds3fArray 	= 108,
+			BoxSphereBounds3dArray 	= 109,
+			Sphere3fArray			= 110,
+			Sphere3dArray 			= 111,
 
 			//Max should always be updated if we add a new supported type
-			Max                 = 58
+			Max                 	= 112
 		};
 
 		/**
@@ -373,7 +449,37 @@ namespace UE
 		};
 
 		template<typename T>
-		inline constexpr bool TIsNonEnumType = !TIsEnum<T>::Value && !TIsTEnumAsByte<T>::Value;
+		struct TTypeTraits
+		{
+			static constexpr bool bIsArrayOfStringType = false;
+			static constexpr bool bIsArrayOfEnumType = false;
+			static constexpr bool bIsArrayOfEnumAsByteType = false;
+		};
+
+		template<typename T>
+		inline constexpr bool TIsStringType = std::is_same_v<T, FString> || std::is_same_v<T, FName> || std::is_same_v<T, FSoftObjectPath>;
+
+		template<typename T>
+		inline constexpr bool TIsArrayOfStringType = TTypeTraits<T>::bIsArrayOfStringType;
+
+		template<typename T>
+		inline constexpr bool TIsArrayOfEnumType = TTypeTraits<T>::bIsArrayOfEnumType;
+
+		template<typename T>
+		inline constexpr bool TIsArrayOfEnumAsByteType = TTypeTraits<T>::bIsArrayOfEnumAsByteType;
+
+		template<typename T>
+		inline constexpr bool TIsNonEnumType = !TIsEnum<T>::Value && !TIsTEnumAsByte<T>::Value && !TIsArrayOfEnumType<T> && !TIsArrayOfEnumAsByteType<T>;
+
+		template<typename T>
+		struct TTypeTraits<TArray<T>>
+		{
+			using ElementType = T;
+
+			static constexpr bool bIsArrayOfStringType = TIsStringType<ElementType>;
+			static constexpr bool bIsArrayOfEnumType = TIsEnum<ElementType>::Value;
+			static constexpr bool bIsArrayOfEnumAsByteType = TIsTEnumAsByte<ElementType>::Value;
+		};
 
 		/**
 		 * This class is a Key/Value storage inside a TArray64<uint8>.
@@ -425,8 +531,8 @@ namespace UE
 				: public TAttributeHandle<std::underlying_type_t<E>>
 			{
 			public:
-				using Base = TAttributeHandle<std::underlying_type_t<E>>;
 				using UnderlyingType = std::underlying_type_t<E>;
+				using Base = TAttributeHandle<UnderlyingType>;
 
 				using Base::Base;
 
@@ -451,14 +557,14 @@ namespace UE
 				friend FAttributeStorage;
 			};
 
-			// Redirect TEnumAsByte instantiations to underlying type
+			// Redirect TEnumAsByte<T> instantiations to uint8
 			template<typename B>
 			class TAttributeHandle<B, std::enable_if_t<TIsTEnumAsByte<B>::Value>>
 				: public TAttributeHandle<uint8>
 			{
 			public:
-				using Base = TAttributeHandle<uint8>;
 				using UnderlyingType = uint8;
+				using Base = TAttributeHandle<UnderlyingType>;
 
 				using Base::Base;
 
@@ -481,6 +587,66 @@ namespace UE
 				friend FAttributeStorage;
 			};
 
+			// Redirect TArray<enum> instantiations to TArray<underlying type>
+			template<typename T>
+			class TAttributeHandle<T, std::enable_if_t<TIsArrayOfEnumType<T>>>
+				: public TAttributeHandle<TArray<std::underlying_type_t<typename T::ElementType>>>
+			{
+			public:
+				using UnderlyingType = std::underlying_type_t<typename T::ElementType>;
+				using Base = TAttributeHandle<TArray<UnderlyingType>>;
+
+				using Base::Base;
+
+				explicit TAttributeHandle(const TAttributeHandle<TArray<UnderlyingType>>& Other)
+					: TAttributeHandle(Other.Key, Other.AttributeStorage)
+				{
+				}
+
+				EAttributeStorageResult Get(T& Value) const
+				{
+					return Base::Get((TArray<UnderlyingType>&)Value);
+				}
+
+				EAttributeStorageResult Set(const T& Value)
+				{
+					return Base::Set((const TArray<UnderlyingType>&)Value);
+				}
+
+			protected:
+				friend FAttributeStorage;
+			};
+
+			// Redirect TArray<TEnumAsByte<T>> instantiations to TArray<uint8>
+			template<typename T>
+			class TAttributeHandle<T, std::enable_if_t<TIsArrayOfEnumAsByteType<T>>>
+				: public TAttributeHandle<TArray<uint8>>
+			{
+			public:
+				using UnderlyingType = uint8;
+				using Base = TAttributeHandle<TArray<UnderlyingType>>;
+
+				using Base::Base;
+
+				explicit TAttributeHandle(const TAttributeHandle<TArray<UnderlyingType>>& Other)
+					: TAttributeHandle(Other.Key, Other.AttributeStorage)
+				{
+				}
+
+				EAttributeStorageResult Get(T& Value) const
+				{
+					return Base::Get((TArray<UnderlyingType>&)Value);
+				}
+
+				EAttributeStorageResult Set(const T& Value)
+				{
+					return Base::Set((const TArray<UnderlyingType>&)Value);
+				}
+
+			protected:
+				friend FAttributeStorage;
+			};
+
 			FAttributeStorage() = default;
 
 			INTERCHANGECORE_API FAttributeStorage(const FAttributeStorage& Other);
@@ -517,6 +683,22 @@ namespace UE
 				return RegisterAttribute<UnderlyingType>(ElementAttributeKey, (const UnderlyingType&)DefaultValue, AttributeProperty);
 			}
 
+			// Redirect TArray<enum> instantiations to TArray<underlying type>
+			template<typename T, typename = std::enable_if_t<TIsArrayOfEnumType<T>>, int = 0, int = 0, int = 0>
+			EAttributeStorageResult RegisterAttribute(const FAttributeKey& ElementAttributeKey, const T& DefaultValue, EAttributeProperty AttributeProperty = EAttributeProperty::None)
+			{
+				using UnderlyingType = std::underlying_type_t<typename T::ElementType>;
+				return RegisterAttribute<TArray<UnderlyingType>>(ElementAttributeKey, (const TArray<UnderlyingType>&)DefaultValue, AttributeProperty);
+			}
+
+			// Redirect TArray<TEnumAsByte> instantiations to TArray<underlying type>
+			template<typename T, typename = std::enable_if_t<TIsArrayOfEnumAsByteType<T>>, int = 0, int = 0, int = 0, int = 0>
+			EAttributeStorageResult RegisterAttribute(const FAttributeKey& ElementAttributeKey, const T& DefaultValue, EAttributeProperty AttributeProperty = EAttributeProperty::None)
+			{
+				using UnderlyingType = uint8;
+				return RegisterAttribute<TArray<uint8>>(ElementAttributeKey, (const TArray<uint8>&)DefaultValue, AttributeProperty);
+			}
+
 			/**
 			 * Remove an attribute from the storage.
 			 *
@@ -558,6 +740,22 @@ namespace UE
 				return GetAttribute<UnderlyingType>(ElementAttributeKey, OutValue);
 			}
 
+			// Redirect TArray<enum> instantiations to TArray<underlying type>
+			template<typename ArrayOfE, typename = std::enable_if_t<TIsArrayOfEnumType<ArrayOfE>>, int = 0, int = 0, int = 0>
+			EAttributeStorageResult GetAttribute(const FAttributeKey& ElementAttributeKey, ArrayOfE& OutValue) const
+			{
+				using UnderlyingType = std::underlying_type_t<typename ArrayOfE::ElementType>;
+				return GetAttribute<TArray<UnderlyingType>>(ElementAttributeKey, (TArray<UnderlyingType>&)OutValue);
+			}
+
+			// Redirect TArray<TEnumAsByte> instantiations to TArray<uint8>
+			template<typename ArrayOfB, typename = std::enable_if_t<TIsArrayOfEnumAsByteType<ArrayOfB>>, int = 0, int = 0, int = 0, int = 0>
+			EAttributeStorageResult GetAttribute(const FAttributeKey& ElementAttributeKey, ArrayOfB& OutValue) const
+			{
+				using UnderlyingType = uint8;
+				return GetAttribute<TArray<uint8>>(ElementAttributeKey, (TArray<uint8>&)OutValue);
+			}
+
 			/**
 			 * Return an attribute handle for the specified attribute. This handle is a compile type check and is use to get and set the attribute value type.
 			 * The function will assert if the key is missing or the type doesn't match the specified template type.
@@ -583,6 +781,22 @@ namespace UE
 				return (TAttributeHandle<B>)GetAttributeHandle<UnderlyingType>(ElementAttributeKey);
 			}
 
+			// Redirect TArray<enum> instantiations to TArray<underlying type>
+			template<typename ArrayOfE, typename = std::enable_if_t<TIsArrayOfEnumType<ArrayOfE>>, int = 0, int = 0, int = 0>
+			TAttributeHandle<ArrayOfE> GetAttributeHandle(const FAttributeKey& ElementAttributeKey) const
+			{
+				using UnderlyingType = std::underlying_type_t<typename ArrayOfE::ElementType>;
+				return (TAttributeHandle<ArrayOfE>)GetAttributeHandle<TArray<UnderlyingType>>(ElementAttributeKey);
+			}
+
+			// Redirect TArray<TEnumAsByte> instantiations to TArray<uint8>
+			template<typename ArrayOfB, typename = std::enable_if_t<TIsArrayOfEnumAsByteType<ArrayOfB>>, int = 0, int = 0, int = 0, int = 0>
+			TAttributeHandle<ArrayOfB> GetAttributeHandle(const FAttributeKey& ElementAttributeKey) const
+			{
+				using UnderlyingType = uint8;
+				return (TAttributeHandle<ArrayOfB>)GetAttributeHandle<TArray<UnderlyingType>>(ElementAttributeKey);
+			}
+
 			/**
 			 * Return the attribute type if the key exists, or None if the key is missing.
 			 *
@@ -684,6 +898,12 @@ namespace UE
 				Ar << Storage.DefragRatio;
 				Ar << Storage.AttributeAllocationTable;
 				Ar << Storage.AttributeStorage;
+
+				if (Ar.IsLoading())
+				{
+					Storage.UpdateAllocationCount();
+				}
+
 				return Ar;
 			}
 
@@ -703,10 +923,10 @@ namespace UE
 				// returning false in inlined code of outer class
 				FAttributeAllocationInfo() {}
 				
-        		//The offset in the storage
-        		uint64 Offset = 0;
-        		//The size of the data in the storage
-        		uint64 Size = 0;
+        		//The offsets of our allocations in the storage
+        		TArray<uint64> Offsets;
+        		//The sizes of our allocations in the storage, in bytes
+        		TArray<uint64> Sizes;
         		//The real type of the attribute
         		EAttributeTypes Type = EAttributeTypes::None;
         		//The attribute properties
@@ -717,8 +937,8 @@ namespace UE
         		bool operator==(const FAttributeAllocationInfo& Other) const
         		{
         			//Offset is a unique key
-        			return Offset == Other.Offset
-        				&& Size == Other.Size
+        			return Offsets == Other.Offsets
+        				&& Sizes == Other.Sizes
         				&& Type == Other.Type
         				&& Property == Other.Property
         				&& Hash == Other.Hash;
@@ -727,11 +947,26 @@ namespace UE
         		//Serialization
         		friend FArchive& operator<<(FArchive& Ar, FAttributeAllocationInfo& AttributeAllocationInfo)
         		{
-        			Ar << AttributeAllocationInfo.Offset;
-        			Ar << AttributeAllocationInfo.Size;
-        			Ar << AttributeAllocationInfo.Type;
-        			Ar << AttributeAllocationInfo.Property;
-        			Ar << AttributeAllocationInfo.Hash;
+					if (Ar.IsLoading()
+						&& Ar.CustomVer(FInterchangeCustomVersion::GUID) < FInterchangeCustomVersion::MultipleAllocationsPerAttributeInStorage)
+					{
+						uint64 OldOffset;
+						Ar << OldOffset;
+						AttributeAllocationInfo.Offsets = {OldOffset};
+
+						uint64 OldSize;
+						Ar << OldSize;
+						AttributeAllocationInfo.Sizes = {OldSize};
+					}
+					else
+					{
+						Ar << AttributeAllocationInfo.Offsets;
+						Ar << AttributeAllocationInfo.Sizes;
+					}
+
+					Ar << AttributeAllocationInfo.Type;
+					Ar << AttributeAllocationInfo.Property;
+					Ar << AttributeAllocationInfo.Hash;
         			return Ar;
         		}
 			};
@@ -762,12 +997,10 @@ namespace UE
 				return 1;
 			}
 
-			void ExtractFStringAttributeFromStorage(const uint8* StorageData, const FAttributeAllocationInfo* AttributeAllocationInfo, FString& OutValue) const;
-
-			const FStringView GetFStringViewAttributeFromStorage(const uint8* StorageData, const FAttributeAllocationInfo* AttributeAllocationInfo) const;
+			const FStringView GetFStringViewAttributeFromStorage(const uint8* StorageData, const FAttributeAllocationInfo* AttributeAllocationInfo, int32 ElementIndex = 0) const;
 
 			template<typename MultiSizeType>
-			EAttributeStorageResult MultiSizeSetAttribute(FAttributeAllocationInfo* AttributeAllocationInfo, const MultiSizeType& Value, TFunctionRef<uint8* ()> GetValuePointer);
+			EAttributeStorageResult MultiSizeSetAttribute(FAttributeAllocationInfo* AttributeAllocationInfo, int32 TargetAllocationIndex, const MultiSizeType& Value, const uint8* SourceDataPtr, bool& bOutNeedsDefrag);
 
 			template<typename ArrayType>
 			EAttributeStorageResult GenericArrayGetAttribute(const FAttributeKey& ElementAttributeKey, ArrayType& OutValue) const;
@@ -775,6 +1008,10 @@ namespace UE
 			template<typename T>
 			friend void CopyStorageAttributesInternal(const FAttributeStorage& SourceStorage, FAttributeStorage& DestinationStorage, const TArray<T>& AttributeKeys);
 
+			void UpdateAllocationInfoHash(FAttributeAllocationInfo& AllocationInfo);
+
+			void UpdateAllocationCount();
+
 			/** The attribute allocation table is use to index the attributes into the storage. */
 			TMap<FAttributeKey, FAttributeAllocationInfo> AttributeAllocationTable;
 
@@ -793,10 +1030,67 @@ namespace UE
 			 */
 			float DefragRatio = 0.1f;
 
+			/**
+			 * Used to help with the debug visualizers (on Unreal.natvis), because now that we support multiple
+			 * allocations per attribute it is no longer trivial to compute the total number of allocations
+			 */
+			uint32 AllocationCount = 0;
+
 			/**
 			 * Mutex use when accessing or modifying the storage
 			 */
 			mutable FCriticalSection StorageMutex;
 		};
+
+		/**
+		 * The purpose of these is exclusively to guarantee these types are instantiated and available whenever we're debugging
+		 * an FAttributeStorage, which ensures the visualizers on Unreal.natvis can work
+		 */
+		namespace Debug
+		{
+			const static FColor UnusedColor{};
+			const static FDateTime UnusedDateTime{};
+			const static FGuid UnusedGuid{};
+			const static FIntRect UnusedIntRect{};
+			const static FLinearColor UnusedLinearColor{};
+			const static FName UnusedName{};
+			const static FRandomStream UnusedRandomStream{};
+			const static FTimespan UnusedTimespan{};
+			const static FString UnusedString{};
+			const static FTwoVectors UnusedTwoVectors{};
+			const static FVector2D UnusedVector2D{};
+			const static FIntPoint UnusedIntPoint{};
+			const static FIntVector UnusedIntVector{};
+			const static FVector2DHalf UnusedVector2DHalf{};
+			const static FFloat16 UnusedFloat16{};
+			const static FOrientedBox UnusedOrientedBox{};
+			const static FFrameNumber UnusedFrameNumber{};
+			const static FFrameRate UnusedFrameRate{};
+			const static FFrameTime UnusedFrameTime{};
+			const static FSoftObjectPath UnusedSoftObjectPath{};
+			const static FMatrix44f UnusedMatrix44f{};
+			const static FMatrix44d UnusedMatrix44d{};
+			const static FPlane4f UnusedPlane4f{};
+			const static FPlane4d UnusedPlane4d{};
+			const static FQuat4f UnusedQuat4f{};
+			const static FQuat4d UnusedQuat4d{};
+			const static FRotator3f UnusedRotator3f{};
+			const static FRotator3d UnusedRotator3d{};
+			const static FTransform3f UnusedTransform3f{};
+			const static FTransform3d UnusedTransform3d{};
+			const static FVector3f UnusedVector3f{};
+			const static FVector3d UnusedVector3d{};
+			const static FVector2f UnusedVector2f{};
+			const static FVector4f UnusedVector4f{};
+			const static FVector4d UnusedVector4d{};
+			const static FBox2f UnusedBox2f{};
+			const static FBox2D UnusedBox2D{};
+			const static FBox3f UnusedBox3f{};
+			const static FBox3d UnusedBox3d{};
+			const static FBoxSphereBounds3f UnusedBoxSphereBounds3f{};
+			const static FBoxSphereBounds3d UnusedBoxSphereBounds3d{};
+			const static FSphere3f UnusedSphere3f{};
+			const static FSphere3d UnusedSphere3d{};
+		}
 	} //ns interchange
 }
diff --git a/Engine/Source/Runtime/Interchange/Engine/Private/InterchangeAssetImportData.cpp b/Engine/Source/Runtime/Interchange/Engine/Private/InterchangeAssetImportData.cpp
--- a/Engine/Source/Runtime/Interchange/Engine/Private/InterchangeAssetImportData.cpp
+++ b/Engine/Source/Runtime/Interchange/Engine/Private/InterchangeAssetImportData.cpp
@@ -161,6 +161,7 @@ void UInterchangeAssetImportData::Serialize(FArchive& Ar)
 	{
 		Ar << CachedNodeContainer;
 		Ar << CachedPipelines;
+		CachedCustomVersions = Ar.GetCustomVersions();
 	}
 }
 
@@ -366,6 +367,8 @@ void UInterchangeAssetImportData::ProcessContainerCache() const
 		if (!TransientNodeContainer && CachedNodeContainer.Num() > 0)
 		{
 			FLargeMemoryReader NodeContainerAr(CachedNodeContainer.GetData(), CachedNodeContainer.Num());
+			NodeContainerAr.SetCustomVersions(CachedCustomVersions);
+
 			TransientNodeContainer = NewObject<UInterchangeBaseNodeContainer>();
 			TransientNodeContainer->SerializeNodeContainerData(NodeContainerAr);
 		}
diff --git a/Engine/Source/Runtime/Interchange/Engine/Public/InterchangeAssetImportData.h b/Engine/Source/Runtime/Interchange/Engine/Public/InterchangeAssetImportData.h
--- a/Engine/Source/Runtime/Interchange/Engine/Public/InterchangeAssetImportData.h
+++ b/Engine/Source/Runtime/Interchange/Engine/Public/InterchangeAssetImportData.h
@@ -195,6 +195,7 @@ private:
 	void ProcessTranslatorCache() const;
 	mutable TArray64<uint8> CachedNodeContainer;
 	mutable TArray<TPair<FString, FString>> CachedPipelines; //Class, Data(serialized JSON) pair
+	mutable FCustomVersionContainer CachedCustomVersions; // Custom versions of the archive we deserialized from, to help us deserialize our cached data
 	mutable TPair<FString, FString> CachedTranslatorSettings;
 
 #if WITH_EDITORONLY_DATA
