Interchange: Remove unnecessary Debug namespace variables from cl 44413699

#rnx

diff --git a/Engine/Source/Runtime/Interchange/Core/Public/Types/AttributeStorage.h b/Engine/Source/Runtime/Interchange/Core/Public/Types/AttributeStorage.h
--- a/Engine/Source/Runtime/Interchange/Core/Public/Types/AttributeStorage.h
+++ b/Engine/Source/Runtime/Interchange/Core/Public/Types/AttributeStorage.h
@@ -1041,56 +1041,5 @@ namespace UE
 			 */
 			mutable FCriticalSection StorageMutex;
 		};
-
-		/**
-		 * The purpose of these is exclusively to guarantee these types are instantiated and available whenever we're debugging
-		 * an FAttributeStorage, which ensures the visualizers on Unreal.natvis can work
-		 */
-		namespace Debug
-		{
-			const static FColor UnusedColor{};
-			const static FDateTime UnusedDateTime{};
-			const static FGuid UnusedGuid{};
-			const static FIntRect UnusedIntRect{};
-			const static FLinearColor UnusedLinearColor{};
-			const static FName UnusedName{};
-			const static FRandomStream UnusedRandomStream{};
-			const static FTimespan UnusedTimespan{};
-			const static FString UnusedString{};
-			const static FTwoVectors UnusedTwoVectors{};
-			const static FVector2D UnusedVector2D{};
-			const static FIntPoint UnusedIntPoint{};
-			const static FIntVector UnusedIntVector{};
-			const static FVector2DHalf UnusedVector2DHalf{};
-			const static FFloat16 UnusedFloat16{};
-			const static FOrientedBox UnusedOrientedBox{};
-			const static FFrameNumber UnusedFrameNumber{};
-			const static FFrameRate UnusedFrameRate{};
-			const static FFrameTime UnusedFrameTime{};
-			const static FSoftObjectPath UnusedSoftObjectPath{};
-			const static FMatrix44f UnusedMatrix44f{};
-			const static FMatrix44d UnusedMatrix44d{};
-			const static FPlane4f UnusedPlane4f{};
-			const static FPlane4d UnusedPlane4d{};
-			const static FQuat4f UnusedQuat4f{};
-			const static FQuat4d UnusedQuat4d{};
-			const static FRotator3f UnusedRotator3f{};
-			const static FRotator3d UnusedRotator3d{};
-			const static FTransform3f UnusedTransform3f{};
-			const static FTransform3d UnusedTransform3d{};
-			const static FVector3f UnusedVector3f{};
-			const static FVector3d UnusedVector3d{};
-			const static FVector2f UnusedVector2f{};
-			const static FVector4f UnusedVector4f{};
-			const static FVector4d UnusedVector4d{};
-			const static FBox2f UnusedBox2f{};
-			const static FBox2D UnusedBox2D{};
-			const static FBox3f UnusedBox3f{};
-			const static FBox3d UnusedBox3d{};
-			const static FBoxSphereBounds3f UnusedBoxSphereBounds3f{};
-			const static FBoxSphereBounds3d UnusedBoxSphereBounds3d{};
-			const static FSphere3f UnusedSphere3f{};
-			const static FSphere3d UnusedSphere3d{};
-		}
 	} //ns interchange
 }
