Unshelved from Daniel Coelho

Interchange: Fix ensure when deserializing assets saved before cl 44413699.

Horde issue 6888f370c61dade5bc837d91

diff --git a/Engine/Source/Runtime/Interchange/Core/Public/Types/AttributeStorage.h b/Engine/Source/Runtime/Interchange/Core/Public/Types/AttributeStorage.h
--- a/Engine/Source/Runtime/Interchange/Core/Public/Types/AttributeStorage.h
+++ b/Engine/Source/Runtime/Interchange/Core/Public/Types/AttributeStorage.h
@@ -947,6 +947,8 @@ namespace UE
         		//Serialization
         		friend FArchive& operator<<(FArchive& Ar, FAttributeAllocationInfo& AttributeAllocationInfo)
         		{
+					Ar.UsingCustomVersion(FInterchangeCustomVersion::GUID);
+
 					if (Ar.IsLoading()
 						&& Ar.CustomVer(FInterchangeCustomVersion::GUID) < FInterchangeCustomVersion::MultipleAllocationsPerAttributeInStorage)
 					{
