Interchange: Fix static analysis issues from cl 44623560.

#rnx

diff --git a/Engine/Source/Runtime/Interchange/Core/Private/Types/AttributeStorage.cpp b/Engine/Source/Runtime/Interchange/Core/Private/Types/AttributeStorage.cpp
--- a/Engine/Source/Runtime/Interchange/Core/Private/Types/AttributeStorage.cpp
+++ b/Engine/Source/Runtime/Interchange/Core/Private/Types/AttributeStorage.cpp
@@ -1879,8 +1879,8 @@ namespace UE
 			});
 
 			uint8* StorageData = AttributeStorage.GetData();
-			//Number of attribute we have done
-			uint64 AttributeCount = 0;
+			//Number of allocations we have handled
+			uint64 HandledAllocations = 0;
 			//Current storage offset we want to defrag
 			uint64 CurrentOffset = 0;
 			//Are we build a defrag block
@@ -1892,8 +1892,8 @@ namespace UE
 			const uint64 NumAllocations = LinearAllocations.Num();
 			for (FLinearAllocation& LinearAllocation : LinearAllocations)
 			{
-				AttributeCount++;
-				const bool bLast = AttributeCount == NumAllocations;
+				HandledAllocations++;
+				const bool bLast = HandledAllocations == NumAllocations;
 
 				if (LinearAllocation.Offset == CurrentOffset)
 				{
@@ -2082,8 +2082,8 @@ namespace UE
 					&StorageData[AllocationInfo.Offsets[AllocIndex]], 
 					AllocationInfo.Sizes[AllocIndex]
 				);
-				Sha.Final();
 			}
+			Sha.Final();
 
 			// Retrieve the hash and use it to construct a pseudo-GUID.
 			uint32 Hash[5];
