[Interchange] Reducing Static Mesh import memory footprint by clearing MeshDescription after its effective life span.

MeshDescription is cleared at the end of `UInterchangeStaticMeshFactory::EndImportAsset_GameThread`

Surprisingly, the effect on memory is mostly visible on the overall memory used (compared to memory used by Interchange only). It's also visible during the Distance Fields generation that happens after the import is done.

Before vs After: (testing using one big mesh of 9M triangles (11M vertices))
Total Memory: 17.35Gb to 13.93Gb
-- Spike: 23.91Gb to 21.25Gb
-- Distance field computation end: 19.40Gb to 15.33Gb

Interchange Memory: 4.46Gb to 3.64Gb
-- Spike: 7.97Gb to 7.06Gb
-- Distance field computation end: 5.55Gb to 4.60Gb

Combining this CL with  the previous optimization (44753863), testing the same FBX model above, from import to asset ready:
Total: 22.13Gb to 15.22Gb (-31%)
Spike: 26.83Gb to 21.26Gb (-20%)
Interchange: 8.39Gb to 4.6Gb (-45%)

#rnx

diff --git a/Engine/Plugins/Interchange/Runtime/Source/Import/Private/Mesh/InterchangeStaticMeshFactory.cpp b/Engine/Plugins/Interchange/Runtime/Source/Import/Private/Mesh/InterchangeStaticMeshFactory.cpp
--- a/Engine/Plugins/Interchange/Runtime/Source/Import/Private/Mesh/InterchangeStaticMeshFactory.cpp
+++ b/Engine/Plugins/Interchange/Runtime/Source/Import/Private/Mesh/InterchangeStaticMeshFactory.cpp
@@ -943,6 +943,9 @@ UInterchangeFactoryBase::FImportAssetResult UInterchangeStaticMeshFactory::EndIm
 	{
 		UStaticMesh::RemoveUnusedMaterialSlots(StaticMesh);
 	}
+
+	// Clear the mesh descriptions to free memory, we don't need them anymore since all the meshes have been committed at this point.
+	StaticMesh->ClearMeshDescriptions();
 #endif
 
 	ImportAssetResult.ImportedObject = StaticMesh;
diff --git a/Engine/Source/Runtime/Interchange/Engine/Private/InterchangeManager.cpp b/Engine/Source/Runtime/Interchange/Engine/Private/InterchangeManager.cpp
--- a/Engine/Source/Runtime/Interchange/Engine/Private/InterchangeManager.cpp
+++ b/Engine/Source/Runtime/Interchange/Engine/Private/InterchangeManager.cpp
@@ -2453,7 +2453,6 @@ UInterchangeManager::ImportInternal(const FString& ContentPath, const UInterchan
 		AnalyticsHelper->AddThreadSafe(UsageImportAnalayticsId, FAnalyticsEventAttribute(TEXT("Canceled"), bImportCanceled));
 
 		AsyncHelper->InitCancel();
-		AsyncHelper->CleanUp();
 
 		//if it is a re-import, re-instate the backedup source data:
 		if (UObject* ReimportCandidate = AsyncHelper->TaskData.ReimportObject)
@@ -2485,7 +2484,7 @@ UInterchangeManager::ImportInternal(const FString& ContentPath, const UInterchan
 	QueuedTaskData.TranslatorClass = AsyncTranslator->GetClass();
 
 	//If we cancel or abort the task we want to avoid putting it in the NonParallelTranslatorQueueTasks (the locks will not be release if the task doesn't start)
-	bool bTranslatorIsThreadSafe = AsyncTranslator->IsThreadSafe() || (bImportCanceled || bImportAborted);
+	bool bTranslatorIsThreadSafe = (bImportCanceled || bImportAborted) || AsyncTranslator->IsThreadSafe();
 	if (bTranslatorIsThreadSafe)
 	{
 		QueuedTasks.Enqueue(QueuedTaskData);
