Interchange Logging - Changed JSON to Message Log Text

#rnx

diff --git a/Engine/Source/Runtime/Interchange/Engine/Private/InterchangeManager.cpp b/Engine/Source/Runtime/Interchange/Engine/Private/InterchangeManager.cpp
--- a/Engine/Source/Runtime/Interchange/Engine/Private/InterchangeManager.cpp
+++ b/Engine/Source/Runtime/Interchange/Engine/Private/InterchangeManager.cpp
@@ -2777,10 +2777,10 @@ void UInterchangeManager::ReleaseAsyncHelper(TWeakPtr<UE::Interchange::FImportAs
 						switch (Result->GetResultType())
 						{
 						case EInterchangeResultType::Warning:
-							UE_LOG(LogInterchangeEngine, Warning, TEXT("%s"), *Result->ToJson());
+							UE_LOG(LogInterchangeEngine, Warning, TEXT("%s"), *Result->GetMessageLogText().ToString());
 							break;
 						case EInterchangeResultType::Error:
-							UE_LOG(LogInterchangeEngine, Error, TEXT("%s"), *Result->ToJson());
+							UE_LOG(LogInterchangeEngine, Error, TEXT("%s"), *Result->GetMessageLogText().ToString());
 							break;
 						}
 					}
