[Interchange] UI : Clean up debug functions.

#rnx

diff --git a/Engine/Source/Runtime/Interchange/Engine/Private/InterchangeSceneImportAsset.cpp b/Engine/Source/Runtime/Interchange/Engine/Private/InterchangeSceneImportAsset.cpp
--- a/Engine/Source/Runtime/Interchange/Engine/Private/InterchangeSceneImportAsset.cpp
+++ b/Engine/Source/Runtime/Interchange/Engine/Private/InterchangeSceneImportAsset.cpp
@@ -326,12 +326,4 @@ bool UInterchangeSceneImportAsset::ContainsObject(const FSoftObjectPath& SoftObj
 #else
 	return false;
 #endif
-}
-
-
-#if WITH_EDITOR
-const TMap<FSoftObjectPath, FString> UInterchangeSceneImportAsset::GetCachedSceneObjects() const
-{
-	return SceneObjects;
-}
-#endif 
\ No newline at end of file
+}
\ No newline at end of file
