Explicitly include UObject/GCObject.h pending its removal from other public UObject headers

diff --git a/Engine/Source/Runtime/Interchange/Engine/Public/InterchangeManager.h b/Engine/Source/Runtime/Interchange/Engine/Public/InterchangeManager.h
--- a/Engine/Source/Runtime/Interchange/Engine/Public/InterchangeManager.h
+++ b/Engine/Source/Runtime/Interchange/Engine/Public/InterchangeManager.h
@@ -23,6 +23,7 @@
 #include "InterchangeWriterBase.h"
 #include "Nodes/InterchangeBaseNodeContainer.h"
 #include "Templates/Tuple.h"
+#include "UObject/GCObject.h"
 #include "UObject/Package.h"
 #include "UObject/Object.h"
 #include "UObject/ObjectMacros.h"
