Fix compile errors

diff --git a/Engine/Source/Runtime/Interchange/Core/Private/Tests/StorageTest.cpp b/Engine/Source/Runtime/Interchange/Core/Private/Tests/StorageTest.cpp
--- a/Engine/Source/Runtime/Interchange/Core/Private/Tests/StorageTest.cpp
+++ b/Engine/Source/Runtime/Interchange/Core/Private/Tests/StorageTest.cpp
@@ -17,7 +17,7 @@ IMPLEMENT_SIMPLE_AUTOMATION_TEST(FAttributeStorageTest, "System.Runtime.Intercha
 #define ADD_ATTRIBUTE_STORAGE(ATTRIBUTE_TYPE, ATTRIBUTE_VALUE)																							\
 {																																						\
 	const ATTRIBUTE_TYPE RefValue = ATTRIBUTE_VALUE;																									\
-	FAttributeKey Key = CreateUniqueKey();																												\
+	UE::Interchange::FAttributeKey Key = CreateUniqueKey();																												\
 	if (!IsAttributeStorageResultSuccess(TestStorage.RegisterAttribute(Key, RefValue)))																	\
 	{																																					\
 		AddError(FString(TEXT("`AttributeStorage` must handle adding " #ATTRIBUTE_TYPE " attribute")));													\
@@ -38,7 +38,7 @@ IMPLEMENT_SIMPLE_AUTOMATION_TEST(FAttributeStorageTest, "System.Runtime.Intercha
 
 #define ADD_ATTRIBUTE_STORAGE_NOCOMPARE(ATTRIBUTE_TYPE, ATTRIBUTE_VALUE)																				\
 	const ATTRIBUTE_TYPE RefValue = ATTRIBUTE_VALUE;																									\
-	FAttributeKey Key = CreateUniqueKey();																												\
+	UE::Interchange::FAttributeKey Key = CreateUniqueKey();																												\
 	if (!IsAttributeStorageResultSuccess(TestStorage.RegisterAttribute(Key, RefValue)))																	\
 	{																																					\
 		AddError(FString(TEXT("`AttributeStorage` must handle adding " #ATTRIBUTE_TYPE " attribute")));													\
@@ -63,22 +63,22 @@ bool FAttributeStorageTest::RunTest(const FString& Parameters)
 
 	FString KeyPrefix = TEXT("TestKey");
 	int32 KeySuffixCounter = 1;
-	auto CreateUniqueKey = [&KeyPrefix, &KeySuffixCounter]()->FAttributeKey
+	auto CreateUniqueKey = [&KeyPrefix, &KeySuffixCounter]()->UE::Interchange::FAttributeKey
 	{
 		check(IsInGameThread());
 		FString KeyFinal = KeyPrefix + TEXT("_") + FString::FromInt(KeySuffixCounter++);
-		const FAttributeKey UniqueKeyRet = FAttributeKey(KeyFinal);
+		const UE::Interchange::FAttributeKey UniqueKeyRet = UE::Interchange::FAttributeKey(KeyFinal);
 		return UniqueKeyRet;
 	};
 
 	FAttributeStorage TestStorage = FAttributeStorage();
 	//This seed make the tests deterministic
 	FRandomStream RandomStream(564389);
-	const FAttributeKey TestInt32KeyName = CreateUniqueKey();
-	const FAttributeKey BigArrayKey = CreateUniqueKey();
+	const UE::Interchange::FAttributeKey TestInt32KeyName = CreateUniqueKey();
+	const UE::Interchange::FAttributeKey BigArrayKey = CreateUniqueKey();
 	const int32 NegativeValueRef = -1;
 	const int32 OverrideValueRef = 3327;
-	const FAttributeKey RandomStreamKey = CreateUniqueKey();
+	const UE::Interchange::FAttributeKey RandomStreamKey = CreateUniqueKey();
 	int32 TestStoredSeed = 0;
 	const int32 RandHelperUInt8 = 255;
 
@@ -197,7 +197,7 @@ bool FAttributeStorageTest::RunTest(const FString& Parameters)
 		FSphere3f RefValue;
 		RefValue.Center = FVector3f(2.0f);
 		RefValue.W = 2.2f;
-		FAttributeKey Key = CreateUniqueKey();
+		UE::Interchange::FAttributeKey Key = CreateUniqueKey();
 		if (!IsAttributeStorageResultSuccess(TestStorage.RegisterAttribute(Key, RefValue)))
 		{
 			AddError(FString(TEXT("`AttributeStorage` must handle adding FSphere3f attribute")));
@@ -213,7 +213,7 @@ bool FAttributeStorageTest::RunTest(const FString& Parameters)
 		FSphere3d RefValue;
 		RefValue.Center = FVector3d(2.0);
 		RefValue.W = 2.2;
-		FAttributeKey Key = CreateUniqueKey();
+		UE::Interchange::FAttributeKey Key = CreateUniqueKey();
 		if (!IsAttributeStorageResultSuccess(TestStorage.RegisterAttribute(Key, RefValue)))
 		{
 			AddError(FString(TEXT("`AttributeStorage` must handle adding FSphere3d attribute")));
@@ -292,7 +292,7 @@ bool FAttributeStorageTest::RunTest(const FString& Parameters)
 		{
 			ValueRef.Add(static_cast<uint8>(RandomStream.RandHelper(RandHelperUInt8)));
 		}
-		FAttributeKey Key = CreateUniqueKey();
+		UE::Interchange::FAttributeKey Key = CreateUniqueKey();
 		if (!IsAttributeStorageResultSuccess(TestStorage.RegisterAttribute(Key, ValueRef)))
 		{
 			AddError(FString(TEXT("`AttributeStorage` must handle adding TArray<uint8> attribute")));
@@ -672,11 +672,11 @@ bool FAttributeStorageTest::RunTest(const FString& Parameters)
 		//Add multithread (add BatchCount*BatchSize attributes)
 		const int32 BatchSize = 10;
 		const int32 BatchCount = 50;
-		TArray<FAttributeKey> Keys;
+		TArray<UE::Interchange::FAttributeKey> Keys;
 		Keys.AddZeroed(BatchSize*BatchCount);
 		bool ThreadError = false;
 		FAttributeStorage::TAttributeHandle<FRandomStream> RandomStreamHandle = TestStorage.GetAttributeHandle<FRandomStream>(RandomStreamKey);
-		for (FAttributeKey& Key : Keys)
+		for (UE::Interchange::FAttributeKey& Key : Keys)
 		{
 			Key = CreateUniqueKey();
 		}
@@ -743,7 +743,7 @@ bool FAttributeStorageTest::RunTest(const FString& Parameters)
 		for (int32 AddedIndex = 0; AddedIndex < 3; ++AddedIndex)
 		{
 			const FVector ValueRef = RandomStream.VRand();
-			FAttributeKey Key = CreateUniqueKey();
+			UE::Interchange::FAttributeKey Key = CreateUniqueKey();
 			if (!IsAttributeStorageResultSuccess(BaseStorage.RegisterAttribute(Key, ValueRef)))
 			{
 				AddError(FString(TEXT("`AttributeStorage` must handle adding FVector attribute")));
@@ -760,7 +760,7 @@ bool FAttributeStorageTest::RunTest(const FString& Parameters)
 		for (int32 AddedIndex = 0; AddedIndex < 3; ++AddedIndex)
 		{
 			const int32 ValueRef = RandomStream.RandHelper(RandHelperUInt8);
-			FAttributeKey Key = CreateUniqueKey();
+			UE::Interchange::FAttributeKey Key = CreateUniqueKey();
 			if (!IsAttributeStorageResultSuccess(BaseStorage.RegisterAttribute(Key, ValueRef)))
 			{
 				AddError(FString(TEXT("`AttributeStorage` must handle adding int32 attribute")));
@@ -772,22 +772,22 @@ bool FAttributeStorageTest::RunTest(const FString& Parameters)
 				break;
 			}
 		}
-		TArray<FAttributeKey> RemovedKeys;
-		TArray<FAttributeKey> AddedKeys;
-		TArray<FAttributeKey> ModifiedKeys;
+		TArray<UE::Interchange::FAttributeKey> RemovedKeys;
+		TArray<UE::Interchange::FAttributeKey> AddedKeys;
+		TArray<UE::Interchange::FAttributeKey> ModifiedKeys;
 		FAttributeStorage::CompareStorage(BaseStorage, VersionStorage, RemovedKeys, AddedKeys, ModifiedKeys);
 		if (RemovedKeys.Num() > 0 || AddedKeys.Num() > 0 || ModifiedKeys.Num() > 0)
 		{
 			AddError(FString(TEXT("`AttributeStorage` Compare storage shoukld have found no difference.")));
 		}
 
-		TArray<FAttributeKey> VersionKeys;
+		TArray<UE::Interchange::FAttributeKey> VersionKeys;
 		VersionStorage.GetAttributeKeys(VersionKeys);
 
 		//Add a integer key to the version storage
 		{
 			const int32 ValueRef = RandomStream.RandHelper(RandHelperUInt8);
-			FAttributeKey Key(TEXT("VersionAddedInteger"));
+			UE::Interchange::FAttributeKey Key(TEXT("VersionAddedInteger"));
 			if (!IsAttributeStorageResultSuccess(VersionStorage.RegisterAttribute(Key, ValueRef)))
 			{
 				AddError(FString(TEXT("`AttributeStorage` must handle adding int32 attribute")));
@@ -831,9 +831,9 @@ bool FAttributeStorageTest::RunTest(const FString& Parameters)
 		FMemoryReader ArRead(MemoryMocked, bMemoryPersistent);
 		ArRead << FromMemoryMocked;
 
-		TArray<FAttributeKey> RemovedKeys;
-		TArray<FAttributeKey> AddedKeys;
-		TArray<FAttributeKey> ModifiedKeys;
+		TArray<UE::Interchange::FAttributeKey> RemovedKeys;
+		TArray<UE::Interchange::FAttributeKey> AddedKeys;
+		TArray<UE::Interchange::FAttributeKey> ModifiedKeys;
 
 		if(TestStorage.GetStorageHash() != FromMemoryMocked.GetStorageHash())
 		{
@@ -850,7 +850,7 @@ bool FAttributeStorageTest::RunTest(const FString& Parameters)
 	//Performance Test, add two million entry, 1 million int32 and 1 million FVector
 	{
 		FGuid StorageHash = TestStorage.GetStorageHash();
-		TArray<FAttributeKey> AttributeKeys;
+		TArray<UE::Interchange::FAttributeKey> AttributeKeys;
 		const int32 TestCount = 10000;
 		AttributeKeys.AddZeroed(2*TestCount);
 		//Reserve 2x Testcount items (Testcount int32 and Testcount FVector) and reserve the data size
@@ -860,7 +860,7 @@ bool FAttributeStorageTest::RunTest(const FString& Parameters)
 		for (int32 AddedIndex = 0; AddedIndex < TestCount; ++AddedIndex)
 		{
 			const int32 ValueRef = AddedIndex;
-			FAttributeKey Key = CreateUniqueKey();
+			UE::Interchange::FAttributeKey Key = CreateUniqueKey();
 			AttributeKeys[AddedIndex*2] = Key;
 			if (!IsAttributeStorageResultSuccess(TestStorage.RegisterAttribute(Key, ValueRef)))
 			{
@@ -880,7 +880,7 @@ bool FAttributeStorageTest::RunTest(const FString& Parameters)
 		int32 MaxKeys = AttributeKeys.Num();
 		for (int32 KeyIndex = 0; KeyIndex < MaxKeys; ++KeyIndex)
 		{
-			const FAttributeKey& Key = AttributeKeys[KeyIndex];
+			const UE::Interchange::FAttributeKey& Key = AttributeKeys[KeyIndex];
 			if(TestStorage.GetAttributeType(Key) == EAttributeTypes::Int32)
 			{
 				int32 StoredValue(0);
@@ -895,7 +895,7 @@ bool FAttributeStorageTest::RunTest(const FString& Parameters)
 		const double SmokeTestReadTime = FPlatformTime::Seconds();
 		for (int32 KeyIndex = 0; KeyIndex < MaxKeys; ++KeyIndex)
 		{
-			const FAttributeKey& Key = AttributeKeys[KeyIndex];
+			const UE::Interchange::FAttributeKey& Key = AttributeKeys[KeyIndex];
 			TestStorage.UnregisterAttribute(Key);
 		}
 		const double SmokeTestStopTime = FPlatformTime::Seconds();
