[Interchange] - Fixed various issues while importing an asset containing elements named "None".

The solution is to use "Null" instead of "None" like FBX is doing.

The issue with the USD prim was the tip of the iceberg since a lot of path can lead to have an asset named "None":

Filename is named None
Asset Name option is set to "None"
A bone is named None
A material is named None
etc.
FBX was supporting some of the above but not all, but neither USD nor GLB was supporting it.

The idea behind this fix is to be as low level as possible by preventing a Node to have the None display name. There was an attempt to fix this in the past inside the AddNode, but this was still not fixing all the possible cases.

Tests (currently in progress) for this specifically are going to be added in another CL.

[RN] Fixed various issues while importing an asset containing elements named "None"

diff --git a/Engine/Plugins/Interchange/Runtime/Source/Parsers/Fbx/Private/FbxHelper.cpp b/Engine/Plugins/Interchange/Runtime/Source/Parsers/Fbx/Private/FbxHelper.cpp
--- a/Engine/Plugins/Interchange/Runtime/Source/Parsers/Fbx/Private/FbxHelper.cpp
+++ b/Engine/Plugins/Interchange/Runtime/Source/Parsers/Fbx/Private/FbxHelper.cpp
@@ -108,12 +108,6 @@ namespace UE
 				FString PropertyName = FFbxConvert::MakeString(Property.GetName());
 				UE::Interchange::SanitizeName(PropertyName);
 
-				if (PropertyName.Equals(TEXT("none"), ESearchCase::IgnoreCase))
-				{
-					//Replace None by Null because None clash with NAME_None and the create asset will instead call the object ClassName_X
-					PropertyName = TEXT("Null");
-				}
-
 				return PropertyName;
 			}
 
@@ -127,12 +121,6 @@ namespace UE
 				FString ObjName = FFbxConvert::MakeString(Object->GetName());
 				UE::Interchange::SanitizeName(ObjName, bIsJoint);
 
-				if (ObjName.Equals(TEXT("none"), ESearchCase::IgnoreCase))
-				{
-					//Replace None by Null because None clash with NAME_None and the create asset will instead call the object ClassName_X
-					ObjName = TEXT("Null");
-				}
-
 				return ObjName;
 			}
 
diff --git a/Engine/Plugins/Interchange/Runtime/Source/Parsers/Fbx/Private/Ufbx/UfbxParser.cpp b/Engine/Plugins/Interchange/Runtime/Source/Parsers/Fbx/Private/Ufbx/UfbxParser.cpp
--- a/Engine/Plugins/Interchange/Runtime/Source/Parsers/Fbx/Private/Ufbx/UfbxParser.cpp
+++ b/Engine/Plugins/Interchange/Runtime/Source/Parsers/Fbx/Private/Ufbx/UfbxParser.cpp
@@ -44,12 +44,6 @@ FUfbxParser::~FUfbxParser()
 // inspired by UE::Interchange::Private::FFbxHelper::GetFbxObjectName
 FString FUfbxParser::GetElementName(const ufbx_element& Element, bool bIsJoint) const
 {
-	if (!FCStringAnsi::Strnicmp("none", Element.name.data, 5))
-	{
-		//Replace None by Null because None clash with NAME_None and the create asset will instead call the object ClassName_X
-		return TEXT("Null");
-	}
-
 	FString ElementName = Convert::ToUnrealString(Element.name);
 	UE::Interchange::SanitizeName(ElementName, bIsJoint);
 	return ElementName;
diff --git a/Engine/Plugins/Interchange/Runtime/Source/Pipelines/Public/InterchangeGenericAssetsPipelineSharedSettings.h b/Engine/Plugins/Interchange/Runtime/Source/Pipelines/Public/InterchangeGenericAssetsPipelineSharedSettings.h
--- a/Engine/Plugins/Interchange/Runtime/Source/Pipelines/Public/InterchangeGenericAssetsPipelineSharedSettings.h
+++ b/Engine/Plugins/Interchange/Runtime/Source/Pipelines/Public/InterchangeGenericAssetsPipelineSharedSettings.h
@@ -70,7 +70,7 @@ public:
  	EInterchangeForceMeshType ForceAllMeshAsType = EInterchangeForceMeshType::IFMT_None;
 
 	/**
-	* If enabled, converted static meshes to skeletal meshes will use a skeleton with a unique bone.
+	* If enabled, converted static meshes to skeletal meshes will use a skeleton with a unique bone at the geometry origin.
 	* Otherwise, the bone structure is determined by the model hierarchy.
 	* This setting is only used if Force All Meshes As Type setting is set to "Skeletal Mesh".
 	*/
diff --git a/Engine/Source/Runtime/Interchange/Core/Private/InterchangeHelper.cpp b/Engine/Source/Runtime/Interchange/Core/Private/InterchangeHelper.cpp
--- a/Engine/Source/Runtime/Interchange/Core/Private/InterchangeHelper.cpp
+++ b/Engine/Source/Runtime/Interchange/Core/Private/InterchangeHelper.cpp
@@ -65,6 +65,8 @@ namespace UE::Interchange
 			OutName.ReplaceCharInline(*InvalidChar, TCHAR('_'), ESearchCase::CaseSensitive);
 			++InvalidChar;
 		}
+
+		OutName = ReplaceNoneWithNull(OutName);
 	}
 
 	FString MakeName(const FString& InName, bool bIsJoint)
@@ -116,4 +118,17 @@ namespace UE::Interchange
 
 		return Left + CharReplacement + Right;
 	}
+
+	FString ReplaceNoneWithNull(const FString& InName)
+	{
+		FString Sanitized = InName.TrimStartAndEnd();
+		if (Sanitized.IsEmpty() || Sanitized.Equals(TEXT("none"), ESearchCase::IgnoreCase))
+		{
+			//Replace None by Null, since None name will be interpreted like NAME_None which will not work with UObject creation
+			//UObject Creation will name it ClassName_X instead of None
+			Sanitized = FString(TEXT("Null"));
+		}
+
+		return Sanitized;
+	}
 };
\ No newline at end of file
diff --git a/Engine/Source/Runtime/Interchange/Core/Private/Nodes/InterchangeBaseNode.cpp b/Engine/Source/Runtime/Interchange/Core/Private/Nodes/InterchangeBaseNode.cpp
--- a/Engine/Source/Runtime/Interchange/Core/Private/Nodes/InterchangeBaseNode.cpp
+++ b/Engine/Source/Runtime/Interchange/Core/Private/Nodes/InterchangeBaseNode.cpp
@@ -7,6 +7,7 @@
 #include "UObject/Class.h"
 #include "UObject/Object.h"
 #include "UObject/ObjectMacros.h"
+#include "InterchangeHelper.h"
 
 #include UE_INLINE_GENERATED_CPP_BY_NAME(InterchangeBaseNode)
 
@@ -94,8 +95,7 @@ void UInterchangeBaseNode::InitializeNode(const FString& UniqueID, const FString
 		LogAttributeStorageErrors(Result, TEXT("RegisterAttribute"), UE::Interchange::FBaseNodeStaticData::UniqueIDKey());
 	}
 
-	Result = Attributes->RegisterAttribute(UE::Interchange::FBaseNodeStaticData::DisplayLabelKey(), DisplayLabel, UE::Interchange::EAttributeProperty::NoHash);
-	if (!IsAttributeStorageResultSuccess(Result))
+	if (!SetDisplayLabel(DisplayLabel))
 	{
 		LogAttributeStorageErrors(Result, TEXT("RegisterAttribute"), UE::Interchange::FBaseNodeStaticData::DisplayLabelKey());
 	}
@@ -343,9 +343,14 @@ FString UInterchangeBaseNode::GetDisplayLabel() const
 	return DisplayLabel;
 }
 
-bool UInterchangeBaseNode::SetDisplayLabel(const FString& DisplayLabel)
+bool UInterchangeBaseNode::SetDisplayLabel(const FString& InDisplayLabel)
 {
-	UE::Interchange::EAttributeStorageResult Result = Attributes->RegisterAttribute(UE::Interchange::FBaseNodeStaticData::DisplayLabelKey(), DisplayLabel);
+	//Having None as a display label will most certainly cause issue in the editor when creating any kind of asset.
+	//It's preferable to be very defensive and avoid that value completely.
+	//This can still be bypassed by using the Attributes system directly.
+	FString Sanitized = UE::Interchange::ReplaceNoneWithNull(InDisplayLabel);
+
+	UE::Interchange::EAttributeStorageResult Result = Attributes->RegisterAttribute(UE::Interchange::FBaseNodeStaticData::DisplayLabelKey(), Sanitized);
 	if (IsAttributeStorageResultSuccess(Result))
 	{
 		UE::Interchange::FAttributeStorage::TAttributeHandle<FString> Handle = Attributes->GetAttributeHandle<FString>(UE::Interchange::FBaseNodeStaticData::DisplayLabelKey());
diff --git a/Engine/Source/Runtime/Interchange/Core/Private/Nodes/InterchangeBaseNodeContainer.cpp b/Engine/Source/Runtime/Interchange/Core/Private/Nodes/InterchangeBaseNodeContainer.cpp
--- a/Engine/Source/Runtime/Interchange/Core/Private/Nodes/InterchangeBaseNodeContainer.cpp
+++ b/Engine/Source/Runtime/Interchange/Core/Private/Nodes/InterchangeBaseNodeContainer.cpp
@@ -39,14 +39,6 @@ FString UInterchangeBaseNodeContainer::AddNode(UInterchangeBaseNode* Node)
 		return NodeUniqueID;
 	}
 
-	if (Node->GetDisplayLabel().IsEmpty())
-	{
-		//Replace None by Null, since None name will be interpret like NAME_None which will not work with UObject creation
-		//UObject Creation will name it ClassName_X instead of None
-		//TODO Log an warning to the user
-		Node->SetDisplayLabel(FString(TEXT("Null")));
-	}
-
 	//Copy the node
 	Nodes.Add(NodeUniqueID, Node);
 	return NodeUniqueID;
diff --git a/Engine/Source/Runtime/Interchange/Core/Public/InterchangeHelper.h b/Engine/Source/Runtime/Interchange/Core/Public/InterchangeHelper.h
--- a/Engine/Source/Runtime/Interchange/Core/Public/InterchangeHelper.h
+++ b/Engine/Source/Runtime/Interchange/Core/Public/InterchangeHelper.h
@@ -32,6 +32,14 @@ namespace UE::Interchange
 	 */
 	INTERCHANGECORE_API FString GetAssetNameWBudget(const FString& ParentPackage, int32 CharBudget, TCHAR CharReplacement = TEXT('_'));
 
+	/**
+     * Returns "Null" if the input is empty or equals "None" (case-insensitive), otherwise returns the input unchanged.
+	 * This avoids Unreal Engine's special handling of "None" as the sentinel NAME_None, which can cause issues during UObject creation.
+     * @param InName - Name string to check.
+     * @return "Null" if InName is "None", otherwise InName.
+    */
+	INTERCHANGECORE_API FString ReplaceNoneWithNull(const FString& InName);
+
 	class FScopedLambda
 	{
 	public:
diff --git a/Engine/Source/Runtime/Interchange/Core/Public/Nodes/InterchangeBaseNode.h b/Engine/Source/Runtime/Interchange/Core/Public/Nodes/InterchangeBaseNode.h
--- a/Engine/Source/Runtime/Interchange/Core/Public/Nodes/InterchangeBaseNode.h
+++ b/Engine/Source/Runtime/Interchange/Core/Public/Nodes/InterchangeBaseNode.h
@@ -436,9 +436,10 @@ public:
 
 	/**
 	 * Change the display label.
-	 */
+     * @param InDisplayLabel - Desired display label text. Sanitized to "Null" if empty or if equals "None".
+    */
 	UFUNCTION(BlueprintCallable, Category = "Interchange | Node")
-	INTERCHANGECORE_API bool SetDisplayLabel(const FString& DisplayName);
+	INTERCHANGECORE_API bool SetDisplayLabel(const FString& InDisplayLabel);
 
 	/**
 	 * Namespace is a prefix added when calling GetUniqueId().
