[Interchange] Added tests for None named assets

This validates the fixes made in  https://p4-swarm.epicgames.net/changes/45733331

#rnx

diff --git a/Engine/Source/Runtime/Interchange/Core/Public/InterchangeHelper.h b/Engine/Source/Runtime/Interchange/Core/Public/InterchangeHelper.h
--- a/Engine/Source/Runtime/Interchange/Core/Public/InterchangeHelper.h
+++ b/Engine/Source/Runtime/Interchange/Core/Public/InterchangeHelper.h
@@ -36,7 +36,7 @@ namespace UE::Interchange
      * Returns "Null" if the input is empty or equals "None" (case-insensitive), otherwise returns the input unchanged.
 	 * This avoids Unreal Engine's special handling of "None" as the sentinel NAME_None, which can cause issues during UObject creation.
      * @param InName - Name string to check.
-     * @return "Null" if InName is "None", otherwise InName.
+     * @return "Null" if InName is empty or "None", otherwise InName.
     */
 	INTERCHANGECORE_API FString ReplaceNoneWithNull(const FString& InName);
 
