Interchange: Fix crash when resaving assets after deserializing them with Interchange disabled.

diff --git a/Engine/Source/Runtime/Interchange/Engine/Private/InterchangeAssetImportData.cpp b/Engine/Source/Runtime/Interchange/Engine/Private/InterchangeAssetImportData.cpp
--- a/Engine/Source/Runtime/Interchange/Engine/Private/InterchangeAssetImportData.cpp
+++ b/Engine/Source/Runtime/Interchange/Engine/Private/InterchangeAssetImportData.cpp
@@ -195,11 +195,14 @@ void UInterchangeAssetImportData::Serialize(FArchive& Ar)
 				ProcessContainerCache();
 				ProcessPipelinesCache();
 
-				// Serialize NodeContainer into CachedNodeContainer
-				FLargeMemoryWriter NodeContainerAr;
-				NodeContainerAr.UsingCustomVersion(FInterchangeCustomVersion::GUID);
-				TransientNodeContainer->SerializeNodeContainerData(NodeContainerAr);
-				CachedNodeContainer = TArray64<uint8>(NodeContainerAr.GetData(), NodeContainerAr.TotalSize());
+				// Serialize NodeContainer into CachedNodeContainer if we have one
+				if (TransientNodeContainer)
+				{
+					FLargeMemoryWriter NodeContainerAr;
+					NodeContainerAr.UsingCustomVersion(FInterchangeCustomVersion::GUID);
+					TransientNodeContainer->SerializeNodeContainerData(NodeContainerAr);
+					CachedNodeContainer = TArray64<uint8>(NodeContainerAr.GetData(), NodeContainerAr.TotalSize());
+				}
 
 				// Serialize pipelines into CachedPipelines
 				TArray<TObjectPtr<UObject>> CopiedPipelines = TransientPipelines;
@@ -416,127 +419,116 @@ UInterchangeFactoryBaseNode* UInterchangeAssetImportData::GetStoredFactoryNode(c
 
 void UInterchangeAssetImportData::ProcessContainerCache() const
 {
-	if (UInterchangeManager::IsInterchangeImportEnabled())
+	//de-serialize
+	if (!TransientNodeContainer && CachedNodeContainer.Num() > 0)
 	{
-		//de-serialize
-		if (!TransientNodeContainer && CachedNodeContainer.Num() > 0)
-		{
-			FLargeMemoryReader NodeContainerAr(CachedNodeContainer.GetData(), CachedNodeContainer.Num());
-			NodeContainerAr.SetCustomVersions(CachedCustomVersions);
+		FLargeMemoryReader NodeContainerAr(CachedNodeContainer.GetData(), CachedNodeContainer.Num());
+		NodeContainerAr.SetCustomVersions(CachedCustomVersions);
 
-			TransientNodeContainer = NewObject<UInterchangeBaseNodeContainer>();
-			TransientNodeContainer->SerializeNodeContainerData(NodeContainerAr);
-		}
-		else if(!TransientNodeContainer)
-		{
-			ProcessDeprecatedData();
-		}
+		TransientNodeContainer = NewObject<UInterchangeBaseNodeContainer>();
+		TransientNodeContainer->SerializeNodeContainerData(NodeContainerAr);
+	}
+
+	if(!TransientNodeContainer)
+	{
+		ProcessDeprecatedData();
 	}
 }
 
 void UInterchangeAssetImportData::ProcessTranslatorCache() const
 {
-	if (UInterchangeManager::IsInterchangeImportEnabled())
+	//Verify our transient object was not garbage collect
+	if (TransientTranslatorSettings && (TransientTranslatorSettings->IsGarbageEliminationEnabled() || TransientTranslatorSettings->HasAnyFlags(RF_BeginDestroyed | RF_FinishDestroyed)) )
 	{
-		//Verify our transient object was not garbage collect
-		if (TransientTranslatorSettings && (TransientTranslatorSettings->IsGarbageEliminationEnabled() || TransientTranslatorSettings->HasAnyFlags(RF_BeginDestroyed | RF_FinishDestroyed)) )
-		{
-			TransientTranslatorSettings = nullptr;
-		}
+		TransientTranslatorSettings = nullptr;
+	}
 
-		//de-serialize
-		if (!TransientTranslatorSettings && !CachedTranslatorSettings.Key.IsEmpty())
+	//de-serialize
+	if (!TransientTranslatorSettings && !CachedTranslatorSettings.Key.IsEmpty())
+	{
+		TMap<FString, UClass*> ClassPerName;
+		for (FThreadSafeObjectIterator It(UClass::StaticClass()); It; ++It)
 		{
-			TMap<FString, UClass*> ClassPerName;
-			for (FThreadSafeObjectIterator It(UClass::StaticClass()); It; ++It)
+			UClass* Class = Cast<UClass>(*It);
+			if (Class->IsChildOf(UInterchangeTranslatorSettings::StaticClass()))
 			{
-				UClass* Class = Cast<UClass>(*It);
-				if (Class->IsChildOf(UInterchangeTranslatorSettings::StaticClass()))
-				{
-					ClassPerName.Add(Class->GetFullName(), Class);
-				}
+				ClassPerName.Add(Class->GetFullName(), Class);
 			}
+		}
 
-			FString ClassFullName = CachedTranslatorSettings.Key;
-			FCoreRedirectObjectName RedirectedObjectName = FCoreRedirects::GetRedirectedName(ECoreRedirectFlags::Type_Class, FCoreRedirectObjectName(ClassFullName));
-			if (RedirectedObjectName.IsValid())
-			{
-				ClassFullName = RedirectedObjectName.ToString();
-			}
-			//This cannot fail to make sure we have a healty serialization
-			if (!ensure(ClassPerName.Contains(ClassFullName)))
-			{
-				//We did not successfully serialize the content of the file into the node container
-				return;
-			}
+		FString ClassFullName = CachedTranslatorSettings.Key;
+		FCoreRedirectObjectName RedirectedObjectName = FCoreRedirects::GetRedirectedName(ECoreRedirectFlags::Type_Class, FCoreRedirectObjectName(ClassFullName));
+		if (RedirectedObjectName.IsValid())
+		{
+			ClassFullName = RedirectedObjectName.ToString();
+		}
+		//This cannot fail to make sure we have a healty serialization
+		if (!ensure(ClassPerName.Contains(ClassFullName)))
+		{
+			//We did not successfully serialize the content of the file into the node container
+			return;
+		}
 
-			UClass* ToCreateClass = ClassPerName.FindChecked(ClassFullName);
+		UClass* ToCreateClass = ClassPerName.FindChecked(ClassFullName);
 
-			TransientTranslatorSettings = Cast<UInterchangeTranslatorSettings>(DeSerializeTranslatorSettings(CachedTranslatorSettings.Value, ToCreateClass));
-			TransientTranslatorSettings->SetFlags(RF_Standalone);
-		}
+		TransientTranslatorSettings = Cast<UInterchangeTranslatorSettings>(DeSerializeTranslatorSettings(CachedTranslatorSettings.Value, ToCreateClass));
+		TransientTranslatorSettings->SetFlags(RF_Standalone);
 	}
 }
 
 void UInterchangeAssetImportData::ProcessPipelinesCache() const
 {
-	if (UInterchangeManager::IsInterchangeImportEnabled())
+	if ((TransientPipelines.Num() == 0) && (CachedPipelines.Num() > 0))
 	{
-		if ((TransientPipelines.Num() == 0) && (CachedPipelines.Num() > 0))
-		{
-			TransientPipelines.Reset(CachedPipelines.Num());
+		TransientPipelines.Reset(CachedPipelines.Num());
 
-			TMap<FString, UClass*> ClassPerName = GetPipelineClassPerName();
+		TMap<FString, UClass*> ClassPerName = GetPipelineClassPerName();
 
-			for (const TPair<FString, FString>& CachedPipeline : CachedPipelines)
+		for (const TPair<FString, FString>& CachedPipeline : CachedPipelines)
+		{
+			FString ClassFullName = CachedPipeline.Key;
+			FCoreRedirectObjectName RedirectedObjectName = FCoreRedirects::GetRedirectedName(ECoreRedirectFlags::Type_Class, FCoreRedirectObjectName(ClassFullName));
+			if (RedirectedObjectName.IsValid())
 			{
-				FString ClassFullName = CachedPipeline.Key;
-				FCoreRedirectObjectName RedirectedObjectName = FCoreRedirects::GetRedirectedName(ECoreRedirectFlags::Type_Class, FCoreRedirectObjectName(ClassFullName));
-				if (RedirectedObjectName.IsValid())
-				{
-					ClassFullName = RedirectedObjectName.ToString();
-				}
+				ClassFullName = RedirectedObjectName.ToString();
+			}
 
-				//This cannot fail to make sure we have a healthy serialization
-				if (!ensure(ClassPerName.Contains(ClassFullName)))
-				{
-					//We did not successfully serialize the cached pipelines
-					TransientPipelines.Empty();
-					return;
-				}
+			//This cannot fail to make sure we have a healthy serialization
+			if (!ensure(ClassPerName.Contains(ClassFullName)))
+			{
+				//We did not successfully serialize the cached pipelines
+				TransientPipelines.Empty();
+				return;
+			}
 
-				UClass* ToCreateClass = ClassPerName.FindChecked(ClassFullName);
-				
-				UObject* Pipeline = DeSerializePipeline(CachedPipeline.Value, ToCreateClass);
+			UClass* ToCreateClass = ClassPerName.FindChecked(ClassFullName);
+			
+			UObject* Pipeline = DeSerializePipeline(CachedPipeline.Value, ToCreateClass);
 
-				TransientPipelines.Add(Pipeline);
-			}
-		}
-		else if (TransientPipelines.Num() == 0)
-		{
-			ProcessDeprecatedData();
+			TransientPipelines.Add(Pipeline);
 		}
 	}
+	else if (TransientPipelines.Num() == 0)
+	{
+		ProcessDeprecatedData();
+	}
 }
 
 void UInterchangeAssetImportData::ProcessDeprecatedData() const
 {
-	if (UInterchangeManager::IsInterchangeImportEnabled())
+	if (!TransientNodeContainer && NodeContainer_DEPRECATED)
 	{
-		if (!TransientNodeContainer && NodeContainer_DEPRECATED)
-		{
-			SetNodeContainer(NodeContainer_DEPRECATED.Get());
-		}
-		if (TransientPipelines.Num() == 0)
-		{
-			TransientPipelines.Empty();
+		SetNodeContainer(NodeContainer_DEPRECATED.Get());
+	}
+	if (TransientPipelines.Num() == 0)
+	{
+		TransientPipelines.Empty();
 
-			for (const TObjectPtr<UObject>& PipelineObject : Pipelines_DEPRECATED)
+		for (const TObjectPtr<UObject>& PipelineObject : Pipelines_DEPRECATED)
+		{
+			if (PipelineObject)
 			{
-				if (PipelineObject)
-				{
-					TransientPipelines.Add(PipelineObject.Get());
-				}
+				TransientPipelines.Add(PipelineObject.Get());
 			}
 		}
 	}
