Fixed issue when comparing two attribute storages containing float values
- Root cause: The current comparison was based on the hash of the attribute storage. No useful with floats
- Solution: Add code to compare float values with a threshold

#rnx

diff --git a/Engine/Source/Runtime/Interchange/Core/Private/Types/AttributeStorage.cpp b/Engine/Source/Runtime/Interchange/Core/Private/Types/AttributeStorage.cpp
--- a/Engine/Source/Runtime/Interchange/Core/Private/Types/AttributeStorage.cpp
+++ b/Engine/Source/Runtime/Interchange/Core/Private/Types/AttributeStorage.cpp
@@ -16,11 +16,13 @@
 #include "Math/IntVector.h"
 #include "Math/Matrix.h"
 #include "Math/Plane.h"
+#include "Math/OrientedBox.h"
 #include "Math/Quat.h"
 #include "Math/RandomStream.h"
 #include "Math/Rotator.h"
 #include "Math/Sphere.h"
 #include "Math/TwoVectors.h"
+#include "Math/UnrealMathUtility.h"
 #include "Math/Vector.h"
 #include "Math/Vector2D.h"
 #include "Math/Vector2DHalf.h"
@@ -1431,6 +1433,330 @@ namespace UE
 			return FGuid(Hash[0] ^ Hash[4], Hash[1], Hash[2], Hash[3]);
 		}
 
+#define COMPARE_MATH_STRUCTS( MathType ) \
+	case EAttributeTypes::MathType: \
+	{ \
+		F##MathType BaseTypeValue; \
+		ensure(GetAttributeValue(*this, BaseTypeValue)); \
+		F##MathType VersionTypeValue; \
+		ensure(GetAttributeValue(VersionStorage, VersionTypeValue)); \
+		return BaseTypeValue.Equals(VersionTypeValue); \
+	}
+
+#define COMPARE_MATH_ARRAY( MathType ) \
+	case EAttributeTypes::MathType##Array: \
+	{ \
+		TArray<F##MathType> BaseArray; \
+		ensure(GetAttributeValue(*this, BaseArray)); \
+		TArray<F##MathType> VersionArray; \
+		ensure(GetAttributeValue(VersionStorage, VersionArray)); \
+		for(int32 Index = 0; Index < BaseArray.Num(); ++Index) \
+		{ \
+			if(!BaseArray[Index].Equals(VersionArray[Index])) \
+			{ \
+				return false; \
+			} \
+		} \
+		return true; \
+	}
+
+#define DECLARE_LOCAL_VARIABLES( ValueType ) \
+	F##ValueType Base##ValueType; \
+	ensure(GetAttributeValue(*this, Base##ValueType)); \
+	F##ValueType Version##ValueType; \
+	ensure(GetAttributeValue(VersionStorage, Version##ValueType));
+
+#define DECLARE_LOCAL_ARRAY_VARIABLES( ValueType ) \
+	TArray<F##ValueType> Base##ValueType##s; \
+	ensure(GetAttributeValue(*this, Base##ValueType##s)); \
+	TArray<F##ValueType> Version##ValueType##s; \
+	ensure(GetAttributeValue(VersionStorage, Version##ValueType##s));
+
+		bool FAttributeStorage::AreAllocationInfosEqual(const FAttributeKey& BaseKey, const FAttributeAllocationInfo& BaseInfo, const FAttributeStorage& VersionStorage, const FAttributeAllocationInfo& VersionInfo) const
+		{
+			if (BaseInfo == VersionInfo)
+			{
+				return true;
+			}
+
+			using namespace UE::Math;
+
+			auto CompareTransform = []<typename T>(const TTransform<T>& Base, const TTransform<T>& Version) -> bool
+			{
+				bool bResult = TTransform<T>::AreTranslationsEqual(Base, Version);
+				bResult &= TTransform<T>::AreScale3DsEqual(Base, Version);
+
+				// Do not compare TQuat values but resulting TRotator values as TQuat values are eventually converted to TRotator
+				// and 2 different TQuat values can generate the same TRotator value
+				bResult &= Base.GetRotation().Rotator().Equals(Version.GetRotation().Rotator());
+
+				return bResult;
+			};
+
+			auto GetAttributeValue = [&BaseKey]<typename ValueType>(const FAttributeStorage & Attributes, ValueType& OutValue) -> bool
+			{
+				if (!Attributes.ContainAttribute(BaseKey))
+				{
+					return false;
+				}
+
+				TAttributeHandle<ValueType> AttributeHandle = Attributes.GetAttributeHandle<ValueType>(BaseKey);
+				if (!AttributeHandle.IsValid())
+				{
+					return false;
+				}
+
+				EAttributeStorageResult Result = AttributeHandle.Get(OutValue);
+				if (!IsAttributeStorageResultSuccess(Result))
+				{
+					LogAttributeStorageErrors(Result, TEXT("FAttributeStorage::AreAllocationInfosEqual"), BaseKey);
+					return false;
+				}
+
+				return true;
+			};
+
+
+			// Hash comparison cannot be relied on for floating point values
+			switch (BaseInfo.Type)
+			{
+				COMPARE_MATH_STRUCTS(Vector4d)
+				COMPARE_MATH_STRUCTS(Vector4f)
+				COMPARE_MATH_STRUCTS(Vector3d)
+				COMPARE_MATH_STRUCTS(Vector3f)
+				COMPARE_MATH_STRUCTS(Vector2d)
+				COMPARE_MATH_STRUCTS(Vector2f)
+				COMPARE_MATH_STRUCTS(Box3d)
+				COMPARE_MATH_STRUCTS(Box3f)
+				COMPARE_MATH_STRUCTS(Box2D)
+				COMPARE_MATH_STRUCTS(Box2f)
+				COMPARE_MATH_STRUCTS(Sphere3d)
+				COMPARE_MATH_STRUCTS(Sphere3f)
+				COMPARE_MATH_STRUCTS(Plane4d)
+				COMPARE_MATH_STRUCTS(Plane4f)
+				COMPARE_MATH_STRUCTS(Quat4d)
+				COMPARE_MATH_STRUCTS(Quat4f)
+				COMPARE_MATH_STRUCTS(Matrix44d)
+				COMPARE_MATH_STRUCTS(Matrix44f)
+				COMPARE_MATH_STRUCTS(Rotator3d)
+				COMPARE_MATH_STRUCTS(Rotator3f)
+				COMPARE_MATH_STRUCTS(LinearColor)
+
+				COMPARE_MATH_ARRAY(Vector4d)
+				COMPARE_MATH_ARRAY(Vector4f)
+				COMPARE_MATH_ARRAY(Vector3d)
+				COMPARE_MATH_ARRAY(Vector3f)
+				COMPARE_MATH_ARRAY(Vector2d)
+				COMPARE_MATH_ARRAY(Vector2f)
+				COMPARE_MATH_ARRAY(Box3d)
+				COMPARE_MATH_ARRAY(Box3f)
+				COMPARE_MATH_ARRAY(Box2D)
+				COMPARE_MATH_ARRAY(Box2f)
+				COMPARE_MATH_ARRAY(Sphere3d)
+				COMPARE_MATH_ARRAY(Sphere3f)
+				COMPARE_MATH_ARRAY(Plane4d)
+				COMPARE_MATH_ARRAY(Plane4f)
+				COMPARE_MATH_ARRAY(Quat4d)
+				COMPARE_MATH_ARRAY(Quat4f)
+				COMPARE_MATH_ARRAY(Matrix44d)
+				COMPARE_MATH_ARRAY(Matrix44f)
+				COMPARE_MATH_ARRAY(Rotator3d)
+				COMPARE_MATH_ARRAY(Rotator3f)
+				COMPARE_MATH_ARRAY(LinearColor)
+
+				case EAttributeTypes::Transform3f:
+				{
+					DECLARE_LOCAL_VARIABLES(Transform3f )
+					return CompareTransform(BaseTransform3f, VersionTransform3f);
+				}
+
+				case EAttributeTypes::Transform3fArray:
+				{
+					DECLARE_LOCAL_ARRAY_VARIABLES(Transform3f)
+
+					for (int32 Index = 0; Index < BaseTransform3fs.Num(); ++Index)
+					{
+						if (!CompareTransform(BaseTransform3fs[Index], VersionTransform3fs[Index]))
+						{
+							return false;
+						}
+					}
+
+					return true;
+				}
+
+				case EAttributeTypes::Transform3d:
+				{
+					DECLARE_LOCAL_VARIABLES(Transform3d)
+					return CompareTransform(BaseTransform3d, VersionTransform3d);
+				}
+
+				case EAttributeTypes::Transform3dArray:
+				{
+					DECLARE_LOCAL_ARRAY_VARIABLES(Transform3d)
+
+					for (int32 Index = 0; Index < BaseTransform3ds.Num(); ++Index)
+					{
+						if (!CompareTransform(BaseTransform3ds[Index], VersionTransform3ds[Index]))
+						{
+							return false;
+						}
+					}
+
+					return true;
+				}
+
+				case EAttributeTypes::OrientedBox:
+				{
+					DECLARE_LOCAL_VARIABLES(OrientedBox)
+
+					bool bIsEqual = BaseOrientedBox.AxisX.Equals(VersionOrientedBox.AxisX);
+					bIsEqual &= BaseOrientedBox.AxisY.Equals(VersionOrientedBox.AxisY);
+					bIsEqual &= BaseOrientedBox.AxisZ.Equals(VersionOrientedBox.AxisZ);
+					bIsEqual &= FMath::IsNearlyEqual(BaseOrientedBox.ExtentX, VersionOrientedBox.ExtentX);
+					bIsEqual &= FMath::IsNearlyEqual(BaseOrientedBox.ExtentY, VersionOrientedBox.ExtentY);
+					bIsEqual &= FMath::IsNearlyEqual(BaseOrientedBox.ExtentZ, VersionOrientedBox.ExtentZ);
+
+					return bIsEqual;
+				}
+
+				case EAttributeTypes::TwoVectors:
+				{
+					DECLARE_LOCAL_VARIABLES(TwoVectors)
+
+					bool bIsEqual = BaseTwoVectors.v1.Equals(BaseTwoVectors.v1);
+					bIsEqual &= BaseTwoVectors.v2.Equals(BaseTwoVectors.v2);
+
+					return bIsEqual;
+				}
+
+				case EAttributeTypes::TwoVectorsArray:
+				{
+					DECLARE_LOCAL_ARRAY_VARIABLES(TwoVectors)
+
+					for (int32 Index = 0; Index < BaseTwoVectorss.Num(); ++Index)
+					{
+						const FTwoVectors& BaseValue = BaseTwoVectorss[Index];
+						const FTwoVectors& VersionValue = VersionTwoVectorss[Index];
+
+						if (!BaseValue.v1.Equals(VersionValue.v1) || !BaseValue.v2.Equals(VersionValue.v2))
+						{
+							return false;
+						}
+					}
+
+					return true;
+				}
+
+				case EAttributeTypes::Float16:
+				{
+					DECLARE_LOCAL_VARIABLES(Float16)
+					return FMath::IsNearlyEqual(BaseFloat16.GetFloat(), VersionFloat16.GetFloat());
+				}
+
+				case EAttributeTypes::Float16Array:
+				{
+					DECLARE_LOCAL_ARRAY_VARIABLES(Float16)
+
+					for (int32 Index = 0; Index < BaseFloat16s.Num(); ++Index)
+					{
+						if (!FMath::IsNearlyEqual(BaseFloat16s[Index].GetFloat(), VersionFloat16s[Index].GetFloat()))
+						{
+							return false;
+						}
+					}
+
+					return true;
+				}
+
+				case EAttributeTypes::Vector2DHalf:
+				{
+					DECLARE_LOCAL_VARIABLES(Vector2DHalf)
+					return FMath::IsNearlyEqual(BaseVector2DHalf.X.GetFloat(), VersionVector2DHalf.X.GetFloat())
+						&& FMath::IsNearlyEqual(BaseVector2DHalf.Y.GetFloat(), VersionVector2DHalf.Y.GetFloat());
+				}
+
+				case EAttributeTypes::Vector2DHalfArray:
+				{
+					DECLARE_LOCAL_ARRAY_VARIABLES(Vector2DHalf)
+
+					for (int32 Index = 0; Index < BaseVector2DHalfs.Num(); ++Index)
+					{
+						const FVector2DHalf& BaseValue = BaseVector2DHalfs[Index];
+						const FVector2DHalf& VersionValue = VersionVector2DHalfs[Index];
+						
+						if (!FMath::IsNearlyEqual(BaseValue.X.GetFloat(), VersionValue.X.GetFloat()) || !FMath::IsNearlyEqual(BaseValue.Y.GetFloat(), VersionValue.Y.GetFloat()))
+						{
+							return false;
+						}
+					}
+
+					return true;
+				}
+
+				case EAttributeTypes::Float:
+				{
+					float BaseFloat;
+					ensure(GetAttributeValue(*this, BaseFloat));
+					float VersionFloat;
+					ensure(GetAttributeValue(VersionStorage, VersionFloat));
+
+					return FMath::IsNearlyEqual(BaseFloat, VersionFloat);
+				}
+
+				case EAttributeTypes::FloatArray:
+				{
+					TArray<float> BaseFloats;
+					ensure(GetAttributeValue(*this, BaseFloats));
+					TArray<float> VersionFloats;
+					ensure(GetAttributeValue(VersionStorage, VersionFloats));
+
+					for (int32 Index = 0; Index < BaseFloats.Num(); ++Index)
+					{
+						if (!FMath::IsNearlyEqual(BaseFloats[Index], VersionFloats[Index]))
+						{
+							return false;
+						}
+					}
+
+					return true;
+				}
+
+				case EAttributeTypes::Double:
+				{
+					double BaseDouble;
+					ensure(GetAttributeValue(*this, BaseDouble));
+					double VersionDouble;
+					ensure(GetAttributeValue(VersionStorage, VersionDouble));
+
+					return FMath::IsNearlyEqual(BaseDouble, VersionDouble);
+				}
+
+				case EAttributeTypes::DoubleArray:
+				{
+					TArray<double> BaseDoubles;
+					ensure(GetAttributeValue(*this, BaseDoubles));
+					TArray<double> VersionDoubles;
+					ensure(GetAttributeValue(VersionStorage, VersionDoubles));
+
+					for (int32 Index = 0; Index < BaseDoubles.Num(); ++Index)
+					{
+						if (!FMath::IsNearlyEqual(BaseDoubles[Index], VersionDoubles[Index]))
+						{
+							return false;
+						}
+					}
+
+					return true;
+				}
+
+				default:
+					break; 
+			}
+
+			return BaseInfo.Hash == VersionInfo.Hash;
+		}
+
 		void FAttributeStorage::CompareStorage(const FAttributeStorage& BaseStorage, const FAttributeStorage& VersionStorage, TArray<FAttributeKey>& RemovedAttributes, TArray<FAttributeKey>& AddedAttributes, TArray<FAttributeKey>& ModifiedAttributes)
 		{
 			//Lock the storage
@@ -1447,7 +1773,7 @@ namespace UE
 					//Add the attribute to RemovedAttributes
 					RemovedAttributes.Add(KeyBase);
 				}
-				else if (KvpBase.Value.Hash != AttributeAllocationInfoVersion->Hash)
+				else if (!BaseStorage.AreAllocationInfosEqual(KeyBase, KvpBase.Value, VersionStorage , *AttributeAllocationInfoVersion))
 				{
 					//Add the attribute to ModifiedAttributes
 					ModifiedAttributes.Add(KeyBase);
diff --git a/Engine/Source/Runtime/Interchange/Core/Public/Types/AttributeStorage.h b/Engine/Source/Runtime/Interchange/Core/Public/Types/AttributeStorage.h
--- a/Engine/Source/Runtime/Interchange/Core/Public/Types/AttributeStorage.h
+++ b/Engine/Source/Runtime/Interchange/Core/Public/Types/AttributeStorage.h
@@ -972,7 +972,18 @@ namespace UE
         			return Ar;
         		}
 			};
-		
+
+			/**
+			 * Compare two allocation infos. Returns true if they are equal
+			 *
+			 * @param BaseKey - The reference key.
+			 * @param BaseInfo - The reference allocation info.
+			 * @param VersionStorage - The storage with the changes.
+			 * @param VersionInfo - The allocation info to compare BaseInfo with.
+			 *
+			 */
+			INTERCHANGECORE_API bool AreAllocationInfosEqual(const FAttributeKey& BaseKey, const FAttributeAllocationInfo& BaseInfo, const FAttributeStorage& VersionStorage, const FAttributeAllocationInfo& VersionInfo) const;
+
 			/**
 			 * Set an attribute value into the storage. Return success if the attribute was properly set.
 			 *
