Fixed reimport issues

#rnx

diff --git a/Engine/Source/Runtime/Interchange/Engine/Private/InterchangeImportCommon.cpp b/Engine/Source/Runtime/Interchange/Engine/Private/InterchangeImportCommon.cpp
--- a/Engine/Source/Runtime/Interchange/Engine/Private/InterchangeImportCommon.cpp
+++ b/Engine/Source/Runtime/Interchange/Engine/Private/InterchangeImportCommon.cpp
@@ -432,6 +432,42 @@ namespace UE::Interchange
 		return nullptr;
 	}
 
+	bool FFactoryCommon::CanProceedWithFactoryNode(UInterchangeFactoryBaseNode& FactoryNode, UObject* ReimportObject, UObject* ObjectToReimport, const FString& PackageName, const FString& AssetName, const FString& SubPathString)
+	{
+		// If ReimportObject is null, proceed with factory node
+		if (!ReimportObject)
+		{
+			return true;
+		}
+
+		// Object associated with the factory node was not part of the previous import, proceed with factory node 
+		if (!GetFactoryNode(ReimportObject, PackageName, AssetName, SubPathString))
+		{
+			return true;
+		}
+
+		// Associated object was not found (deleted or renamed) and factory node has been flagged as to be reimported. Proceed with factory node
+		if (!ObjectToReimport && FactoryNode.ShouldForceNodeReimport())
+		{
+			return true;
+		}
+
+		// If ObjectToReimport is null but still tracked in the UInterchangeSceneImportAsset, it means that
+		// the associated object was part of a previous import/reimport but deleted or renamed since.
+		// The UInterchangeSceneImportAsset only tracks factory nodes with reference object.
+		// Consequently keep track of it for future reimport.
+		if (!ObjectToReimport)
+		{
+			const FTopLevelAssetPath AssetPath(*PackageName, *AssetName);
+			FactoryNode.SetCustomReferenceObject(FSoftObjectPath(AssetPath, SubPathString));
+
+			return false;
+		}
+
+
+		return true;
+	}
+
 	UObject* FFactoryCommon::AsyncFindObject(UInterchangeFactoryBaseNode* FactoryNode, const UClass* FactoryClass, UObject* Parent, const FString& AssetName)
 	{
 		if (!FactoryNode)
diff --git a/Engine/Source/Runtime/Interchange/Engine/Private/Tasks/InterchangeTaskCreateSceneObjects.cpp b/Engine/Source/Runtime/Interchange/Engine/Private/Tasks/InterchangeTaskCreateSceneObjects.cpp
--- a/Engine/Source/Runtime/Interchange/Engine/Private/Tasks/InterchangeTaskCreateSceneObjects.cpp
+++ b/Engine/Source/Runtime/Interchange/Engine/Private/Tasks/InterchangeTaskCreateSceneObjects.cpp
@@ -108,20 +108,20 @@ void UE::Interchange::FTaskCreateSceneObjects_GameThread::Execute()
 		}
 		const FString FactoryNodeWorldPath = FactoryImportWorld->GetOutermost()->GetPathName();
 		const FString FactoryNodeWorldName = FactoryImportWorld->GetName();
-		const FString FactoryNodeNodePrefix = FactoryImportLevel->GetName() + TEXT(".");
+		const FString FactoryNodeSceneNodeName = FactoryImportLevel->GetName() + TEXT(".") + SceneNodeName;
 
 		UInterchangeFactoryBase::FImportSceneObjectsParams CreateSceneObjectsParams;
 		CreateSceneObjectsParams.ObjectName = SceneNodeName;
 		CreateSceneObjectsParams.FactoryNode = FactoryNode;
 		CreateSceneObjectsParams.Level = FactoryImportLevel;
 		CreateSceneObjectsParams.ImportAssets = ImportObjects;
-		CreateSceneObjectsParams.ReimportObject = FFactoryCommon::GetObjectToReimport(Factory, ReimportObject, *FactoryNode, FactoryNodeWorldPath, FactoryNodeWorldName, FactoryNodeNodePrefix + SceneNodeName);
-		CreateSceneObjectsParams.ReimportFactoryNode = FFactoryCommon::GetFactoryNode(ReimportObject, FactoryNodeWorldPath, FactoryNodeWorldName, FactoryNodeNodePrefix + SceneNodeName);
+		CreateSceneObjectsParams.ReimportObject = FFactoryCommon::GetObjectToReimport(Factory, ReimportObject, *FactoryNode, FactoryNodeWorldPath, FactoryNodeWorldName, FactoryNodeSceneNodeName);
+		CreateSceneObjectsParams.ReimportFactoryNode = FFactoryCommon::GetFactoryNode(ReimportObject, FactoryNodeWorldPath, FactoryNodeWorldName, FactoryNodeSceneNodeName);
 		CreateSceneObjectsParams.bIsReimport = ReimportObject != nullptr;
 		CreateSceneObjectsParams.AsyncHelperUid = AsyncHelper->UniqueId;
 
-		// Skip factory node if this is a reimport but no object to reimport is found and reimport is not requested
-		if (!CreateSceneObjectsParams.ReimportObject && ReimportObject && !FactoryNode->ShouldForceNodeReimport())
+		// Skip factory node if conditions are not meant
+		if (!FFactoryCommon::CanProceedWithFactoryNode(*FactoryNode, ReimportObject, CreateSceneObjectsParams.ReimportObject, FactoryNodeWorldPath, FactoryNodeWorldName, FactoryNodeSceneNodeName))
 		{
 			FactoryNode->SetSkipNodeImport();
 			continue;
diff --git a/Engine/Source/Runtime/Interchange/Engine/Private/Tasks/InterchangeTaskImportObject.cpp b/Engine/Source/Runtime/Interchange/Engine/Private/Tasks/InterchangeTaskImportObject.cpp
--- a/Engine/Source/Runtime/Interchange/Engine/Private/Tasks/InterchangeTaskImportObject.cpp
+++ b/Engine/Source/Runtime/Interchange/Engine/Private/Tasks/InterchangeTaskImportObject.cpp
@@ -301,11 +301,12 @@ void UE::Interchange::FTaskImportObject_GameThread::Execute()
 	UPackage* Pkg = nullptr;
 	FString PackageName;
 	FString AssetName;
+	UObject* ReimportObject = AsyncHelper->TaskData.ReimportObject;
 	Private::InternalGetPackageName(*AsyncHelper, SourceIndex, PackageBasePath, FactoryNode, PackageName, AssetName);
-	UObject* ObjectToReimport = FFactoryCommon::GetObjectToReimport(Factory, AsyncHelper->TaskData.ReimportObject, *FactoryNode, PackageName, AssetName);
+	UObject* ObjectToReimport = FFactoryCommon::GetObjectToReimport(Factory, ReimportObject, *FactoryNode, PackageName, AssetName);
 
-	// Skip factory node if this is a reimport but no object to reimport is found and reimport is not requested
-	if (!ObjectToReimport && AsyncHelper->TaskData.ReimportObject && !FactoryNode->ShouldForceNodeReimport())
+	// Skip factory node if conditions are not meant
+	if (!FFactoryCommon::CanProceedWithFactoryNode(*FactoryNode, ReimportObject, ObjectToReimport, PackageName, AssetName))
 	{
 		FactoryNode->SetSkipNodeImport();
 		return;
diff --git a/Engine/Source/Runtime/Interchange/Engine/Public/InterchangeImportCommon.h b/Engine/Source/Runtime/Interchange/Engine/Public/InterchangeImportCommon.h
--- a/Engine/Source/Runtime/Interchange/Engine/Public/InterchangeImportCommon.h
+++ b/Engine/Source/Runtime/Interchange/Engine/Public/InterchangeImportCommon.h
@@ -7,6 +7,8 @@
 #include "CoreTypes.h"
 #include "EditorFramework/AssetImportData.h"
 
+#define UE_API INTERCHANGEENGINE_API
+
 class UAssetImportData;
 class UInterchangeAssetImportData;
 class UInterchangeBaseNode;
@@ -39,7 +41,7 @@ namespace UE::Interchange
 			const TArray<UObject*> Pipelines;
 			const UInterchangeTranslatorBase* Translator;
 
-			INTERCHANGEENGINE_API FUpdateImportAssetDataParameters(UObject* InAssetImportDataOuter
+			UE_API FUpdateImportAssetDataParameters(UObject* InAssetImportDataOuter
 																, UAssetImportData* InAssetImportData
 																, const UInterchangeSourceData* InSourceData
 																, FString InNodeUniqueID
@@ -53,7 +55,7 @@ namespace UE::Interchange
 		 * If the AssetImportData is null it will create one. If the AssetImportData is not an UInterchangeAssetImportData it will create a new one.
 		 * @return The source data that should be stored on the asset or nullptr if a parameter is invalid
 		 */
-		static INTERCHANGEENGINE_API UAssetImportData* UpdateImportAssetData(FUpdateImportAssetDataParameters& Parameters);
+		static UE_API UAssetImportData* UpdateImportAssetData(FUpdateImportAssetDataParameters& Parameters);
 
 		/**
 		 * Update the AssetImportData of the specified asset in the parameters. Also update the node container and the node unique id.
@@ -61,7 +63,7 @@ namespace UE::Interchange
 		 * The file source update is done by calling the function parameter CustomFileSourceUpdate, so its the client responsability to properly update the file source.
 		 * @return The source data that should be stored on the asset or nullptr if a parameter is invalid
 		 */
-		static INTERCHANGEENGINE_API UAssetImportData* UpdateImportAssetData(FUpdateImportAssetDataParameters& Parameters, TFunctionRef<void(UInterchangeAssetImportData*)> CustomFileSourceUpdate);
+		static UE_API UAssetImportData* UpdateImportAssetData(FUpdateImportAssetDataParameters& Parameters, TFunctionRef<void(UInterchangeAssetImportData*)> CustomFileSourceUpdate);
 
 #if WITH_EDITORONLY_DATA
 		struct FSetImportAssetDataParameters : public FUpdateImportAssetDataParameters
@@ -69,7 +71,7 @@ namespace UE::Interchange
 			// Allow the factory to provide is own list of source files.
 			TArray<FAssetImportInfo::FSourceFile> SourceFiles;
 
-			INTERCHANGEENGINE_API FSetImportAssetDataParameters(UObject* InAssetImportDataOuter
+			UE_API FSetImportAssetDataParameters(UObject* InAssetImportDataOuter
 				, UAssetImportData* InAssetImportData
 				, const UInterchangeSourceData* InSourceData
 				, FString InNodeUniqueID
@@ -83,45 +85,45 @@ namespace UE::Interchange
 		 * If the AssetImportData is null it will create one. If the AssetImportData is not an UInterchangeAssetImportData it will create a new one.
 		 * @return The source data that should be stored on the asset or nullptr if a parameter is invalid
 		 */
-		static INTERCHANGEENGINE_API UAssetImportData* SetImportAssetData(FSetImportAssetDataParameters& Parameters);
+		static UE_API UAssetImportData* SetImportAssetData(FSetImportAssetDataParameters& Parameters);
 
 		/**
 		 * Fills the OutSourceFilenames array with the list of source files contained in the asset source data.
 		 * Returns true if the operation was successful.
 		 */
-		static INTERCHANGEENGINE_API bool GetSourceFilenames(const UAssetImportData* AssetImportData, TArray<FString>& OutSourceFilenames);
+		static UE_API bool GetSourceFilenames(const UAssetImportData* AssetImportData, TArray<FString>& OutSourceFilenames);
 
 		/**
 		 * Backups the SourceData. Primary usage for re-instating SourceData on Re-Import cancellation.
 		 */
-		static INTERCHANGEENGINE_API void BackupSourceData(const UAssetImportData* AssetImportData);
+		static UE_API void BackupSourceData(const UAssetImportData* AssetImportData);
 
 		/**
 		 * Reinstates the backedup SourceData. Primary usage for re-instating SourceData on Re-Import cancellation.
 		 */
-		static INTERCHANGEENGINE_API void ReinstateSourceData(UAssetImportData* AssetImportData);
+		static UE_API void ReinstateSourceData(UAssetImportData* AssetImportData);
 
 		/**
 		 * Clears the backedup SourceData. Primary usage for re-instating SourceData on Re-Import cancellation.
 		 */
-		static INTERCHANGEENGINE_API void ClearBackupSourceData(const UAssetImportData* AssetImportData);
+		static UE_API void ClearBackupSourceData(const UAssetImportData* AssetImportData);
 
 		/**
 		 * Sets the SourceFileName value at the specified index.
 		 */
-		static INTERCHANGEENGINE_API bool SetSourceFilename(UAssetImportData* AssetImportData, const FString& SourceFilename, int32 SourceIndex, const FString& SourceLabel = FString());
+		static UE_API bool SetSourceFilename(UAssetImportData* AssetImportData, const FString& SourceFilename, int32 SourceIndex, const FString& SourceLabel = FString());
 
 		/**
 		 * Set the object's reimport source at the specified index value.
 		 */
-		static INTERCHANGEENGINE_API bool SetReimportSourceIndex(const UObject* Object, UAssetImportData* AssetImportData, int32 SourceIndex);
+		static UE_API bool SetReimportSourceIndex(const UObject* Object, UAssetImportData* AssetImportData, int32 SourceIndex);
 
 #endif // WITH_EDITORONLY_DATA
 
 		/**
 		 * Apply the current strategy to the PipelineAssetNode
 		 */
-		static INTERCHANGEENGINE_API void ApplyReimportStrategyToAsset(UObject* Asset
+		static UE_API void ApplyReimportStrategyToAsset(UObject* Asset
 											, const UInterchangeFactoryBaseNode* PreviousAssetNode
 											, const UInterchangeFactoryBaseNode* CurrentAssetNode
 											, UInterchangeFactoryBaseNode* PipelineAssetNode);
@@ -130,27 +132,48 @@ namespace UE::Interchange
 		 * If the ReimportObject is a UInterchangeSceneImportAsset, returns the UObject which asset path
 		 * is //PackageName.AssetName[:SubPathString].
 		 * Returns ReimportObject otherwise.
-		 * @param ReimportObject: Original object to reimport
+		 * @param ReimportObject: Object that a reimport has been initialized from
+		 * @param FactoryNode: Factory node being reimported
 		 * @param PackageName: Package path of the actual object to reimport
 		 * @param AssetName: Asset name of the actual object to reimport
 		 * @param SubPathString: Optional subobject name
 		 */
-		static 	INTERCHANGEENGINE_API UObject* GetObjectToReimport(UInterchangeFactoryBase* Factory, UObject* ReimportObject, const UInterchangeFactoryBaseNode& FactoryNode, const FString& PackageName, const FString& AssetName, const FString& SubPathString = FString());
+		static 	UE_API UObject* GetObjectToReimport(UInterchangeFactoryBase* Factory, UObject* ReimportObject, const UInterchangeFactoryBaseNode& FactoryNode, const FString& PackageName, const FString& AssetName, const FString& SubPathString = FString());
 
 		/**
 		 * If the ReimportObject is a UInterchangeSceneImportAsset, returns the factory node
 		 * associated with the asset which path is //PackageName.AssetName[:SubPathString].
 		 * Returns nullptr otherwise.
-		 * @param ReimportObject: Original object to reimport
+		 * @param ReimportObject: Object that a reimport has been initialized from
 		 * @param PackageName: Package path of the actual object to reimport
 		 * @param AssetName: Asset name of the actual object to reimport
 		 * @param SubPathString: Optional subobject name
 		 */
-		static 	INTERCHANGEENGINE_API const UInterchangeFactoryBaseNode* GetFactoryNode(UObject* ReimportObject, const FString& PackageName, const FString& AssetName, const FString& SubPathString = FString());
-		
+		static 	UE_API const UInterchangeFactoryBaseNode* GetFactoryNode(UObject* ReimportObject, const FString& PackageName, const FString& AssetName, const FString& SubPathString = FString());
+
+		/**
+		 * Evaluates whether or not a factory should be called on a factory node in a given reimport context.
+		 * @param FactoryNode: Factory node to evaluate
+		 * @param ReimportObject: Object a reimport has been initialized from
+		 * @param ObjectToReimport: Actual object to reimport
+		 * @param PackageName: Package path of the object associated with the factory node to reimport
+		 * @param AssetName: Asset name of the object associated with the factory node to reimport
+		 * @param SubPathString: Optional subobject name
+		 * Returns:
+		 *		- true if it is an import. ReimportObject is null.
+		 *      - true if, based on its path, object associated with the factory node was not part of the previous import.
+		 *		- true if factory node was part of the previous import and ObjectToReimport is null but
+		 *			FactoryNode has been flagged as required for the reimport
+		 *		- true ObjectToReimport is not null
+		 *		- false otherwise
+		 *	@note The path of the missing associated object will be set on the factory node if ObjectToReimport is null
+		 */
+		static 	bool UE_API CanProceedWithFactoryNode(UInterchangeFactoryBaseNode& FactoryNode, UObject* ReimportObject, UObject* ObjectToReimport, const FString& PackageName, const FString& AssetName, const FString& SubPathString = FString());
 		/**
 		 * Factory helper use to find the object we have create in the game thread when we do the async part of the factory import.
 		 */
-		static INTERCHANGEENGINE_API UObject* AsyncFindObject(UInterchangeFactoryBaseNode* FactoryNode, const UClass* FactoryClass, UObject* Parent, const FString& AssetName);
+		static UE_API UObject* AsyncFindObject(UInterchangeFactoryBaseNode* FactoryNode, const UClass* FactoryClass, UObject* Parent, const FString& AssetName);
 	};
 } //ns UE::Interchange
+
+#undef UE_API
