Fixed crash when bogus content path is passed as targeted folder

#rnx
#DisabledValidations CheckCopyright, CheckBalancedMacros

diff --git a/Engine/Source/Runtime/Interchange/Engine/Private/InterchangeManager.cpp b/Engine/Source/Runtime/Interchange/Engine/Private/InterchangeManager.cpp
--- a/Engine/Source/Runtime/Interchange/Engine/Private/InterchangeManager.cpp
+++ b/Engine/Source/Runtime/Interchange/Engine/Private/InterchangeManager.cpp
@@ -1831,6 +1831,12 @@ UInterchangeManager::ImportInternal(const FString& ContentPath, const UInterchan
 		return TTuple<UE::Interchange::FAssetImportResultRef, UE::Interchange::FSceneImportResultRef>{ MakeShared< UE::Interchange::FImportResult, ESPMode::ThreadSafe >(), MakeShared< UE::Interchange::FImportResult, ESPMode::ThreadSafe >() };
 	}
 
+	if (!FPackageName::IsValidPath(ContentPath))
+	{
+		UE_LOG(LogInterchangeEngine, Error, TEXT("Destination content path is not valid."));
+		return TTuple<UE::Interchange::FAssetImportResultRef, UE::Interchange::FSceneImportResultRef>{ MakeShared< UE::Interchange::FImportResult, ESPMode::ThreadSafe >(), MakeShared< UE::Interchange::FImportResult, ESPMode::ThreadSafe >() };
+	}
+
 	ensure(IsInterchangeImportEnabled());
 
 	static int32 GeneratedUniqueID = 0;
