Fixed logic flaw in CL 46548153 : Reimport does not provide a destination path

#rnx
#DisabledValidations CheckCopyright, CheckBalancedMacros

diff --git a/Engine/Source/Runtime/Interchange/Engine/Private/InterchangeManager.cpp b/Engine/Source/Runtime/Interchange/Engine/Private/InterchangeManager.cpp
--- a/Engine/Source/Runtime/Interchange/Engine/Private/InterchangeManager.cpp
+++ b/Engine/Source/Runtime/Interchange/Engine/Private/InterchangeManager.cpp
@@ -1831,9 +1831,10 @@ UInterchangeManager::ImportInternal(const FString& ContentPath, const UInterchan
 		return TTuple<UE::Interchange::FAssetImportResultRef, UE::Interchange::FSceneImportResultRef>{ MakeShared< UE::Interchange::FImportResult, ESPMode::ThreadSafe >(), MakeShared< UE::Interchange::FImportResult, ESPMode::ThreadSafe >() };
 	}
 
-	if (!FPackageName::IsValidPath(ContentPath))
+	// Check on validity of destination content folder on import not reimport
+	if (!ImportAssetParameters.ReimportAsset && !FPackageName::IsValidPath(ContentPath))
 	{
-		UE_LOG(LogInterchangeEngine, Error, TEXT("Destination content path is not valid."));
+		UE_LOG(LogInterchangeEngine, Error, TEXT("Destination content path, %s, is not valid."), *ContentPath);
 		return TTuple<UE::Interchange::FAssetImportResultRef, UE::Interchange::FSceneImportResultRef>{ MakeShared< UE::Interchange::FImportResult, ESPMode::ThreadSafe >(), MakeShared< UE::Interchange::FImportResult, ESPMode::ThreadSafe >() };
 	}
 
