Fixed eventual crash on a level reimport targeting packed actors

#rnx

diff --git a/Engine/Source/Runtime/Interchange/Engine/Private/Tasks/InterchangeTaskCreateSceneObjects.cpp b/Engine/Source/Runtime/Interchange/Engine/Private/Tasks/InterchangeTaskCreateSceneObjects.cpp
--- a/Engine/Source/Runtime/Interchange/Engine/Private/Tasks/InterchangeTaskCreateSceneObjects.cpp
+++ b/Engine/Source/Runtime/Interchange/Engine/Private/Tasks/InterchangeTaskCreateSceneObjects.cpp
@@ -97,6 +97,26 @@ void UE::Interchange::FTaskCreateSceneObjects_GameThread::Execute()
 				{
 					if (UWorld* FactoryNodeWorld = Cast<UWorld>(LevelAssetPath.TryLoad()))
 					{
+#if WITH_EDITOR
+						// Because of the level instance's referenced world being edited then saved on reimports
+						// The referenced world's cleanup state get updated to a non zero value which prevents
+						// the registration of all the world's components.
+						// A temporary fix is to reset that state. This may have side effects. No found any so far.
+						// In any case, the creation and update of level instances and its associated level needs to be
+						// revisited to better align on how this is done in the editor in next release.
+						// #interchange_levelinstance_todo: Revisit this logic when reworking the handling of level instances
+						if (FactoryNodeWorld->IsCleanedUp())
+						{
+							UInterchangeResultWarning_Generic* Message = Factory->AddMessage<UInterchangeResultWarning_Generic>();
+							Message->SourceAssetName = AsyncHelper->SourceDatas[SourceIndex]->GetFilename();
+							Message->DestinationAssetName = SceneNodeName;
+							Message->AssetType = LevelFactoryNode->GetObjectClass();
+							Message->Text = NSLOCTEXT("FTaskCreateSceneObjects_GameThread", "DirtyWorld", "World has not been properly cleaned up. Reseting state.");
+							
+							FactoryNodeWorld->ResetCleanedUpState();
+						}
+#endif
+
 						if (FactoryNodeWorld->PersistentLevel)
 						{
 							FactoryImportWorld = FactoryNodeWorld;
