Interchange: Show an error mesage when the import/reimport failed to produce any actor or asset

diff --git a/Engine/Source/Runtime/Interchange/Engine/Private/Tasks/InterchangeTaskCompletion.cpp b/Engine/Source/Runtime/Interchange/Engine/Private/Tasks/InterchangeTaskCompletion.cpp
--- a/Engine/Source/Runtime/Interchange/Engine/Private/Tasks/InterchangeTaskCompletion.cpp
+++ b/Engine/Source/Runtime/Interchange/Engine/Private/Tasks/InterchangeTaskCompletion.cpp
@@ -210,10 +210,12 @@ void UE::Interchange::FTaskCompletion_GameThread::Execute()
 	//No need anymore of the translators sources
 	AsyncHelper->ReleaseTranslatorsSource();
 
+	bool bImportedSomething = false;
+
 	if (!AsyncHelper->bCancel)
 	{
 		//Broadcast OnAssetPostImport/OnAssetPostReimport for each imported asset
-		AsyncHelper->IterateImportedAssetsPerSourceIndex([AsyncHelper](int32 SourceIndex, const TArray<FImportAsyncHelper::FImportedObjectInfo>& AssetInfos)
+		AsyncHelper->IterateImportedAssetsPerSourceIndex([AsyncHelper, &bImportedSomething](int32 SourceIndex, const TArray<FImportAsyncHelper::FImportedObjectInfo>& AssetInfos)
 			{
 				for (const FImportAsyncHelper::FImportedObjectInfo& AssetInfo : AssetInfos)
 				{
@@ -229,10 +231,12 @@ void UE::Interchange::FTaskCompletion_GameThread::Execute()
 						{
 							//Notify the asset registry, only when we have created the asset
 							FAssetRegistryModule::AssetCreated(Asset);
+							bImportedSomething = true;
 						}
 						else if (AsyncHelper->TaskData.ReimportObject && AsyncHelper->TaskData.ReimportObject == Asset)
 						{
 							UInterchangeManager::GetInterchangeManager().OnAssetPostReimport.Broadcast(Asset);
+							bImportedSomething = true;
 						}
 						//We broadcast this event for both import and reimport.
 						UInterchangeManager::GetInterchangeManager().OnAssetPostImport.Broadcast(Asset);
@@ -258,12 +262,13 @@ void UE::Interchange::FTaskCompletion_GameThread::Execute()
 			});
 
 		//Iterate the Scene Actors
-		AsyncHelper->IterateImportedSceneObjectsPerSourceIndex([AsyncHelper](int32 SourceIndex, const TArray<FImportAsyncHelper::FImportedObjectInfo>& AssetInfos)
+		AsyncHelper->IterateImportedSceneObjectsPerSourceIndex([AsyncHelper, &bImportedSomething](int32 SourceIndex, const TArray<FImportAsyncHelper::FImportedObjectInfo>& AssetInfos)
 			{
 				for (const FImportAsyncHelper::FImportedObjectInfo& SceneObjectInfo : AssetInfos)
 				{
 					if (AActor* Actor = Cast<AActor>(SceneObjectInfo.ImportedObject.TryLoad()))
 					{
+						bImportedSomething = true;
 						for (UInterchangePipelineBase* PipelineBase : AsyncHelper->Pipelines)
 						{
 							PipelineBase->ScriptedExecutePostBroadcastPipeline(AsyncHelper->BaseNodeContainers[SourceIndex].Get()
@@ -312,6 +317,23 @@ void UE::Interchange::FTaskCompletion_GameThread::Execute()
 			});
 	}
 
+	if (!bImportedSomething && !AsyncHelper->bCancel)
+	{
+		for (TObjectPtr<UInterchangeTranslatorBase> Translator : AsyncHelper->Translators)
+		{
+			// We need a valid translator to add a result message
+			if (!Translator)
+			{
+				continue;
+			}
+
+			UInterchangeResultError_Generic* ErrorResult = NewObject<UInterchangeResultError_Generic>(GetTransientPackage(), UInterchangeResultError_Generic::StaticClass());
+			ErrorResult->Text = NSLOCTEXT("InterchangeTaskCompletion", "NothingImported", "There was nothing to import from the provided source data using the chosen pipeline options.");
+			ErrorResult->SourceAssetName = Translator->GetSourceData()->GetFilename();
+			Translator->AddMessage(ErrorResult);
+		}
+	}
+
 	AsyncHelper->AssetImportResult->SetDone();
 	AsyncHelper->SceneImportResult->SetDone();
 
diff --git a/Engine/Source/Runtime/Interchange/Engine/Private/Tasks/InterchangeTaskParsing.cpp b/Engine/Source/Runtime/Interchange/Engine/Private/Tasks/InterchangeTaskParsing.cpp
--- a/Engine/Source/Runtime/Interchange/Engine/Private/Tasks/InterchangeTaskParsing.cpp
+++ b/Engine/Source/Runtime/Interchange/Engine/Private/Tasks/InterchangeTaskParsing.cpp
@@ -395,23 +395,6 @@ void UE::Interchange::FTaskParsing::Execute()
 	TMap<FString, FTaskParsingRenameInfo> RenameAssets;
 	TSet<FString> CreatedTasksAssetNames; // Tracks for which asset name we have created a task so that we don't have 2 tasks for the same asset name
 
-	//If we find nothing to import, we notify the user
-	if (TaskDatas.Num() == 0)
-	{
-		for(TObjectPtr<UInterchangeTranslatorBase> Translator : AsyncHelper->Translators)
-		{
-			//We need a valid translator to add a result message
-			if (!Translator || AsyncHelper->bCancel)
-			{
-				continue;
-			}
-			UInterchangeResultError_Generic* ErrorResult = NewObject<UInterchangeResultError_Generic>(GetTransientPackage(), UInterchangeResultError_Generic::StaticClass());
-			ErrorResult->Text = NSLOCTEXT("InterchangeTaskParsing", "TranslatorNoAssetImported", "There was no data to import in the provided source data.");
-			ErrorResult->SourceAssetName = Translator->GetSourceData()->GetFilename();
-			Translator->AddMessage(ErrorResult);
-		}
-	}
-
 	TArray<uint64> CompletionPrerequistes;
 	const int32 PoolWorkerThreadCount = FTaskGraphInterface::Get().GetNumWorkerThreads() / 2;
 	const int32 MaxNumWorker = FMath::Max(PoolWorkerThreadCount, 1);
