Interchange : Encapsulate the way we Get/Set ImportedObject of FImportedObjectInfo struct

diff --git a/Engine/Plugins/Interchange/Runtime/Source/Import/Private/Scene/InterchangeLevelInstanceActorFactory.cpp b/Engine/Plugins/Interchange/Runtime/Source/Import/Private/Scene/InterchangeLevelInstanceActorFactory.cpp
--- a/Engine/Plugins/Interchange/Runtime/Source/Import/Private/Scene/InterchangeLevelInstanceActorFactory.cpp
+++ b/Engine/Plugins/Interchange/Runtime/Source/Import/Private/Scene/InterchangeLevelInstanceActorFactory.cpp
@@ -207,7 +207,7 @@ UObject* UInterchangeLevelInstanceActorFactory::ImportSceneObject_GameThread(con
 				MutableNodeContainer->SetupNode(BpFactoryNode, UniqueId, NewBP->GetName(), EInterchangeNodeContainerType::FactoryData);
 				BpFactoryNode->SetCustomReferenceObject(NewBP);
 
-				AssetInfo.ImportedObject = NewBP;
+				AssetInfo.SetImportedObject(NewBP);
 				AssetInfo.Factory = this;
 				AssetInfo.FactoryNode = BpFactoryNode;
 				AssetInfo.bIsReimport = false;
diff --git a/Engine/Source/Runtime/Interchange/Engine/Private/InterchangeManager.cpp b/Engine/Source/Runtime/Interchange/Engine/Private/InterchangeManager.cpp
--- a/Engine/Source/Runtime/Interchange/Engine/Private/InterchangeManager.cpp
+++ b/Engine/Source/Runtime/Interchange/Engine/Private/InterchangeManager.cpp
@@ -546,7 +546,7 @@ bool UE::Interchange::FImportAsyncHelper::IsImportingObject(UObject* Object) con
 		}
 		for (const FImportedObjectInfo& ImportedObjectInfo : ImportedObjects)
 		{
-			if (ImportedObjectInfo.ImportedObject == Object)
+			if (ImportedObjectInfo.GetImportedObject() == Object)
 			{
 				bFoundAsset = true;
 				break;
diff --git a/Engine/Source/Runtime/Interchange/Engine/Private/Tasks/InterchangeTaskCompletion.cpp b/Engine/Source/Runtime/Interchange/Engine/Private/Tasks/InterchangeTaskCompletion.cpp
--- a/Engine/Source/Runtime/Interchange/Engine/Private/Tasks/InterchangeTaskCompletion.cpp
+++ b/Engine/Source/Runtime/Interchange/Engine/Private/Tasks/InterchangeTaskCompletion.cpp
@@ -84,7 +84,7 @@ void UE::Interchange::FTaskPreCompletion_GameThread::Execute()
 		//First iteration to call SetupObject_GameThread and pipeline ExecutePostFactoryPipeline
 		for (const FImportAsyncHelper::FImportedObjectInfo& ObjectInfo : ImportedObjects)
 		{
-			UObject* ImportedObject = ObjectInfo.ImportedObject.TryLoad();
+			UObject* ImportedObject = ObjectInfo.GetImportedObject();
 			//In case Some factory code cannot run outside of the main thread we offer this callback to finish the work before calling post edit change (building the asset)
 			if (bCallPostImportGameThreadCallback && ObjectInfo.Factory)
 			{
@@ -148,7 +148,7 @@ void UE::Interchange::FTaskPreCompletion_GameThread::Execute()
 		//Second iteration to call BuildObject
 		for (const FImportAsyncHelper::FImportedObjectInfo& ObjectInfo : ImportedObjects)
 		{
-			UObject* ImportedObject = ObjectInfo.ImportedObject.TryLoad();
+			UObject* ImportedObject = ObjectInfo.GetImportedObject();
 			if (ImportedObject == nullptr || !IsValid(ImportedObject))
 			{
 				continue;
@@ -166,7 +166,7 @@ void UE::Interchange::FTaskPreCompletion_GameThread::Execute()
 		//Third iteration to register the assets
 		for (const FImportAsyncHelper::FImportedObjectInfo& ObjectInfo : ImportedObjects)
 		{
-			UObject* ImportedObject = ObjectInfo.ImportedObject.TryLoad();
+			UObject* ImportedObject = ObjectInfo.GetImportedObject();
 			if (ImportedObject == nullptr || !IsValid(ImportedObject))
 			{
 				continue;
@@ -219,7 +219,7 @@ void UE::Interchange::FTaskCompletion_GameThread::Execute()
 			{
 				for (const FImportAsyncHelper::FImportedObjectInfo& AssetInfo : AssetInfos)
 				{
-					if (UObject* Asset = AssetInfo.ImportedObject.TryLoad())
+					if (UObject* Asset = AssetInfo.GetImportedObject())
 					{
 						if (Asset->HasAnyFlags(RF_MirroredGarbage))
 						{
@@ -246,7 +246,7 @@ void UE::Interchange::FTaskCompletion_GameThread::Execute()
 				//Do a second pass for the post broadcast pipeline call
 				for (const FImportAsyncHelper::FImportedObjectInfo& AssetInfo : AssetInfos)
 				{
-					if (UObject* Asset = AssetInfo.ImportedObject.TryLoad())
+					if (UObject* Asset = AssetInfo.GetImportedObject())
 					{
 						for (UInterchangePipelineBase* PipelineBase : AsyncHelper->Pipelines)
 						{
@@ -266,7 +266,7 @@ void UE::Interchange::FTaskCompletion_GameThread::Execute()
 			{
 				for (const FImportAsyncHelper::FImportedObjectInfo& SceneObjectInfo : AssetInfos)
 				{
-					if (AActor* Actor = Cast<AActor>(SceneObjectInfo.ImportedObject.TryLoad()))
+					if (AActor* Actor = Cast<AActor>(SceneObjectInfo.GetImportedObject()))
 					{
 						bImportedSomething = true;
 						for (UInterchangePipelineBase* PipelineBase : AsyncHelper->Pipelines)
@@ -293,7 +293,7 @@ void UE::Interchange::FTaskCompletion_GameThread::Execute()
 					using namespace UE::Interchange::Private::ObjectDeletionUtils;
 
 					//Make any created asset go away
-					UObject* ObjectToPurge = AssetInfo.ImportedObject.TryLoad();
+					UObject* ObjectToPurge = AssetInfo.GetImportedObject();
 					PurgeObject(ObjectToPurge);
 				}
 			});
@@ -303,7 +303,7 @@ void UE::Interchange::FTaskCompletion_GameThread::Execute()
 			{
 				for (const FImportAsyncHelper::FImportedObjectInfo& SceneObjectInfo : AssetInfos)
 				{
-					if (AActor* Actor = Cast<AActor>(SceneObjectInfo.ImportedObject.TryLoad()))
+					if (AActor* Actor = Cast<AActor>(SceneObjectInfo.GetImportedObject()))
 					{
 						if (UWorld* ActorWorld = Actor->GetWorld())
 						{
diff --git a/Engine/Source/Runtime/Interchange/Engine/Private/Tasks/InterchangeTaskCreateSceneObjects.cpp b/Engine/Source/Runtime/Interchange/Engine/Private/Tasks/InterchangeTaskCreateSceneObjects.cpp
--- a/Engine/Source/Runtime/Interchange/Engine/Private/Tasks/InterchangeTaskCreateSceneObjects.cpp
+++ b/Engine/Source/Runtime/Interchange/Engine/Private/Tasks/InterchangeTaskCreateSceneObjects.cpp
@@ -56,10 +56,10 @@ void UE::Interchange::FTaskCreateSceneObjects_GameThread::Execute()
 		{
 			for (const UE::Interchange::FImportAsyncHelper::FImportedObjectInfo& Info : ImportedObjectInfos)
 			{
-				UObject* ImportObject = Info.ImportedObject.TryLoad();
-				if(ImportObject)
+				UObject* ImportObject = Info.GetImportedObject();
+				if(IsValid(ImportObject))
 				{
-					ImportObjects.Add(Info.ImportedObject.TryLoad());
+					ImportObjects.Add(ImportObject);
 				}
 			}
 		});
@@ -158,13 +158,13 @@ void UE::Interchange::FTaskCreateSceneObjects_GameThread::Execute()
 		{
 			const FImportAsyncHelper::FImportedObjectInfo* ImportedInfoPtr = AsyncHelper->FindImportedSceneObjects(SourceIndex, [SceneObject](const FImportAsyncHelper::FImportedObjectInfo& CurInfo)
 				{
-					return CurInfo.ImportedObject == SceneObject;
+					return CurInfo.GetImportedObject() == SceneObject;
 				});
 
 			if (!ImportedInfoPtr)
 			{
 				UE::Interchange::FImportAsyncHelper::FImportedObjectInfo& ObjectInfo = AsyncHelper->AddDefaultImportedSceneObjectGetRef(SourceIndex);
-				ObjectInfo.ImportedObject = SceneObject;
+				ObjectInfo.SetImportedObject(SceneObject);
 				ObjectInfo.Factory = Factory;
 				ObjectInfo.FactoryNode = FactoryNode;
 			}
diff --git a/Engine/Source/Runtime/Interchange/Engine/Private/Tasks/InterchangeTaskImportObject.cpp b/Engine/Source/Runtime/Interchange/Engine/Private/Tasks/InterchangeTaskImportObject.cpp
--- a/Engine/Source/Runtime/Interchange/Engine/Private/Tasks/InterchangeTaskImportObject.cpp
+++ b/Engine/Source/Runtime/Interchange/Engine/Private/Tasks/InterchangeTaskImportObject.cpp
@@ -613,7 +613,7 @@ void UE::Interchange::FTaskImportObject_GameThread::Execute()
 			}
 			FImportAsyncHelper::FImportedObjectInfo& AssetInfo = AsyncHelper->AddDefaultImportedAssetGetRef(SourceIndex);
 
-			AssetInfo.ImportedObject = ImportAssetResult.ImportedObject;
+			AssetInfo.SetImportedObject(ImportAssetResult.ImportedObject);
 			AssetInfo.Factory = Factory;
 			AssetInfo.FactoryNode = FactoryNode;
 			AssetInfo.bIsReimport = bool(ObjectToReimport != nullptr);
@@ -720,13 +720,13 @@ void UE::Interchange::FTaskImportObjectFinalize_GameThread::Execute()
 		{
 			const FImportAsyncHelper::FImportedObjectInfo* AssetInfoPtr = AsyncHelper->FindImportedAssets(SourceIndex, [&ImportAssetResult](const FImportAsyncHelper::FImportedObjectInfo& CurInfo)
 				{
-					return CurInfo.ImportedObject == ImportAssetResult.ImportedObject;
+					return CurInfo.GetImportedObject() == ImportAssetResult.ImportedObject;
 				});
 
 			if (!AssetInfoPtr)
 			{
 				FImportAsyncHelper::FImportedObjectInfo& AssetInfo = AsyncHelper->AddDefaultImportedAssetGetRef(SourceIndex);
-				AssetInfo.ImportedObject = ImportAssetResult.ImportedObject;
+				AssetInfo.SetImportedObject(ImportAssetResult.ImportedObject);
 				AssetInfo.Factory = Factory;
 				AssetInfo.FactoryNode = FactoryNode;
 
diff --git a/Engine/Source/Runtime/Interchange/Engine/Private/Tasks/InterchangeTaskPipeline.cpp b/Engine/Source/Runtime/Interchange/Engine/Private/Tasks/InterchangeTaskPipeline.cpp
--- a/Engine/Source/Runtime/Interchange/Engine/Private/Tasks/InterchangeTaskPipeline.cpp
+++ b/Engine/Source/Runtime/Interchange/Engine/Private/Tasks/InterchangeTaskPipeline.cpp
@@ -85,7 +85,7 @@ void UE::Interchange::FTaskWaitAssetCompilation_GameThread::Execute()
 			ImportedObjects.Reserve(ImportedObjects.Num() + ImportedInfos.Num());
 			for (const UE::Interchange::FImportAsyncHelper::FImportedObjectInfo& ImportedInfo : ImportedInfos)
 			{
-				ImportedObjects.Add(ImportedInfo.ImportedObject.TryLoad());
+				ImportedObjects.Add(ImportedInfo.GetImportedObject());
 			}
 		};
 
@@ -155,7 +155,7 @@ void UE::Interchange::FTaskPostImport_GameThread::Execute()
 		{
 			for (const UE::Interchange::FImportAsyncHelper::FImportedObjectInfo& ObjectInfo : ImportedInfos)
 			{
-				UObject* ImportedObject = ObjectInfo.ImportedObject.TryLoad();
+				UObject* ImportedObject = ObjectInfo.GetImportedObject();
 				if (ImportedObject && !ObjectInfo.bPostEditChangeCalled)
 				{
 					ObjectInfo.bPostEditChangeCalled = true;
@@ -192,7 +192,7 @@ void UE::Interchange::FTaskPostImport_GameThread::Execute()
 
 			for (const UE::Interchange::FImportAsyncHelper::FImportedObjectInfo& ObjectInfo : ImportedInfos)
 			{
-				UObject* ImportedObject = ObjectInfo.ImportedObject.TryLoad();
+				UObject* ImportedObject = ObjectInfo.GetImportedObject();
 
 				//In case Some factory code cannot run outside of the main thread we offer this callback to finish the work after asset build is finish.
 				if (bCallPostImportGameThreadCallback && ObjectInfo.Factory)
@@ -210,7 +210,7 @@ void UE::Interchange::FTaskPostImport_GameThread::Execute()
 				}
 
 				NodeUniqueIDs.Add(ObjectInfo.FactoryNode->GetUniqueID());
-				ImportedObjects.Add(ObjectInfo.ImportedObject.TryLoad());
+				ImportedObjects.Add(ObjectInfo.GetImportedObject());
 				IsAssetsReimported.Add(ObjectInfo.bIsReimport);
 			}
 		};
diff --git a/Engine/Source/Runtime/Interchange/Engine/Public/InterchangeManager.h b/Engine/Source/Runtime/Interchange/Engine/Public/InterchangeManager.h
--- a/Engine/Source/Runtime/Interchange/Engine/Public/InterchangeManager.h
+++ b/Engine/Source/Runtime/Interchange/Engine/Public/InterchangeManager.h
@@ -285,7 +285,23 @@ namespace UE
 
 			struct FImportedObjectInfo
 			{
+				UE_DEPRECATED(5.8, "Direct use of this variable is deprecated, Please use GetImportedObject()/SetImportedObject() instead")
 				FSoftObjectPath ImportedObject; // The object that was imported.
+
+				UObject* GetImportedObject() const 
+				{ 
+					PRAGMA_DISABLE_DEPRECATION_WARNINGS
+					return ImportedObject.TryLoad(); 
+					PRAGMA_ENABLE_DEPRECATION_WARNINGS
+				}
+
+				void SetImportedObject(UObject* InImportedObject) 
+				{ 
+					PRAGMA_DISABLE_DEPRECATION_WARNINGS
+					ImportedObject = InImportedObject; 
+					PRAGMA_ENABLE_DEPRECATION_WARNINGS
+				}
+
 				UInterchangeFactoryBase* Factory = nullptr; //The factory that created the imported object.
 				UInterchangeFactoryBaseNode* FactoryNode = nullptr; //The node that describes the object.
 				bool bIsReimport;
