Fixed issues reported by BugHawk

#rnx
#Synced-CL 48455087

diff --git a/Engine/Source/Runtime/Interchange/Core/Private/Types/AttributeStorage.cpp b/Engine/Source/Runtime/Interchange/Core/Private/Types/AttributeStorage.cpp
--- a/Engine/Source/Runtime/Interchange/Core/Private/Types/AttributeStorage.cpp
+++ b/Engine/Source/Runtime/Interchange/Core/Private/Types/AttributeStorage.cpp
@@ -1439,7 +1439,11 @@ namespace UE
 		F##MathType BaseTypeValue; \
 		ensure(GetAttributeValue(*this, BaseTypeValue)); \
 		F##MathType VersionTypeValue; \
-		ensure(GetAttributeValue(VersionStorage, VersionTypeValue)); \
+		bool bGetVersionValue = GetAttributeValue(VersionStorage, VersionTypeValue); \
+		if(!bGetVersionValue) \
+		{ \
+			return false; \
+		} \
 		return BaseTypeValue.Equals(VersionTypeValue); \
 	}
 
@@ -1449,8 +1453,12 @@ namespace UE
 		TArray<F##MathType> BaseArray; \
 		ensure(GetAttributeValue(*this, BaseArray)); \
 		TArray<F##MathType> VersionArray; \
-		ensure(GetAttributeValue(VersionStorage, VersionArray)); \
-		for(int32 Index = 0; Index < BaseArray.Num(); ++Index) \
+		bool bGetVersionValue = GetAttributeValue(VersionStorage, VersionArray); \
+		if(!bGetVersionValue || BaseArray.Num() != VersionArray.Num()) \
+		{ \
+			return false; \
+		} \
+		for (int32 Index = 0; Index < BaseArray.Num(); ++Index) \
 		{ \
 			if(!BaseArray[Index].Equals(VersionArray[Index])) \
 			{ \
@@ -1464,13 +1472,21 @@ namespace UE
 	F##ValueType Base##ValueType; \
 	ensure(GetAttributeValue(*this, Base##ValueType)); \
 	F##ValueType Version##ValueType; \
-	ensure(GetAttributeValue(VersionStorage, Version##ValueType));
+	bool bGetVersionValue = GetAttributeValue(VersionStorage, Version##ValueType); \
+	if (!bGetVersionValue) \
+	{ \
+		return false; \
+	}
 
 #define DECLARE_LOCAL_ARRAY_VARIABLES( ValueType ) \
 	TArray<F##ValueType> Base##ValueType##s; \
 	ensure(GetAttributeValue(*this, Base##ValueType##s)); \
 	TArray<F##ValueType> Version##ValueType##s; \
-	ensure(GetAttributeValue(VersionStorage, Version##ValueType##s));
+	bool bGetVersionValue = GetAttributeValue(VersionStorage, Version##ValueType##s); \
+	if (!bGetVersionValue || Base##ValueType##s.Num() != Version##ValueType##s.Num()) \
+	{ \
+		return false; \
+	}
 
 		bool FAttributeStorage::AreAllocationInfosEqual(const FAttributeKey& BaseKey, const FAttributeAllocationInfo& BaseInfo, const FAttributeStorage& VersionStorage, const FAttributeAllocationInfo& VersionInfo) const
 		{
@@ -1613,6 +1629,7 @@ namespace UE
 					bool bIsEqual = BaseOrientedBox.AxisX.Equals(VersionOrientedBox.AxisX);
 					bIsEqual &= BaseOrientedBox.AxisY.Equals(VersionOrientedBox.AxisY);
 					bIsEqual &= BaseOrientedBox.AxisZ.Equals(VersionOrientedBox.AxisZ);
+					bIsEqual &= BaseOrientedBox.Center.Equals(VersionOrientedBox.Center);
 					bIsEqual &= FMath::IsNearlyEqual(BaseOrientedBox.ExtentX, VersionOrientedBox.ExtentX);
 					bIsEqual &= FMath::IsNearlyEqual(BaseOrientedBox.ExtentY, VersionOrientedBox.ExtentY);
 					bIsEqual &= FMath::IsNearlyEqual(BaseOrientedBox.ExtentZ, VersionOrientedBox.ExtentZ);
@@ -1624,8 +1641,8 @@ namespace UE
 				{
 					DECLARE_LOCAL_VARIABLES(TwoVectors)
 
-					bool bIsEqual = BaseTwoVectors.v1.Equals(BaseTwoVectors.v1);
-					bIsEqual &= BaseTwoVectors.v2.Equals(BaseTwoVectors.v2);
+					bool bIsEqual = BaseTwoVectors.v1.Equals(VersionTwoVectors.v1);
+					bIsEqual &= BaseTwoVectors.v2.Equals(VersionTwoVectors.v2);
 
 					return bIsEqual;
 				}
@@ -1699,7 +1716,12 @@ namespace UE
 					float BaseFloat;
 					ensure(GetAttributeValue(*this, BaseFloat));
 					float VersionFloat;
-					ensure(GetAttributeValue(VersionStorage, VersionFloat));
+					bool bVersionValue = GetAttributeValue(VersionStorage, VersionFloat);
+					
+					if (!bVersionValue)
+					{
+						return false;
+					}
 
 					return FMath::IsNearlyEqual(BaseFloat, VersionFloat);
 				}
@@ -1709,7 +1731,12 @@ namespace UE
 					TArray<float> BaseFloats;
 					ensure(GetAttributeValue(*this, BaseFloats));
 					TArray<float> VersionFloats;
-					ensure(GetAttributeValue(VersionStorage, VersionFloats));
+					bool bVersionValue = GetAttributeValue(VersionStorage, VersionFloats);
+
+					if (!bVersionValue || BaseFloats.Num() != VersionFloats.Num())
+					{
+						return false;
+					}
 
 					for (int32 Index = 0; Index < BaseFloats.Num(); ++Index)
 					{
@@ -1727,7 +1754,12 @@ namespace UE
 					double BaseDouble;
 					ensure(GetAttributeValue(*this, BaseDouble));
 					double VersionDouble;
-					ensure(GetAttributeValue(VersionStorage, VersionDouble));
+					bool bVersionValue = GetAttributeValue(VersionStorage, VersionDouble);
+
+					if (!bVersionValue)
+					{
+						return false;
+					}
 
 					return FMath::IsNearlyEqual(BaseDouble, VersionDouble);
 				}
@@ -1737,7 +1769,12 @@ namespace UE
 					TArray<double> BaseDoubles;
 					ensure(GetAttributeValue(*this, BaseDoubles));
 					TArray<double> VersionDoubles;
-					ensure(GetAttributeValue(VersionStorage, VersionDoubles));
+					bool bVersionValue = GetAttributeValue(VersionStorage, VersionDoubles);
+
+					if (!bVersionValue || BaseDoubles.Num() != VersionDoubles.Num())
+					{
+						return false;
+					}
 
 					for (int32 Index = 0; Index < BaseDoubles.Num(); ++Index)
 					{
