[Interchange Analytics] Analytics Rework to make the Analytics Identifiable across all the events.
- Added GUID to the Import.
- Refactored the code that would allow us to override the category name for analytics of the properties on the pipelines.

#rnx
#Synced-CL 49210907

diff --git a/Engine/Source/Runtime/Interchange/Core/InterchangeCore.Build.cs b/Engine/Source/Runtime/Interchange/Core/InterchangeCore.Build.cs
--- a/Engine/Source/Runtime/Interchange/Core/InterchangeCore.Build.cs
+++ b/Engine/Source/Runtime/Interchange/Core/InterchangeCore.Build.cs
@@ -18,6 +18,7 @@ public class InterchangeCore : ModuleRules
 			new string[]
 			{
 				"SlateCore",
+				"Analytics"
 			}
 		);
 	}
diff --git a/Engine/Source/Runtime/Interchange/Core/Private/InterchangePipelineBase.cpp b/Engine/Source/Runtime/Interchange/Core/Private/InterchangePipelineBase.cpp
--- a/Engine/Source/Runtime/Interchange/Core/Private/InterchangePipelineBase.cpp
+++ b/Engine/Source/Runtime/Interchange/Core/Private/InterchangePipelineBase.cpp
@@ -7,6 +7,7 @@
 #include "Misc/ConfigCacheIni.h"
 #include "UObject/Object.h"
 #include "UObject/ObjectMacros.h"
+#include "AnalyticsEventAttribute.h"
 
 #include UE_INLINE_GENERATED_CPP_BY_NAME(InterchangePipelineBase)
 
@@ -180,6 +181,97 @@ FName UInterchangePipelineBase::GetResultsPropertyName()
 	return GET_MEMBER_NAME_CHECKED(UInterchangePipelineBase, Results);
 }
 
+bool UInterchangePipelineBase::CanAddPropertyToAnalytics(FProperty* Property) const
+{
+	if (!ensure(Property))
+	{
+		return false;
+	}
+
+	if (Property->HasAnyPropertyFlags(CPF_Transient))
+	{
+		return false;
+	}
+
+	if (Property->GetFName() == UInterchangePipelineBase::GetPropertiesStatesPropertyName())
+	{
+		return false;
+	}
+
+	if (Property->GetFName() == UInterchangePipelineBase::GetResultsPropertyName())
+	{
+		return false;
+	}
+
+#if WITH_EDITORONLY_DATA
+	FObjectProperty* SubObject = CastField<FObjectProperty>(Property);
+	UInterchangePipelineBase* SubPipeline = SubObject 
+												? Cast<UInterchangePipelineBase>(SubObject->GetObjectPropertyValue_InContainer(this)) 
+												: nullptr;
+	if (!SubPipeline && Property->GetMetaData("Category").IsEmpty())
+	{
+		//In Editor do not add property with no category
+		return false;
+	}
+#endif
+
+	return true;
+}
+
+void UInterchangePipelineBase::AddPropertyAnalytics(FProperty* Property, TArray<FAnalyticsEventAttribute>& PipelineAttributes)
+{
+	if (!ensure(Property))
+	{
+		return;
+	}
+
+	const FString PropertyName = Property->GetName();
+	constexpr int PortFlags = 0;
+
+	if (FArrayProperty* Array = CastField<FArrayProperty>(Property))
+	{
+		FScriptArrayHelper_InContainer ArrayHelper(Array, this);
+		for (int32 i = 0; i < ArrayHelper.Num(); i++)
+		{
+			FString	Buffer;
+			Array->Inner->ExportTextItem_Direct(Buffer, ArrayHelper.GetRawPtr(i), ArrayHelper.GetRawPtr(i), this, PortFlags);
+
+			PipelineAttributes.Add(FAnalyticsEventAttribute(FString::Printf(TEXT("Property%s.%s"), *GetCategoryNameForAnalytics(*Property), *PropertyName), Buffer));
+		}
+	}
+	else
+	{
+		for (int32 Index = 0; Index < Property->ArrayDim; Index++)
+		{
+			FString PropertyIndexName = PropertyName;
+			if (Property->ArrayDim != 1)
+			{
+				PropertyIndexName += TEXT("[") + FString::FromInt(Index) + TEXT("]");
+			}
+
+			FString	Value;
+			Property->ExportText_InContainer(Index, Value, this, this, this, PortFlags);
+			PipelineAttributes.Add(FAnalyticsEventAttribute(FString::Printf(TEXT("Property%s.%s"), *GetCategoryNameForAnalytics(*Property), *PropertyIndexName), Value));
+		}
+	}
+}
+
+FString UInterchangePipelineBase::GetCategoryNameForAnalytics(const FProperty& Property) const
+{
+	FString CategoryName = FString();
+
+#if WITH_EDITORONLY_DATA
+	CategoryName = Property.GetMetaData("Category");
+	CategoryName.ReplaceCharInline(TEXT('.'), TEXT('_'));
+	CategoryName.RemoveSpacesInline();
+	CategoryName = TEXT(".") + CategoryName;
+#endif
+
+	return CategoryName;
+}
+
+
+
 void UInterchangePipelineBase::LoadSettingsInternal(const FName PipelineStackName, const FString& ConfigFilename, TMap<FName, FInterchangePipelinePropertyStates>& ParentPropertiesStates, bool bResetPreDialog, bool& bOutRequiresSaving)
 {
 	int32 PortFlags = 0;
diff --git a/Engine/Source/Runtime/Interchange/Core/Public/InterchangePipelineBase.h b/Engine/Source/Runtime/Interchange/Core/Public/InterchangePipelineBase.h
--- a/Engine/Source/Runtime/Interchange/Core/Public/InterchangePipelineBase.h
+++ b/Engine/Source/Runtime/Interchange/Core/Public/InterchangePipelineBase.h
@@ -23,11 +23,14 @@
 
 #include "InterchangePipelineBase.generated.h"
 
+#define UE_API INTERCHANGECORE_API
+
 class FText;
 class UClass;
 class UInterchangeBaseNodeContainer;
 class UInterchangeResult;
 class UInterchangeSourceData;
+struct FAnalyticsEventAttribute;
 struct FFrame;
 
 UENUM(BlueprintType)
@@ -234,7 +237,7 @@ public:
 	 * FString PipelineDisplayName;
 	 */
 	UFUNCTION(BlueprintNativeEvent, BlueprintCallable, Category = "Interchange | Pipeline")
-	INTERCHANGECORE_API FString ScriptedGetPipelineDisplayName() const;
+	UE_API FString ScriptedGetPipelineDisplayName() const;
 	/** The default implementation (call if the blueprint do not have any implementation) will call the virtual ExecutePipeline */
 	FString ScriptedGetPipelineDisplayName_Implementation() const
 	{
@@ -249,7 +252,7 @@ public:
 	 * @note - the FTaskPipeline is calling this function not the virtual one that is call by the default implementation.
 	 */
 	UFUNCTION(BlueprintNativeEvent, BlueprintCallable, Category = "Interchange | Pipeline")
-	INTERCHANGECORE_API void ScriptedExecutePipeline(UInterchangeBaseNodeContainer* BaseNodeContainer, const TArray<UInterchangeSourceData*>& SourceDatas, const FString& ContentBasePath);
+	UE_API void ScriptedExecutePipeline(UInterchangeBaseNodeContainer* BaseNodeContainer, const TArray<UInterchangeSourceData*>& SourceDatas, const FString& ContentBasePath);
 	/** The default implementation, which is called if the Blueprint does not have any implementation, calls the virtual ExecutePipeline(). */
 	void ScriptedExecutePipeline_Implementation(UInterchangeBaseNodeContainer* BaseNodeContainer, const TArray<UInterchangeSourceData*>& SourceDatas, const FString& ContentBasePath)
 	{
@@ -262,7 +265,7 @@ public:
 	 * @note - The FTaskPreCompletion task calls this function, not the virtual one that is called by the default implementation.
 	 */
 	UFUNCTION(BlueprintNativeEvent, BlueprintCallable, Category = "Interchange | Pipeline")
-	INTERCHANGECORE_API void ScriptedExecutePostFactoryPipeline(const UInterchangeBaseNodeContainer* BaseNodeContainer, const FString& FactoryNodeKey, UObject* CreatedAsset, bool bIsAReimport);
+	UE_API void ScriptedExecutePostFactoryPipeline(const UInterchangeBaseNodeContainer* BaseNodeContainer, const FString& FactoryNodeKey, UObject* CreatedAsset, bool bIsAReimport);
 	/** The default implementation, which is called if the Blueprint does not have any implementation, calls the virtual ExecutePostFactoryPipeline(). */
 	void ScriptedExecutePostFactoryPipeline_Implementation(const UInterchangeBaseNodeContainer* BaseNodeContainer, const FString& FactoryNodeKey, UObject* CreatedAsset, bool bIsAReimport)
 	{
@@ -276,7 +279,7 @@ public:
 	 * @note - the FTaskPostImport calls this function not the virtual one that is call by the default implementation.
 	 */
 	UFUNCTION(BlueprintNativeEvent, BlueprintCallable, Category = "Interchange | Pipeline")
-	INTERCHANGECORE_API void ScriptedExecutePostImportPipeline(const UInterchangeBaseNodeContainer* BaseNodeContainer, const FString& FactoryNodeKey, UObject* CreatedAsset, bool bIsAReimport);
+	UE_API void ScriptedExecutePostImportPipeline(const UInterchangeBaseNodeContainer* BaseNodeContainer, const FString& FactoryNodeKey, UObject* CreatedAsset, bool bIsAReimport);
 	/** The default implementation, which is called if the Blueprint does not have any implementation, calls the virtual ExecutePostImportPipeline(). */
 	void ScriptedExecutePostImportPipeline_Implementation(const UInterchangeBaseNodeContainer* BaseNodeContainer, const FString& FactoryNodeKey, UObject* CreatedAsset, bool bIsAReimport)
 	{
@@ -290,7 +293,7 @@ public:
 	 * @note - the FTaskCompletion_GameThread calls this function not the virtual one that is call by the default implementation.
 	 */
 	UFUNCTION(BlueprintNativeEvent, BlueprintCallable, Category = "Interchange | Pipeline")
-	INTERCHANGECORE_API void ScriptedExecutePostBroadcastPipeline(const UInterchangeBaseNodeContainer* BaseNodeContainer, const FString& FactoryNodeKey, UObject* CreatedAsset, bool bIsAReimport);
+	UE_API void ScriptedExecutePostBroadcastPipeline(const UInterchangeBaseNodeContainer* BaseNodeContainer, const FString& FactoryNodeKey, UObject* CreatedAsset, bool bIsAReimport);
 	/** The default implementation, which is called if the Blueprint does not have any implementation, calls the virtual ExecutePostBroadcastPipeline(). */
 	void ScriptedExecutePostBroadcastPipeline_Implementation(const UInterchangeBaseNodeContainer* BaseNodeContainer, const FString& FactoryNodeKey, UObject* CreatedAsset, bool bIsAReimport)
 	{
@@ -303,7 +306,7 @@ public:
 	 * The Interchange manager calls this function, not the virtual one that is called by the default implementation.
 	 */
 	UFUNCTION(BlueprintCallable, BlueprintNativeEvent, Category = "Interchange | Pipeline")
-	INTERCHANGECORE_API void ScriptedExecuteExportPipeline(UInterchangeBaseNodeContainer* BaseNodeContainer);
+	UE_API void ScriptedExecuteExportPipeline(UInterchangeBaseNodeContainer* BaseNodeContainer);
 
 	/** The default implementation, which is called if the Blueprint does not have any implementation, calls the virtual ExecuteExportPipeline(). */
 	void ScriptedExecuteExportPipeline_Implementation(UInterchangeBaseNodeContainer* BaseNodeContainer)
@@ -333,7 +336,7 @@ public:
 	 * the Interchange framework calls this function, not the virtual one that is called by the default implementation.
 	 */
 	UFUNCTION(BlueprintNativeEvent, BlueprintCallable, Category = "Interchange | Pipeline")
-	INTERCHANGECORE_API void ScriptedSetReimportSourceIndex(UClass* ReimportObjectClass, const int32 SourceFileIndex);
+	UE_API void ScriptedSetReimportSourceIndex(UClass* ReimportObjectClass, const int32 SourceFileIndex);
 
 	/** The default implementation, which is called if the Blueprint does not have any implementation, calls the virtual SetReimportContentFromSourceIndex(). */
 	void ScriptedSetReimportSourceIndex_Implementation(UClass* ReimportObjectClass, const int32 SourceFileIndex)
@@ -353,11 +356,11 @@ public:
 		return true;
 	}
 
-	INTERCHANGECORE_API static FString GetDefaultConfigFileName();
+	UE_API static FString GetDefaultConfigFileName();
 	
-	INTERCHANGECORE_API void LoadSettings(const FName PipelineStackName, bool bResetPreDialog = false);
+	UE_API void LoadSettings(const FName PipelineStackName, bool bResetPreDialog = false);
 
-	INTERCHANGECORE_API void SaveSettings(const FName PipelineStackName);
+	UE_API void SaveSettings(const FName PipelineStackName);
 
 	/**
 	 * This function is called before we show the pipeline dialog. Pipelines that override it can change the existing settings according to the reimport type.
@@ -366,11 +369,11 @@ public:
 	 * @Note - The function will set the context of the pipeline.
 	 * @Param ContextParams - Give all the context information to the pipeline for the current import.
 	 */
-	INTERCHANGECORE_API virtual void AdjustSettingsForContext(const FInterchangePipelineContextParams& ContextParams);
-	INTERCHANGECORE_API virtual void AdjustSettingsFromCache();
+	UE_API virtual void AdjustSettingsForContext(const FInterchangePipelineContextParams& ContextParams);
+	UE_API virtual void AdjustSettingsFromCache();
 
 	/** Transfer the source pipeline adjust settings to this pipeline. */
-	INTERCHANGECORE_API void TransferAdjustSettings(const UInterchangePipelineBase* SourcePipeline);
+	UE_API void TransferAdjustSettings(const UInterchangePipelineBase* SourcePipeline);
 
 	void SetShowEssentialsMode(bool bShowEssentialsModeValue)
 	{
@@ -476,31 +479,31 @@ public:
 	/**
 	 * Return a const property states pointer. Return nullptr if the property does not exist.
 	 */
-	INTERCHANGECORE_API const FInterchangePipelinePropertyStates* GetPropertyStates(const FName PropertyPath) const;
+	UE_API const FInterchangePipelinePropertyStates* GetPropertyStates(const FName PropertyPath) const;
 
 	/**
 	 * Return a mutable property states pointer. Return nullptr if the property does not exist.
 	 */
-	INTERCHANGECORE_API FInterchangePipelinePropertyStates* GetMutablePropertyStates(const FName PropertyPath);
+	UE_API FInterchangePipelinePropertyStates* GetMutablePropertyStates(const FName PropertyPath);
 
 	/**
 	 * Return true if the property has valid states, or false if no states were set for the property.
 	 */
 	UFUNCTION(BlueprintCallable, Category = "Interchange | Pipeline")
-	INTERCHANGECORE_API bool DoesPropertyStatesExist(const FName PropertyPath) const;
+	UE_API bool DoesPropertyStatesExist(const FName PropertyPath) const;
 
 	/**
 	 * Return a mutable property states reference. Add the property states if it doesn't exist.
 	 */
 	UFUNCTION(BlueprintCallable, Category = "Interchange | Pipeline")
-	INTERCHANGECORE_API FInterchangePipelinePropertyStates& FindOrAddPropertyStates(const FName PropertyPath);
+	UE_API FInterchangePipelinePropertyStates& FindOrAddPropertyStates(const FName PropertyPath);
 
 	/**
 	 * Return the property name of the properties states map.
 	 */
-	static INTERCHANGECORE_API FName GetPropertiesStatesPropertyName();
+	static UE_API FName GetPropertiesStatesPropertyName();
 
-	static INTERCHANGECORE_API FName GetResultsPropertyName();
+	static UE_API FName GetResultsPropertyName();
 
 	bool CanEditPropertiesStates() { return bAllowPropertyStatesEdition; }
 
@@ -536,11 +539,11 @@ public:
 	* Update/reset the TWeakObjectPtrs to point to their original TObjectPtr.
 	* Works based on Variable naming.
 	*/
-	INTERCHANGECORE_API void UpdateWeakObjectPtrs();
+	UE_API void UpdateWeakObjectPtrs();
 
 	/** Begin UObject overrides*/
-	INTERCHANGECORE_API virtual void PostDuplicate(bool bDuplicateForPIE) override;
-	INTERCHANGECORE_API virtual void Serialize(FArchive& Ar) override;
+	UE_API virtual void PostDuplicate(bool bDuplicateForPIE) override;
+	UE_API virtual void Serialize(FArchive& Ar) override;
 	/** End UObject overrides */
 
 	/*
@@ -573,11 +576,17 @@ public:
 	FSoftObjectPath OriginalPipelinePath;
 #endif
 
-	INTERCHANGECORE_API UInterchangePipelineBase* GetMostPipelineOuter() const;
+	UE_API UInterchangePipelineBase* GetMostPipelineOuter() const;
+
+	UE_API virtual bool CanAddPropertyToAnalytics(FProperty* Property) const;
+
+	UE_API virtual void AddPropertyAnalytics(FProperty* Property, TArray<FAnalyticsEventAttribute>& PipelineAttributes);
 
 protected:
 
-	INTERCHANGECORE_API virtual FString GetPipelineDisplayName() const;
+	UE_API virtual FString GetCategoryNameForAnalytics(const FProperty& Property) const;
+
+	UE_API virtual FString GetPipelineDisplayName() const;
 
 	virtual void ExecutePipeline(UInterchangeBaseNodeContainer* BaseNodeContainer, const TArray<UInterchangeSourceData*>& SourceDatas, const FString& ContentBasePath)
 	{
@@ -617,15 +626,15 @@ protected:
 	{
 	}
 
-	INTERCHANGECORE_API void LoadSettingsInternal(const FName PipelineStackName, const FString& ConfigFilename, TMap<FName, FInterchangePipelinePropertyStates>& ParentPropertiesStates, bool bResetPreDialog, bool& bOutRequiresSaving);
+	UE_API void LoadSettingsInternal(const FName PipelineStackName, const FString& ConfigFilename, TMap<FName, FInterchangePipelinePropertyStates>& ParentPropertiesStates, bool bResetPreDialog, bool& bOutRequiresSaving);
 
-	INTERCHANGECORE_API void SaveSettingsInternal(const FName PipelineStackName, const FString& ConfigFilename);
+	UE_API void SaveSettingsInternal(const FName PipelineStackName, const FString& ConfigFilename);
 
 #if WITH_EDITOR
-	static INTERCHANGECORE_API void InternalToggleVisibilityPropertiesOfMetaDataValue(UInterchangePipelineBase* OuterMostPipeline, UInterchangePipelineBase* Pipeline, bool bDoTransientSubPipeline, const FString& MetaDataKey, const FString& MetaDataValue, const bool bVisibilityState);
-	static INTERCHANGECORE_API void HidePropertiesOfCategory(UInterchangePipelineBase* OuterMostPipeline, UInterchangePipelineBase* Pipeline, const FString& HideCategoryName, bool bDoTransientSubPipeline = false);
-	static INTERCHANGECORE_API void HidePropertiesOfSubCategory(UInterchangePipelineBase* OuterMostPipeline, UInterchangePipelineBase* Pipeline, const FString& HideSubCategoryName, bool bDoTransientSubPipeline = false);
-	static INTERCHANGECORE_API void HideProperty(UInterchangePipelineBase* OuterMostPipeline, UInterchangePipelineBase* Pipeline, const FName& HidePropertyName);
+	static UE_API void InternalToggleVisibilityPropertiesOfMetaDataValue(UInterchangePipelineBase* OuterMostPipeline, UInterchangePipelineBase* Pipeline, bool bDoTransientSubPipeline, const FString& MetaDataKey, const FString& MetaDataValue, const bool bVisibilityState);
+	static UE_API void HidePropertiesOfCategory(UInterchangePipelineBase* OuterMostPipeline, UInterchangePipelineBase* Pipeline, const FString& HideCategoryName, bool bDoTransientSubPipeline = false);
+	static UE_API void HidePropertiesOfSubCategory(UInterchangePipelineBase* OuterMostPipeline, UInterchangePipelineBase* Pipeline, const FString& HideSubCategoryName, bool bDoTransientSubPipeline = false);
+	static UE_API void HideProperty(UInterchangePipelineBase* OuterMostPipeline, UInterchangePipelineBase* Pipeline, const FName& HidePropertyName);
 #endif //WITH_EDITOR
 	
 	/**
@@ -681,3 +690,5 @@ protected:
 
 	TArray<FInterchangeConflictInfo> ConflictInfos; 
 };
+
+#undef UE_API
\ No newline at end of file
diff --git a/Engine/Source/Runtime/Interchange/Engine/Private/InterchangeAnalyticsHelper.cpp b/Engine/Source/Runtime/Interchange/Engine/Private/InterchangeAnalyticsHelper.cpp
--- a/Engine/Source/Runtime/Interchange/Engine/Private/InterchangeAnalyticsHelper.cpp
+++ b/Engine/Source/Runtime/Interchange/Engine/Private/InterchangeAnalyticsHelper.cpp
@@ -1,11 +1,51 @@
 // Copyright Epic Games, Inc. All Rights Reserved.
 
 #include "InterchangeAnalyticsHelper.h"
-#include "AnalyticsEventAttribute.h"
 #include "EngineAnalytics.h"
 
 namespace UE::Interchange
 {
+	FString FAnalyticsHelper::GetDefaultAnalyticsCategoryName(EInterchangeAnalyticsDefaultCategory DefaultCategory)
+	{	
+		static const FString CategoryNames[(size_t)EInterchangeAnalyticsDefaultCategory::InterchangeAnalyticsCategory_COUNT]
+		{
+		#define Category(Name) TEXT("."#Name)
+
+			TEXT(""), // Not Specified
+			Category(Models),
+			Category(Scenes),
+			Category(Materials),
+			Category(Textures),
+			Category(Sounds),
+			Category(Meshes),
+			Category(StaticMeshes),
+			Category(SkeletalMeshesAndAnimations),
+			Category(GeometryCaches),
+			Category(Grooms),
+
+		#undef Category
+		};
+
+		if (DefaultCategory > EInterchangeAnalyticsDefaultCategory::NotSpecified
+			&& DefaultCategory < EInterchangeAnalyticsDefaultCategory::InterchangeAnalyticsCategory_COUNT)
+		{
+			return CategoryNames[(size_t)DefaultCategory];
+		}
+
+		return CategoryNames[(size_t)EInterchangeAnalyticsDefaultCategory::NotSpecified];
+	}
+
+	FAnalyticsEventAttribute FAnalyticsHelper::CreateAnalyticsIDAttribute(FGuid AnalyticsID)
+	{
+		ensure(AnalyticsID.IsValid());
+		return FAnalyticsEventAttribute(TEXT("InterchangeAnalyticsGUID"), AnalyticsID.ToString());
+	}
+
+	FAnalyticsHelper::FAnalyticsHelper(const FGuid& InAnalyticsID)
+		: AnalyticsID(InAnalyticsID)
+	{
+	}
+
 	void FAnalyticsHelper::AppendThreadSafe(const FString& Identifier, const TArray<FAnalyticsEventAttribute>& ToAdd)
 	{
 		FScopeLock ScopeLock(&CriticalSection);
@@ -37,10 +77,28 @@ namespace UE::Interchange
 	{
 		if (FEngineAnalytics::IsAvailable())
 		{
-			for (const TPair<FString, TArray<FAnalyticsEventAttribute>>& AnalyticsPerIdentifier : AnalyticsAttributes)
+			FAnalyticsEventAttribute InterchangeAnalyticsIDAttribute(TEXT("InterchangeAnalyticsGUID"), AnalyticsID.ToString());
+
+			for (TPair<FString, TArray<FAnalyticsEventAttribute>>& AnalyticsPerIdentifier : AnalyticsAttributes)
 			{
-				FEngineAnalytics::GetProvider().RecordEvent(AnalyticsPerIdentifier.Key, AnalyticsPerIdentifier.Value);
+				TArray<FAnalyticsEventAttribute> ModifiedAttributes;
+				ModifiedAttributes.Reserve(AnalyticsPerIdentifier.Value.Num() + 1);
+				ModifiedAttributes = MoveTemp(AnalyticsPerIdentifier.Value);
+				ModifiedAttributes.Add(InterchangeAnalyticsIDAttribute);
+
+				FEngineAnalytics::GetProvider().RecordEvent(AnalyticsPerIdentifier.Key, ModifiedAttributes);
 			}
+
+			AnalyticsAttributes.Empty();
+		}
+	}
+	
+	void FAnalyticsHelper::SendAnalyticsImmediate(const FString& Identifier, TArray<FAnalyticsEventAttribute> Attributes)
+	{
+		if (FEngineAnalytics::IsAvailable())
+		{			
+			Attributes.Add(CreateAnalyticsIDAttribute(AnalyticsID));
+			FEngineAnalytics::GetProvider().RecordEvent(Identifier, Attributes);
 		}
 	}
 };
\ No newline at end of file
diff --git a/Engine/Source/Runtime/Interchange/Engine/Private/InterchangeManager.cpp b/Engine/Source/Runtime/Interchange/Engine/Private/InterchangeManager.cpp
--- a/Engine/Source/Runtime/Interchange/Engine/Private/InterchangeManager.cpp
+++ b/Engine/Source/Runtime/Interchange/Engine/Private/InterchangeManager.cpp
@@ -98,7 +98,7 @@ namespace UE::Interchange::Private
 #endif
 	}
 
-	void FillPipelineAnalyticData(UInterchangePipelineBase* Pipeline, const int32 UniqueId, const FString& ParentPipeline)
+	void FillPipelineAnalyticData(FAnalyticsHelper& AnalyticsHelper, UInterchangePipelineBase* Pipeline, const int32 UniqueId, const FString& ParentPipeline)
 	{
 		if (!FEngineAnalytics::IsAvailable())
 		{
@@ -116,78 +116,35 @@ namespace UE::Interchange::Private
 
 		for (FProperty* Property = Class->PropertyLink; Property; Property = Property->PropertyLinkNext)
 		{
-			if (Property->HasAnyPropertyFlags(CPF_Transient))
+			if (!Pipeline || !Pipeline->CanAddPropertyToAnalytics(Property))
 			{
 				continue;
 			}
 
-			if (Property->GetFName() == UInterchangePipelineBase::GetPropertiesStatesPropertyName())
-			{
-				continue;
-			}
-
-			if (Property->GetFName() == UInterchangePipelineBase::GetResultsPropertyName())
-			{
-				continue;
-			}
-
-			const FString PropertyName = Property->GetName();
 			FObjectProperty* SubObject = CastField<FObjectProperty>(Property);
 			UInterchangePipelineBase* SubPipeline = SubObject ? Cast<UInterchangePipelineBase>(SubObject->GetObjectPropertyValue_InContainer(Pipeline)) : nullptr;
-			//Add the category name to the key
-			FString CategoryName = FString();
-#if WITH_EDITORONLY_DATA
-			CategoryName = Property->GetMetaData("Category");
-			if (!SubPipeline && CategoryName.IsEmpty())
-			{
-				//In Editor do not add property with no category
-				continue;
-			}
-			CategoryName.ReplaceCharInline(TEXT('.'), TEXT('_'));
-			CategoryName.RemoveSpacesInline();
-			CategoryName = TEXT(".") + CategoryName;
-#endif
-
 			
-			if (FArrayProperty* Array = CastField<FArrayProperty>(Property))
-			{
-				FScriptArrayHelper_InContainer ArrayHelper(Array, Pipeline);
-				for (int32 i = 0; i < ArrayHelper.Num(); i++)
-				{
-					FString	Buffer;
-					Array->Inner->ExportTextItem_Direct(Buffer, ArrayHelper.GetRawPtr(i), ArrayHelper.GetRawPtr(i), Pipeline, PortFlags);
-					
-					PipelineAttribs.Add(FAnalyticsEventAttribute(FString::Printf(TEXT("Property%s.%s"), *CategoryName, *PropertyName), Buffer));
-				}
-			}
-			else if (SubPipeline)
+			
+			if (SubPipeline)
 			{
 				// Save the settings if the referenced pipeline is a subobject of ours
 				if (SubPipeline->IsInOuter(Pipeline))
 				{
 					//Go recursive with subObject, like if they are part of the same object
-					FillPipelineAnalyticData(SubPipeline, UniqueId, PipelineChainName);
+					FillPipelineAnalyticData(AnalyticsHelper, SubPipeline, UniqueId, PipelineChainName);
 				}
 			}
 			else
 			{
-				for (int32 Index = 0; Index < Property->ArrayDim; Index++)
-				{
-					FString PropertyIndexName = PropertyName;
-					if (Property->ArrayDim != 1)
-					{
-						PropertyIndexName += TEXT("[") + FString::FromInt(Index) + TEXT("]");
-					}
-
-					FString	Value;
-					Property->ExportText_InContainer(Index, Value, Pipeline, Pipeline, Pipeline, PortFlags);
-					PipelineAttribs.Add(FAnalyticsEventAttribute(FString::Printf(TEXT("Property%s.%s"), *CategoryName, *PropertyIndexName), Value));
-				}
+				Pipeline->AddPropertyAnalytics(Property, PipelineAttribs);
 			}
 		}
 		
-		FString EventString = "Interchange.Usage.Import.Pipeline";
-		FEngineAnalytics::GetProvider().RecordEvent(EventString, PipelineAttribs);
+		{
+			const FString PipelineEventIdentifier = "Interchange.Usage.Import.Pipeline";
+			AnalyticsHelper.SendAnalyticsImmediate(PipelineEventIdentifier, PipelineAttribs);
+		}
+		
 	}
 
 	// Json writer subclass to allow us to avoid using a SharedPtr to write basic Json.
@@ -331,9 +288,10 @@ UInterchangeBaseNodeContainer* UE::Interchange::FScopedBaseNodeContainer::GetBas
 	return ScopedBaseNodeContainerPtr.Get();
 }
 
-UE::Interchange::FImportAsyncHelper::FImportAsyncHelper()
+UE::Interchange::FImportAsyncHelper::FImportAsyncHelper(const FGuid& InAnalyticsID)
 	: AssetImportResult(MakeShared<FImportResult>())
 	, SceneImportResult(MakeShared<FImportResult>())
+	, AnalyticsID(InAnalyticsID)
 {
 	bCancel = false;
 
@@ -565,19 +523,17 @@ bool UE::Interchange::FImportAsyncHelper::IsImportingObject(UObject* Object) con
 
 void UE::Interchange::FImportAsyncHelper::SendAnalyticImportEndData()
 {
-	if (!FEngineAnalytics::IsAvailable())
+	if (bCancel || !FEngineAnalytics::IsAvailable())
 	{
 		return;
 	}
 
 	TArray<FAnalyticsEventAttribute> Attribs;
 	//Set the unique id of this import
+	Attribs.Add(FAnalyticsHelper::CreateAnalyticsIDAttribute(AnalyticsID));
 	Attribs.Add(FAnalyticsEventAttribute(TEXT("UniqueId"), UniqueId));
 	Attribs.Add(FAnalyticsEventAttribute(TEXT("IsCanceled"), bCancel));
-	if (bCancel)
-	{
-		return;
-	}
+	
 
 	const UInterchangeBaseNodeContainer* PrimaryBaseNodeContainer = BaseNodeContainers[0].Get();
 	if(const UInterchangeSourceNode* PrimarySourceNode = UInterchangeSourceNode::GetUniqueInstance(PrimaryBaseNodeContainer))
@@ -661,8 +617,10 @@ void UE::Interchange::FImportAsyncHelper::SendAnalyticImportEndData()
 		Attribs.Add(FAnalyticsEventAttribute(TEXT("ErrorTypes"), AnalyticsConvertMapToJsonFragment(ErrorTypes)));
 	}
 
-	FString EventString = "Interchange.Usage.ImportResult";
-	FEngineAnalytics::GetProvider().RecordEvent(EventString, Attribs);
+	{
+		const FString ImportResultEventIdentifier = "Interchange.Usage.ImportResult";
+		FEngineAnalytics::GetProvider().RecordEvent(ImportResultEventIdentifier, Attribs);
+	}
 }
 
 void UE::Interchange::FImportAsyncHelper::ReleaseTranslatorsSource()
@@ -1843,7 +1801,9 @@ UInterchangeManager::ImportInternal(const FString& ContentPath, const UInterchan
 	static int32 GeneratedUniqueID = 0;
 	int32 UniqueId = ++GeneratedUniqueID;
 
-	TSharedPtr<UE::Interchange::FAnalyticsHelper, ESPMode::ThreadSafe> AnalyticsHelper = MakeShared< UE::Interchange::FAnalyticsHelper, ESPMode::ThreadSafe>();
+	const FGuid AnalyticsGuid = FGuid::NewDeterministicGuid(ContentPath, static_cast<uint64>(FDateTime::UtcNow().GetTicks()));
+
+	TSharedPtr<UE::Interchange::FAnalyticsHelper, ESPMode::ThreadSafe> AnalyticsHelper = MakeShared< UE::Interchange::FAnalyticsHelper, ESPMode::ThreadSafe>(AnalyticsGuid);
 	const FString UsageImportAnalayticsId = "Interchange.Usage.Import";
 	auto PreReturn = [&AnalyticsHelper]()
 	{
@@ -2067,7 +2027,7 @@ UInterchangeManager::ImportInternal(const FString& ContentPath, const UInterchan
 	TaskData.DestinationName = ImportAssetParameters.DestinationName;
 	TaskData.bReplaceExisting = ImportAssetParameters.bReplaceExisting;
 
-	TSharedRef<UE::Interchange::FImportAsyncHelper, ESPMode::ThreadSafe> AsyncHelper = CreateAsyncHelper(TaskData, ImportAssetParameters);
+	TSharedRef<UE::Interchange::FImportAsyncHelper, ESPMode::ThreadSafe> AsyncHelper = CreateAsyncHelper(TaskData, ImportAssetParameters, AnalyticsGuid);
 	AsyncHelper->UniqueId = UniqueId;
 	AsyncHelper->bRunSynchronous = ImportAssetParameters.bRunSynchronous;
 
@@ -2180,7 +2140,7 @@ UInterchangeManager::ImportInternal(const FString& ContentPath, const UInterchan
 			{
 				AsyncHelper->Pipelines.Add(Pipeline);
 				AsyncHelper->OriginalPipelines.Add(Pipeline);
-				UE::Interchange::Private::FillPipelineAnalyticData(Pipeline, UniqueId, FString());
+				UE::Interchange::Private::FillPipelineAnalyticData(*AnalyticsHelper, Pipeline, UniqueId, FString());
 			}
 		}
 	}
@@ -2454,7 +2414,7 @@ UInterchangeManager::ImportInternal(const FString& ContentPath, const UInterchan
 				{
 					AsyncHelper->OriginalPipelines.Add(Pipeline);
 				}
-				UE::Interchange::Private::FillPipelineAnalyticData(Pipeline, UniqueId, FString());
+				UE::Interchange::Private::FillPipelineAnalyticData(*AnalyticsHelper, Pipeline, UniqueId, FString());
 			}
 		}
 	}
@@ -2718,9 +2678,9 @@ const UClass* UInterchangeManager::GetRegisteredFactoryClass(const UClass* Class
 	return Result;
 }
 
-TSharedRef<UE::Interchange::FImportAsyncHelper, ESPMode::ThreadSafe> UInterchangeManager::CreateAsyncHelper(const UE::Interchange::FImportAsyncHelperData& Data, const FImportAssetParameters& ImportAssetParameters)
+TSharedRef<UE::Interchange::FImportAsyncHelper, ESPMode::ThreadSafe> UInterchangeManager::CreateAsyncHelper(const UE::Interchange::FImportAsyncHelperData& Data, const FImportAssetParameters& ImportAssetParameters, const FGuid& InAnalyticsID)
 {
-	TSharedRef<UE::Interchange::FImportAsyncHelper, ESPMode::ThreadSafe> AsyncHelper = MakeShared<UE::Interchange::FImportAsyncHelper, ESPMode::ThreadSafe>();
+	TSharedRef<UE::Interchange::FImportAsyncHelper, ESPMode::ThreadSafe> AsyncHelper = MakeShared<UE::Interchange::FImportAsyncHelper, ESPMode::ThreadSafe>(InAnalyticsID);
 	//Copy the task data
 	AsyncHelper->TaskData = Data;
 	
diff --git a/Engine/Source/Runtime/Interchange/Engine/Public/InterchangeAnalyticsHelper.h b/Engine/Source/Runtime/Interchange/Engine/Public/InterchangeAnalyticsHelper.h
--- a/Engine/Source/Runtime/Interchange/Engine/Public/InterchangeAnalyticsHelper.h
+++ b/Engine/Source/Runtime/Interchange/Engine/Public/InterchangeAnalyticsHelper.h
@@ -3,30 +3,68 @@
 #pragma once
 
 #include "CoreMinimal.h"
-struct FAnalyticsEventAttribute;
+#include "AnalyticsEventAttribute.h"
+
+#define UE_API INTERCHANGEENGINE_API
 
 namespace UE::Interchange
 {
+	/**
+	 * If you add or remove from this enum, make sure to update the corresponding strings in FAnalyticsHelper::GetDefaultAnalyticsCategoryName
+	 */
+	enum class EInterchangeAnalyticsDefaultCategory : uint8
+	{
+		NotSpecified = 0,
+		Models,
+		Scenes,
+		Materials,
+		Textures,
+		Sounds,
+		Meshes,
+		StaticMeshes,
+		SkeletalMeshesAndAnimations,
+		GeometryCaches,
+		Grooms,
+		InterchangeAnalyticsCategory_COUNT
+	};
+
 	struct FAnalyticsHelper
 	{
 	public:
-		INTERCHANGEENGINE_API void AppendThreadSafe(const FString& Identifier, const TArray<FAnalyticsEventAttribute>& ToAdd);
+		UE_API static FAnalyticsEventAttribute CreateAnalyticsIDAttribute(FGuid AnalyticsID);
+
+		UE_API static FString GetDefaultAnalyticsCategoryName(EInterchangeAnalyticsDefaultCategory DefaultCategory);
+
+	public:
+		UE_API FAnalyticsHelper(const FGuid& InAnalyticsID);
+
+		UE_API void AppendThreadSafe(const FString& Identifier, const TArray<FAnalyticsEventAttribute>& ToAdd);
 
 		// Remove all the event attributes if the identifier matches.
-		INTERCHANGEENGINE_API void ClearAnalyticsEventData(const FString& Identifier);
+		UE_API void ClearAnalyticsEventData(const FString& Identifier);
 
 		//Make sure it is safe to call this directly.
-		INTERCHANGEENGINE_API void Add(const FString& Identifier, const FAnalyticsEventAttribute& Entry);
+		UE_API void Add(const FString& Identifier, const FAnalyticsEventAttribute& Entry);
 
 		//It is suggest to use the AppendThreadSafe, if more than one Entry is added as it will Scope Lock for every entry.
 		// (where append will Scope Lock only per append)
-		INTERCHANGEENGINE_API void AddThreadSafe(const FString& Identifier, const FAnalyticsEventAttribute& Entry);
+		UE_API void AddThreadSafe(const FString& Identifier, const FAnalyticsEventAttribute& Entry);
+
+		UE_API void SendAnalytics();
 
-		INTERCHANGEENGINE_API void SendAnalytics();
+		/** 
+		 * Adds the GUID to the Attributes and sends out the data immediately. Useful when we want to send multiple events with same identifier in quick succession.
+		 * For example, Interchange.Usage.Import.Pipeline Events are sent these events out when we traverse the pipeline objects recursively for subpipelines. 
+		 */
+		UE_API void SendAnalyticsImmediate(const FString& Identifier, TArray<FAnalyticsEventAttribute> Attributes);
 
 	private:
 		//Analytics Attributes to share with the Translators. (to be sent when Import is finished)
 		TMap<FString, TArray<FAnalyticsEventAttribute>> AnalyticsAttributes;
 		FCriticalSection CriticalSection;
+
+		FGuid AnalyticsID;
 	};
-};
\ No newline at end of file
+};
+
+#undef UE_API
\ No newline at end of file
diff --git a/Engine/Source/Runtime/Interchange/Engine/Public/InterchangeManager.h b/Engine/Source/Runtime/Interchange/Engine/Public/InterchangeManager.h
--- a/Engine/Source/Runtime/Interchange/Engine/Public/InterchangeManager.h
+++ b/Engine/Source/Runtime/Interchange/Engine/Public/InterchangeManager.h
@@ -21,6 +21,7 @@
 #include "InterchangeTaskSystem.h"
 #include "InterchangeTranslatorBase.h"
 #include "InterchangeWriterBase.h"
+#include "Misc/Guid.h"
 #include "Nodes/InterchangeBaseNodeContainer.h"
 #include "Templates/Tuple.h"
 #include "UObject/GCObject.h"
@@ -39,6 +40,8 @@ class UInterchangePipelineBase;
 class UInterchangePythonPipelineBase;
 class ULevel;
 
+struct FAnalyticsEventAttribute;
+
 namespace UE::Interchange
 {
 	struct FAnalyticsHelper;
@@ -229,7 +232,7 @@ namespace UE
 		class FImportAsyncHelper : protected FGCObject
 		{
 		public:
-			FImportAsyncHelper();
+			FImportAsyncHelper(const FGuid& AnalyticsID);
 
 			~FImportAsyncHelper();
 
@@ -365,6 +368,8 @@ namespace UE
 
 			mutable FCriticalSection ImportedSceneObjectsPerSourceIndexLock;
 			TMap<int32, TArray<FImportedObjectInfo>> ImportedSceneObjectsPerSourceIndex;
+
+			FGuid AnalyticsID;
 		};
 
 		/* This function takes an asset that represents a pipeline and generates a UInterchangePipelineBase asset. */
@@ -798,7 +803,7 @@ public:
 	 * Return a pointer to an FImportAsyncHelper. The pointer is deleted when ReleaseAsyncHelper is called.
 	 * @param Data - The data to pass to the different import tasks.
 	 */
-	INTERCHANGEENGINE_API TSharedRef<UE::Interchange::FImportAsyncHelper, ESPMode::ThreadSafe> CreateAsyncHelper(const UE::Interchange::FImportAsyncHelperData& Data, const FImportAssetParameters& ImportAssetParameters);
+	INTERCHANGEENGINE_API TSharedRef<UE::Interchange::FImportAsyncHelper, ESPMode::ThreadSafe> CreateAsyncHelper(const UE::Interchange::FImportAsyncHelperData& Data, const FImportAssetParameters& ImportAssetParameters, const FGuid& InAnalyticsID);
 
 	/** Delete the specified AsyncHelper and remove it from the array that was holding it. */
 	INTERCHANGEENGINE_API void ReleaseAsyncHelper(TWeakPtr<UE::Interchange::FImportAsyncHelper, ESPMode::ThreadSafe> AsyncHelper);
