[Interchange] Added nullcheck to Pipeline for Analytics

#rnx
#Synced-CL 49210907

diff --git a/Engine/Source/Runtime/Interchange/Engine/Private/InterchangeManager.cpp b/Engine/Source/Runtime/Interchange/Engine/Private/InterchangeManager.cpp
--- a/Engine/Source/Runtime/Interchange/Engine/Private/InterchangeManager.cpp
+++ b/Engine/Source/Runtime/Interchange/Engine/Private/InterchangeManager.cpp
@@ -100,7 +100,7 @@ namespace UE::Interchange::Private
 
 	void FillPipelineAnalyticData(FAnalyticsHelper& AnalyticsHelper, UInterchangePipelineBase* Pipeline, const int32 UniqueId, const FString& ParentPipeline)
 	{
-		if (!FEngineAnalytics::IsAvailable())
+		if (!Pipeline || !FEngineAnalytics::IsAvailable())
 		{
 			return;
 		}
