[Interchange Analytics] Exposed ability to get Analytics Attribute from the Analytics Helper.

#rnx
#Synced-CL 49287615

diff --git a/Engine/Source/Runtime/Interchange/Engine/Private/InterchangeAnalyticsHelper.cpp b/Engine/Source/Runtime/Interchange/Engine/Private/InterchangeAnalyticsHelper.cpp
--- a/Engine/Source/Runtime/Interchange/Engine/Private/InterchangeAnalyticsHelper.cpp
+++ b/Engine/Source/Runtime/Interchange/Engine/Private/InterchangeAnalyticsHelper.cpp
@@ -101,4 +101,9 @@ namespace UE::Interchange
 			FEngineAnalytics::GetProvider().RecordEvent(Identifier, Attributes);
 		}
 	}
+
+	FAnalyticsEventAttribute FAnalyticsHelper::GetAnalyticsIDAttribute() const
+	{
+		return CreateAnalyticsIDAttribute(AnalyticsID);
+	}
 };
\ No newline at end of file
diff --git a/Engine/Source/Runtime/Interchange/Engine/Public/InterchangeAnalyticsHelper.h b/Engine/Source/Runtime/Interchange/Engine/Public/InterchangeAnalyticsHelper.h
--- a/Engine/Source/Runtime/Interchange/Engine/Public/InterchangeAnalyticsHelper.h
+++ b/Engine/Source/Runtime/Interchange/Engine/Public/InterchangeAnalyticsHelper.h
@@ -58,6 +58,13 @@ namespace UE::Interchange
 		 */
 		UE_API void SendAnalyticsImmediate(const FString& Identifier, TArray<FAnalyticsEventAttribute> Attributes);
 
+
+		/** 
+		 * Creates an Analytics Event Attribute for the current instance of the analytics helper. 
+		 * Helpful to be used on external analytics events not sent by AnalyticsHelper.
+		 */
+		UE_API FAnalyticsEventAttribute GetAnalyticsIDAttribute() const;
+
 	private:
 		//Analytics Attributes to share with the Translators. (to be sent when Import is finished)
 		TMap<FString, TArray<FAnalyticsEventAttribute>> AnalyticsAttributes;
