Interchange Experimental Scene Graph Support

#rnx

#virtualized
#Synced-CL 50223086

diff --git a/Engine/Plugins/Interchange/Runtime/Source/Import/Public/Scene/InterchangeLevelFactory.h b/Engine/Plugins/Interchange/Runtime/Source/Import/Public/Scene/InterchangeLevelFactory.h
--- a/Engine/Plugins/Interchange/Runtime/Source/Import/Public/Scene/InterchangeLevelFactory.h
+++ b/Engine/Plugins/Interchange/Runtime/Source/Import/Public/Scene/InterchangeLevelFactory.h
@@ -21,6 +21,7 @@ public:
 	// Interchange factory base interface begin
 
 	UE_API virtual UClass* GetFactoryClass() const override;
+	virtual EInterchangeFactoryAssetType GetFactoryAssetType() override { return EInterchangeFactoryAssetType::Levels; }
 	UE_API virtual FImportAssetResult BeginImportAsset_GameThread(const FImportAssetObjectParams& Arguments) override;
 	UE_API virtual void SetupObject_GameThread(const FSetupObjectParams& Arguments) override;
 	// For a level reimport, returns the level created on the previous import if found
diff --git a/Engine/Plugins/Interchange/Runtime/Source/Import/Public/Scene/InterchangeSceneImportAssetFactory.h b/Engine/Plugins/Interchange/Runtime/Source/Import/Public/Scene/InterchangeSceneImportAssetFactory.h
--- a/Engine/Plugins/Interchange/Runtime/Source/Import/Public/Scene/InterchangeSceneImportAssetFactory.h
+++ b/Engine/Plugins/Interchange/Runtime/Source/Import/Public/Scene/InterchangeSceneImportAssetFactory.h
@@ -21,6 +21,9 @@ public:
 	// Interchange factory base interface begin
 
 	UE_API virtual UClass* GetFactoryClass() const override;
+
+	virtual EInterchangeFactoryAssetType GetFactoryAssetType() override { return EInterchangeFactoryAssetType::Custom; }
+	
 	UE_API virtual FImportAssetResult BeginImportAsset_GameThread(const FImportAssetObjectParams& Arguments) override;
 	UE_API virtual void SetupObject_GameThread(const FSetupObjectParams& Arguments) override;
 	UE_API virtual void FinalizeObject_GameThread(const FSetupObjectParams& Arguments) override;
diff --git a/Engine/Source/Editor/UnrealEd/Private/Editor.cpp b/Engine/Source/Editor/UnrealEd/Private/Editor.cpp
--- a/Engine/Source/Editor/UnrealEd/Private/Editor.cpp
+++ b/Engine/Source/Editor/UnrealEd/Private/Editor.cpp
@@ -1014,6 +1014,9 @@ void FReimportManager::GetNewReimportPath(UObject* Obj, TArray<FString>& InOutFi
 					case EInterchangeFactoryAssetType::Grooms:
 						TranslatorAssetType = EInterchangeTranslatorAssetType::Grooms;
 						break;
+					case EInterchangeFactoryAssetType::Sounds:
+						TranslatorAssetType = EInterchangeTranslatorAssetType::Sounds;
+						break;
 				};
 				TranslatorFormats = UInterchangeManager::GetInterchangeManager().GetSupportedAssetTypeFormats(TranslatorAssetType, EInterchangeTranslatorType::Assets, false /*bStrictMatchTranslatorType*/);
 			}
diff --git a/Engine/Source/Runtime/Interchange/Core/Public/InterchangeFactoryBase.h b/Engine/Source/Runtime/Interchange/Core/Public/InterchangeFactoryBase.h
--- a/Engine/Source/Runtime/Interchange/Core/Public/InterchangeFactoryBase.h
+++ b/Engine/Source/Runtime/Interchange/Core/Public/InterchangeFactoryBase.h
@@ -29,7 +29,9 @@ enum class EInterchangeFactoryAssetType : uint8
 	Animations,
 	Physics,
 	Grooms,
-	Sounds
+	Sounds,
+	Levels,
+	Custom
 };
 
 /**
diff --git a/Engine/Source/Runtime/Interchange/Engine/Private/Tasks/InterchangeTaskParsing.cpp b/Engine/Source/Runtime/Interchange/Engine/Private/Tasks/InterchangeTaskParsing.cpp
--- a/Engine/Source/Runtime/Interchange/Engine/Private/Tasks/InterchangeTaskParsing.cpp
+++ b/Engine/Source/Runtime/Interchange/Engine/Private/Tasks/InterchangeTaskParsing.cpp
@@ -344,23 +344,26 @@ void UE::Interchange::FTaskParsing::Execute()
 				UClass* ObjectClass = FactoryNode->GetObjectClass();
 				if (ObjectClass != nullptr)
 				{
-					const UClass* RegisteredFactoryClass = InterchangeManager->GetRegisteredFactoryClass(ObjectClass);
-
-					const bool bIsSceneNode = FactoryNode->GetObjectClass()->IsChildOf<AActor>() || FactoryNode->GetObjectClass()->IsChildOf<UActorComponent>();
-
-					if (!RegisteredFactoryClass || (bIsSceneNode && !bCanImportSceneNode))
+					if (const UClass* RegisteredFactoryClass = InterchangeManager->GetRegisteredFactoryClass(ObjectClass))
 					{
-						//nothing we can import from this element
-						return;
+						if (UInterchangeFactoryBase* FactoryCDO = Cast<UInterchangeFactoryBase>(RegisteredFactoryClass->GetDefaultObject()))
+						{
+							const bool bIsSceneNode = (FactoryCDO->GetFactoryAssetType() == EInterchangeFactoryAssetType::None);
+							if (bIsSceneNode && !bCanImportSceneNode)
+							{
+								//nothing we can import from this element
+								return;
+							}
+
+							FTaskData& NodeTaskData = SourceTaskDatas.AddDefaulted_GetRef();
+							NodeTaskData.UniqueID = FactoryNode->GetUniqueID();
+							NodeTaskData.SourceIndex = SourceIndex;
+							NodeTaskData.bIsSceneNode = bIsSceneNode;
+							NodeTaskData.Nodes.Add(FactoryNode);
+							FactoryNode->GetFactoryDependencies(NodeTaskData.Dependencies);
+							NodeTaskData.FactoryClass = RegisteredFactoryClass;
+						}
 					}
-
-					FTaskData& NodeTaskData = SourceTaskDatas.AddDefaulted_GetRef();
-					NodeTaskData.UniqueID = FactoryNode->GetUniqueID();
-					NodeTaskData.SourceIndex = SourceIndex;
-					NodeTaskData.bIsSceneNode = bIsSceneNode;
-					NodeTaskData.Nodes.Add(FactoryNode);
-					FactoryNode->GetFactoryDependencies(NodeTaskData.Dependencies);
-					NodeTaskData.FactoryClass = RegisteredFactoryClass;
 				}
 			});
 
