[Interchange]: Optim Task Parsing dependency tracking.  Use a map of UniqueId to index instead of sorting tasks and iterating over all tasks to compare their ID

#Synced-CL 50534375

diff --git a/Engine/Source/Runtime/Interchange/Engine/Private/Tasks/InterchangeTaskParsing.cpp b/Engine/Source/Runtime/Interchange/Engine/Private/Tasks/InterchangeTaskParsing.cpp
--- a/Engine/Source/Runtime/Interchange/Engine/Private/Tasks/InterchangeTaskParsing.cpp
+++ b/Engine/Source/Runtime/Interchange/Engine/Private/Tasks/InterchangeTaskParsing.cpp
@@ -24,6 +24,7 @@
 #include "UObject/ObjectMacros.h"
 #include "UObject/UObjectGlobals.h"
 #include "UObject/WeakObjectPtrTemplates.h"
+#include "Algo/TopologicalSort.h"
 
 namespace UE::InterchangeTaskParsing::Private
 {
@@ -41,85 +42,13 @@ namespace UE::InterchangeTaskParsing::Private
 		int32 SourceIndex = INDEX_NONE;
 		bool bIsSceneNode = false;
 		TArray<FString> Dependencies;
-		uint64 GraphEventRef;
+		uint64 GraphEventRef = INTERCHANGE_INVALID_TASK_ID;
 		TArray<uint64> Prerequisites;
 		const UClass* FactoryClass;
 
 		TArray<UInterchangeFactoryBaseNode*, TInlineAllocator<1>> Nodes; // For scenes, we can group multiple nodes into a single task as they are usually very light
 	};
 
-	/**
-	* For the Dependency sort to work the predicate must be transitive ( A > B > C implying A > C).
-	* That means we must take into account the whole dependency chain, not just the immediate dependencies.
-	* 
-	* This is a helper struct to quickly create the dependencies chain of a node using a cache.
-	*/
-	struct FNodeDependencyCache
-	{
-		void CacheAccumulatedDependencies(const UInterchangeBaseNodeContainer* NodeContainer, const TArray<FTaskData>& SourceTaskDatas)
-		{
-			CachedDependencies.Reset();
-			for (const FTaskData& TaskData : SourceTaskDatas)
-			{
-				TSet<FString> NodeStack;
-				GetAccumulatedDependencies(NodeContainer, TaskData.UniqueID, NodeStack);
-			}
-		}
-
-		const TSet<FString>& GetAccumulatedDependencies(const UInterchangeBaseNodeContainer* NodeContainer, const FString& NodeID)
-		{
-			TRACE_CPUPROFILER_EVENT_SCOPE(FNodeDependencyCache::GetAccumulatedDependencies)
-			ensure(CachedDependencies.Find(NodeID));
-			TSet<FString> NodeStack;
-			return GetAccumulatedDependencies(NodeContainer, NodeID, NodeStack);
-		}
-
-	private:
-
-		const TSet<FString>& GetAccumulatedDependencies(const UInterchangeBaseNodeContainer* NodeContainer, const FString& NodeID, TSet<FString>& NodeStack)
-		{
-			if (const TSet<FString>* DependenciesPtr = CachedDependencies.Find(NodeID))
-			{
-				return *DependenciesPtr;
-			}
-
-			TSet<FString> Dependencies;
-			AccumulateDependencies(NodeContainer, NodeID, Dependencies, NodeStack);
-			return CachedDependencies.Add(NodeID, MoveTemp(Dependencies));
-		}
-
-		void AccumulateDependencies(const UInterchangeBaseNodeContainer* NodeContainer, const FString& NodeID, TSet<FString>& OutDependenciesSet, TSet<FString>& NodeStack)
-		{
-			const UInterchangeFactoryBaseNode* FactoryNode = NodeContainer->GetFactoryNode(NodeID);
-			if (!FactoryNode)
-			{
-				return;
-			}
-
-			bool bAlreadyInSet = false;
-			NodeStack.Add(NodeID, &bAlreadyInSet);
-			if (ensureMsgf(!bAlreadyInSet, TEXT("FNodeDependencyCache::AccumulateDependencies - Node \"%s\" is in a circular dependency, assets may not be imported properly."), *NodeID))
-			{
-				TArray<FString> FactoryDependencies;
-				FactoryNode->GetFactoryDependencies(FactoryDependencies);
-				OutDependenciesSet.Reserve(OutDependenciesSet.Num() + FactoryDependencies.Num());
-				for (const FString& DependencyID : FactoryDependencies)
-				{
-					bAlreadyInSet = false;
-					OutDependenciesSet.Add(DependencyID, &bAlreadyInSet);
-					// Avoid infinite recursion.
-					if (!bAlreadyInSet)
-					{
-						OutDependenciesSet.Append(GetAccumulatedDependencies(NodeContainer, DependencyID, NodeStack));
-					}
-				}
-				NodeStack.Remove(NodeID);
-			}
-		}
-
-		TMap<FString, TSet<FString>> CachedDependencies;
-	};
-
 	uint64 CreateTasksFromData(
 		FTaskData& InTaskData,
 		const TSharedPtr<Interchange::FImportAsyncHelper, ESPMode::ThreadSafe>& InAsyncHelper,
@@ -316,14 +245,13 @@ void UE::Interchange::FTaskParsing::Execute()
 
 	//Parse each graph and prepare import task data, we will then be able to create all the task with the correct dependencies
 	TArray<FTaskData> TaskDatas;
+	TMap<FString, int32> TaskUniqueIdToIndex;
 
 	//Avoid creating asset if the asynchronous import is canceled, just create the completion task
 	if (!AsyncHelper->bCancel)
 	{
 		for (int32 SourceIndex = 0; SourceIndex < AsyncHelper->SourceDatas.Num(); ++SourceIndex)
 		{
-			TArray<FTaskData> SourceTaskDatas;
-
 			if (!AsyncHelper->BaseNodeContainers.IsValidIndex(SourceIndex))
 			{
 				continue;
@@ -355,81 +283,64 @@ void UE::Interchange::FTaskParsing::Execute()
 								return;
 							}
 
-							FTaskData& NodeTaskData = SourceTaskDatas.AddDefaulted_GetRef();
+							FTaskData NodeTaskData;
 							NodeTaskData.UniqueID = FactoryNode->GetUniqueID();
 							NodeTaskData.SourceIndex = SourceIndex;
 							NodeTaskData.bIsSceneNode = bIsSceneNode;
 							NodeTaskData.Nodes.Add(FactoryNode);
 							FactoryNode->GetFactoryDependencies(NodeTaskData.Dependencies);
 							NodeTaskData.FactoryClass = RegisteredFactoryClass;
+							int32 TaskDataIndex = TaskDatas.Add(NodeTaskData);
+							TaskUniqueIdToIndex.Add(NodeTaskData.UniqueID, TaskDataIndex);
 						}
 					}
 				}
 			});
-
-			{
-				FNodeDependencyCache DependencyCache;
-				//Reserving the cache will allow us to use reference when sorting the dependencies.
-				DependencyCache.CacheAccumulatedDependencies(BaseNodeContainer, SourceTaskDatas);
-
-				//Sort per dependencies
-				auto SortByDependencies =
-					[&BaseNodeContainer, &DependencyCache](const FTaskData& A, const FTaskData& B)
-				{
-					const TSet<FString>& BDependencies = DependencyCache.GetAccumulatedDependencies(BaseNodeContainer, B.UniqueID);
-					const TSet<FString>& ADependencies = DependencyCache.GetAccumulatedDependencies(BaseNodeContainer, A.UniqueID);
-
-					if (ADependencies.Num() == BDependencies.Num())
-					{
-						return A.UniqueID < B.UniqueID;
-					}
-
-					return ADependencies.Num() < BDependencies.Num();
-				};
-
-				// Nodes cannot depend on a node from another source, so it's faster to sort the dependencies per-source and then append those to the TaskData arrays.
-				SourceTaskDatas.Sort(SortByDependencies);
-			}
-
-			TaskDatas.Append(MoveTemp(SourceTaskDatas));
 		}
 	}
 
 	TMap<FString, FTaskParsingRenameInfo> RenameAssets;
 	TSet<FString> CreatedTasksAssetNames; // Tracks for which asset name we have created a task so that we don't have 2 tasks for the same asset name
 
-	TArray<uint64> CompletionPrerequistes;
-	const int32 PoolWorkerThreadCount = FTaskGraphInterface::Get().GetNumWorkerThreads() / 2;
-	const int32 MaxNumWorker = FMath::Max(PoolWorkerThreadCount, 1);
-	TArray<uint64> GroupPrerequistes;
-	for (int32 TaskIndex = 0; TaskIndex < TaskDatas.Num(); ++TaskIndex)
+	auto GetTaskDependencies = [&TaskDatas, &TaskUniqueIdToIndex](int32 Index)
 	{
-		FTaskData& TaskData = TaskDatas[TaskIndex];
-
-		if (TaskData.Dependencies.Num() > 0)
+		TArray<int32> TaskDependencies;
+		TaskDependencies.Reserve(TaskDatas[Index].Dependencies.Num());
+		for(const FString& DependencyTaskID : TaskDatas[Index].Dependencies)
 		{
-			//Search the previous node to find the dependence
-			for (int32 DepTaskIndex = 0; DepTaskIndex < TaskIndex; ++DepTaskIndex)
+			if (int32* DependencyTaskIndex = TaskUniqueIdToIndex.Find(DependencyTaskID))
 			{
-				if (TaskData.Dependencies.Contains(TaskDatas[DepTaskIndex].UniqueID))
-				{
-					//Add has prerequisite
-					TaskData.Prerequisites.Add(TaskDatas[DepTaskIndex].GraphEventRef);
-				}
+				TaskDependencies.Add(*DependencyTaskIndex);
 			}
 		}
+		return TaskDependencies;
+	};
+
+	TArray<int32> TasksIndexArray;
+	TasksIndexArray.SetNumUninitialized(TaskDatas.Num());
+	for (int32 i = 0; i < TaskDatas.Num(); i++)
+	{
+		TasksIndexArray[i] = i;
+	}
+
+	// Sort Tasks According to their dependency in order that their GraphEventRef are created in in the right order
+	Algo::TopologicalSort(TasksIndexArray, GetTaskDependencies);
 
-		//Append the group prerequisite to the task data prerequisite if the group is full
-		//This allow to chain the group dependencies to control the number of task
-		if (GroupPrerequistes.Num() >= MaxNumWorker)
+	TArray<uint64> CompletionPrerequistes;
+	for (int32 TaskIndex : TasksIndexArray)
+	{
+		FTaskData& TaskData = TaskDatas[TaskIndex];
+		for(const FString& DependencyTaskID : TaskData.Dependencies)
 		{
-			TaskData.Prerequisites.Append(GroupPrerequistes);
-			GroupPrerequistes.Reset();
+			if (int32* DependencyTaskIndex = TaskUniqueIdToIndex.Find(DependencyTaskID))
+			{
+				ensure(TaskDatas[*DependencyTaskIndex].GraphEventRef != INTERCHANGE_INVALID_TASK_ID);				
+				TaskData.Prerequisites.Add(TaskDatas[*DependencyTaskIndex].GraphEventRef);
+			}
 		}
 		TaskData.GraphEventRef = CreateTasksFromData(TaskData, AsyncHelper, RenameAssets, CreatedTasksAssetNames);
-		GroupPrerequistes.Add(TaskData.GraphEventRef);
+		CompletionPrerequistes.Add(TaskData.GraphEventRef);
 	}
-	CompletionPrerequistes.Append(GroupPrerequistes);
 
 	if (!RenameAssets.IsEmpty())
 	{
