[Interchange]: Optim FTaskCreateSceneObjects_GameThread.  Remove the unnecessary Building of the imported assets list for every scene node. Remove unnecessary scene node find of a newly created scene node.

#Synced-CL 50595740

diff --git a/Engine/Plugins/Interchange/Runtime/Source/Import/Private/Scene/InterchangeSkeletalMeshActorFactory.cpp b/Engine/Plugins/Interchange/Runtime/Source/Import/Private/Scene/InterchangeSkeletalMeshActorFactory.cpp
--- a/Engine/Plugins/Interchange/Runtime/Source/Import/Private/Scene/InterchangeSkeletalMeshActorFactory.cpp
+++ b/Engine/Plugins/Interchange/Runtime/Source/Import/Private/Scene/InterchangeSkeletalMeshActorFactory.cpp
@@ -29,7 +29,7 @@ UObject* UInterchangeSkeletalMeshActorFactory::ProcessActor(AActor& SpawnedActor
 	if (USkeletalMeshComponent* SkeletalMeshComponent = SkeletalMeshActor->GetSkeletalMeshComponent())
 	{
 		USkeletalMesh* SkeletalMesh = SkeletalMeshComponent->GetSkeletalMeshAsset();
-		if (SkeletalMesh && Params.ImportAssets.Contains(SkeletalMesh))
+		if (SkeletalMesh)
 		{
 #if WITH_EDITOR
 			//If we are importing the skeletal mesh, we need to use a async build scope to unregister the component to avoid a deadlock
diff --git a/Engine/Source/Runtime/Interchange/Core/Public/InterchangeFactoryBase.h b/Engine/Source/Runtime/Interchange/Core/Public/InterchangeFactoryBase.h
--- a/Engine/Source/Runtime/Interchange/Core/Public/InterchangeFactoryBase.h
+++ b/Engine/Source/Runtime/Interchange/Core/Public/InterchangeFactoryBase.h
@@ -233,11 +233,6 @@ public:
 		 */
 		const UInterchangeFactoryBaseNode* ReimportFactoryNode = nullptr;
 
-		/**
-		 * Get all assets we are importing the actor could point on
-		 */
-		TArray<TObjectPtr<UObject>> ImportAssets;
-
 		/**
 		 * true if we are re-importing
 		 */
diff --git a/Engine/Source/Runtime/Interchange/Engine/Private/Tasks/InterchangeTaskCreateSceneObjects.cpp b/Engine/Source/Runtime/Interchange/Engine/Private/Tasks/InterchangeTaskCreateSceneObjects.cpp
--- a/Engine/Source/Runtime/Interchange/Engine/Private/Tasks/InterchangeTaskCreateSceneObjects.cpp
+++ b/Engine/Source/Runtime/Interchange/Engine/Private/Tasks/InterchangeTaskCreateSceneObjects.cpp
@@ -51,19 +51,6 @@ void UE::Interchange::FTaskCreateSceneObjects_GameThread::Execute()
 		return;
 	}
 
-	TArray<TObjectPtr<UObject>> ImportObjects;
-	AsyncHelper->IterateImportedAssets(SourceIndex, [&ImportObjects](const TArray<UE::Interchange::FImportAsyncHelper::FImportedObjectInfo>& ImportedObjectInfos)
-		{
-			for (const UE::Interchange::FImportAsyncHelper::FImportedObjectInfo& Info : ImportedObjectInfos)
-			{
-				UObject* ImportObject = Info.GetImportedObject();
-				if(IsValid(ImportObject))
-				{
-					ImportObjects.Add(ImportObject);
-				}
-			}
-		});
-
 	UObject* ReimportObject = AsyncHelper->TaskData.ReimportObject;
 	ULevel* ImportLevel = AsyncHelper->TaskData.ImportLevel ? AsyncHelper->TaskData.ImportLevel : GWorld->GetCurrentLevel();
 	UWorld* ImportWorld = ImportLevel->GetWorld();
@@ -134,7 +121,6 @@ void UE::Interchange::FTaskCreateSceneObjects_GameThread::Execute()
 		CreateSceneObjectsParams.ObjectName = SceneNodeName;
 		CreateSceneObjectsParams.FactoryNode = FactoryNode;
 		CreateSceneObjectsParams.Level = FactoryImportLevel;
-		CreateSceneObjectsParams.ImportAssets = ImportObjects;
 		CreateSceneObjectsParams.ReimportObject = FFactoryCommon::GetObjectToReimport(Factory, ReimportObject, *FactoryNode, FactoryNodeWorldPath, FactoryNodeWorldName, FactoryNodeSceneNodeName);
 		CreateSceneObjectsParams.ReimportFactoryNode = FFactoryCommon::GetFactoryNode(ReimportObject, FactoryNodeWorldPath, FactoryNodeWorldName, FactoryNodeSceneNodeName);
 		CreateSceneObjectsParams.bIsReimport = ReimportObject != nullptr;
@@ -156,19 +142,10 @@ void UE::Interchange::FTaskCreateSceneObjects_GameThread::Execute()
 		UObject* SceneObject = Factory->ImportSceneObject_GameThread(CreateSceneObjectsParams);
 		if (SceneObject)
 		{
-			const FImportAsyncHelper::FImportedObjectInfo* ImportedInfoPtr = AsyncHelper->FindImportedSceneObjects(SourceIndex, [SceneObject](const FImportAsyncHelper::FImportedObjectInfo& CurInfo)
-				{
-					return CurInfo.GetImportedObject() == SceneObject;
-				});
-
-			if (!ImportedInfoPtr)
-			{
-				UE::Interchange::FImportAsyncHelper::FImportedObjectInfo& ObjectInfo = AsyncHelper->AddDefaultImportedSceneObjectGetRef(SourceIndex);
-				ObjectInfo.SetImportedObject(SceneObject);
-				ObjectInfo.Factory = Factory;
-				ObjectInfo.FactoryNode = FactoryNode;
-			}
-
+			UE::Interchange::FImportAsyncHelper::FImportedObjectInfo& ObjectInfo = AsyncHelper->AddDefaultImportedSceneObjectGetRef(SourceIndex);
+			ObjectInfo.SetImportedObject(SceneObject);
+			ObjectInfo.Factory = Factory;
+			ObjectInfo.FactoryNode = FactoryNode;
 			FactoryNode->SetCustomReferenceObject(FSoftObjectPath(SceneObject));
 		}
 	}
