[Interchange Scene Import Asset] Added a utility method: GetFactoryNodeFromUniqueId
- Also renamed other methods for clarity.
- Pre-requisite for the fix for UE-350381

#rnx

#Synced-CL 50731115

diff --git a/Engine/Plugins/Interchange/Editor/Source/InterchangeEditor/Private/InterchangeEditorScriptLibrary.cpp b/Engine/Plugins/Interchange/Editor/Source/InterchangeEditor/Private/InterchangeEditorScriptLibrary.cpp
--- a/Engine/Plugins/Interchange/Editor/Source/InterchangeEditor/Private/InterchangeEditorScriptLibrary.cpp
+++ b/Engine/Plugins/Interchange/Editor/Source/InterchangeEditor/Private/InterchangeEditorScriptLibrary.cpp
@@ -178,7 +178,7 @@ namespace UE::InterchangeEditorScriptLibrary::Private
 				}
 			}
 
-			return SceneImportAsset->GetFactoryNode((UObject*)(ActorToReset));
+			return SceneImportAsset->GetFactoryNodeFromPath(FSoftObjectPath(ActorToReset));
 		}
 
 		void SetWorldAssetForLevelInstanceActor(ALevelInstance* LevelInstanceActor, UWorld* ReferencedWorld)
diff --git a/Engine/Source/Runtime/Interchange/Engine/Private/InterchangeImportCommon.cpp b/Engine/Source/Runtime/Interchange/Engine/Private/InterchangeImportCommon.cpp
--- a/Engine/Source/Runtime/Interchange/Engine/Private/InterchangeImportCommon.cpp
+++ b/Engine/Source/Runtime/Interchange/Engine/Private/InterchangeImportCommon.cpp
@@ -425,7 +425,7 @@ namespace UE::Interchange
 #if WITH_EDITORONLY_DATA
 		if (UInterchangeSceneImportAsset* SceneImportAsset = Cast<UInterchangeSceneImportAsset>(ObjectToReimport))
 		{
-			return SceneImportAsset->GetFactoryNode(PackageName, AssetName, SubPathString);
+			return SceneImportAsset->GetFactoryNodeFromPath(PackageName, AssetName, SubPathString);
 		}
 #endif
 
diff --git a/Engine/Source/Runtime/Interchange/Engine/Private/InterchangeSceneImportAsset.cpp b/Engine/Source/Runtime/Interchange/Engine/Private/InterchangeSceneImportAsset.cpp
--- a/Engine/Source/Runtime/Interchange/Engine/Private/InterchangeSceneImportAsset.cpp
+++ b/Engine/Source/Runtime/Interchange/Engine/Private/InterchangeSceneImportAsset.cpp
@@ -265,17 +265,17 @@ UObject* UInterchangeSceneImportAsset::GetSceneObject(const FString& PackageName
 	return nullptr;
 }
 
-const UInterchangeFactoryBaseNode* UInterchangeSceneImportAsset::GetFactoryNode(const FString& PackageName, const FString& AssetName, const FString& SubPathString) const
+const UInterchangeFactoryBaseNode* UInterchangeSceneImportAsset::GetFactoryNodeFromPath(const FString& PackageName, const FString& AssetName, const FString& SubPathString) const
 {
 #if WITH_EDITORONLY_DATA
 	const FSoftObjectPath ObjectPath(FTopLevelAssetPath(FName(PackageName), FName(AssetName)), SubPathString);
-	return GetFactoryNode(ObjectPath);
+	return GetFactoryNodeFromPath(ObjectPath);
 #else
 	return nullptr;
 #endif
 }
 
-const UInterchangeFactoryBaseNode* UInterchangeSceneImportAsset::GetFactoryNode(const FSoftObjectPath& SoftObjectPath) const
+const UInterchangeFactoryBaseNode* UInterchangeSceneImportAsset::GetFactoryNodeFromPath(const FSoftObjectPath& SoftObjectPath) const
 {
 #if WITH_EDITORONLY_DATA
 	if (!ensure(AssetImportData))
@@ -298,7 +298,7 @@ const UInterchangeFactoryBaseNode* UInterchangeSceneImportAsset::GetFactoryNode(
 			if (SceneObjectSubPathString.Contains(SoftObjectPathSubPathString)
 				&& FActorSpawnUtils::GetBaseName(*SceneObjectSubPathString) == ObjectSubPathBaseName)
 			{
-				return Cast<UInterchangeFactoryBaseNode>(AssetImportData->GetStoredNode(SceneObjectPair.Value));
+				return AssetImportData->GetStoredFactoryNode(SceneObjectPair.Value);
 			}
 		}
 	}
@@ -306,6 +306,21 @@ const UInterchangeFactoryBaseNode* UInterchangeSceneImportAsset::GetFactoryNode(
 
 	return nullptr;
 }
+
+const UInterchangeFactoryBaseNode* UInterchangeSceneImportAsset::GetFactoryNodeFromUniqueID(const FString& InNodeUniqueId) const
+{
+#if WITH_EDITORONLY_DATA
+	if (!ensure(AssetImportData))
+	{
+		return nullptr;
+	}
+
+	return AssetImportData->GetStoredFactoryNode(InNodeUniqueId);
+#else
+	return nullptr;
+#endif
+}
+
 TArray<const UInterchangeFactoryBaseNode*> UInterchangeSceneImportAsset::GetFactoryNodesOfClass(const UClass* Class) const
 {
 #if WITH_EDITORONLY_DATA
diff --git a/Engine/Source/Runtime/Interchange/Engine/Public/InterchangeSceneImportAsset.h b/Engine/Source/Runtime/Interchange/Engine/Public/InterchangeSceneImportAsset.h
--- a/Engine/Source/Runtime/Interchange/Engine/Public/InterchangeSceneImportAsset.h
+++ b/Engine/Source/Runtime/Interchange/Engine/Public/InterchangeSceneImportAsset.h
@@ -66,6 +66,18 @@ public:
 	 */
 	UE_API UObject* GetSceneObject(const FString& PackageName, const FString& AssetName, const FString& SubPathString = FString()) const;
 
+	UE_DEPRECATED(5.8, "This function is deprecated. Use the GetFactoryNodeFromPath() instead.")
+	const UInterchangeFactoryBaseNode* GetFactoryNode(const FString& PackageName, const FString& AssetName, const FString& SubPathString = FString()) const
+	{
+		return GetFactoryNodeFromPath(PackageName, AssetName, SubPathString);
+	}
+
+	UE_DEPRECATED(5.8, "This function is deprecated. Use the GetFactoryNodeFromPath() instead.")
+	const UInterchangeFactoryBaseNode* GetFactoryNode(const FSoftObjectPath& ObjectPath) const
+	{
+		return GetFactoryNodeFromPath(ObjectPath);
+	}
+
 	/**
 	 * Returns the factory node associated with the asset which path is '//PackageName.AssetName[:SubPathString]'.
 	 * Returns nullptr if the asset which path is '//PackageName.AssetName[:SubPathString]' was not part of
@@ -74,7 +86,16 @@ public:
 	 * @param AssetName: Asset name of the actual object to reimport
 	 * @param SubPathString: Optional subobject name
 	 */
-	UE_API const UInterchangeFactoryBaseNode* GetFactoryNode(const FString& PackageName, const FString& AssetName, const FString& SubPathString = FString()) const;
+	UE_API const UInterchangeFactoryBaseNode* GetFactoryNodeFromPath(const FString& PackageName, const FString& AssetName, const FString& SubPathString = FString()) const;
+	
+	UE_API const UInterchangeFactoryBaseNode* GetFactoryNodeFromPath(const FSoftObjectPath& ObjectPath) const;
+
+	/**
+	 * Returns the factory node associated with a Node UID.
+	 * @param InNodeUniqueId: Unique ID of the Node
+	 * @return Factory Node if available or nullptr. 
+	 */
+	UE_API const UInterchangeFactoryBaseNode* GetFactoryNodeFromUniqueID(const FString& InNodeUniqueId) const;
 
 	/**
 	 * Returns the array of factory nodes corresponding to the given class.
@@ -82,8 +103,6 @@ public:
 	 */
 	UE_API TArray<const UInterchangeFactoryBaseNode*> GetFactoryNodesOfClass(const UClass* Class) const;
 
-	UE_API const UInterchangeFactoryBaseNode* GetFactoryNode(const FSoftObjectPath& ObjectPath) const;
-
 	UE_API void GetSceneSoftObjectPaths(TArray<FSoftObjectPath>& SoftObjectPaths) const;
 
 	UE_API bool ContainsObject(const FSoftObjectPath& SoftObjectPath) const;
