[Interchange Analytics] Get Per Asset Class  Imported Object Count

#rnx

#Synced-CL 50825332

diff --git a/Engine/Plugins/Experimental/Enterprise/DatasmithInterchange/Source/DatasmithInterchange/Private/InterchangeDatasmithTranslator.cpp b/Engine/Plugins/Experimental/Enterprise/DatasmithInterchange/Source/DatasmithInterchange/Private/InterchangeDatasmithTranslator.cpp
--- a/Engine/Plugins/Experimental/Enterprise/DatasmithInterchange/Source/DatasmithInterchange/Private/InterchangeDatasmithTranslator.cpp
+++ b/Engine/Plugins/Experimental/Enterprise/DatasmithInterchange/Source/DatasmithInterchange/Private/InterchangeDatasmithTranslator.cpp
@@ -748,7 +748,7 @@ TOptional<UE::Interchange::FImportImage> UInterchangeDatasmithTranslator::GetTex
 		return TOptional<UE::Interchange::FImportImage>();
 	}
 
-	UE::Interchange::Private::FScopedTranslator ScopedTranslator(TextureElement->GetFile(), Results, AnalyticsHelper);
+	UE::Interchange::Private::FScopedTranslator ScopedTranslator(TextureElement->GetFile(), Results, AnalyticsHandler);
 	const IInterchangeTexturePayloadInterface* TextureTranslator = ScopedTranslator.GetPayLoadInterface< IInterchangeTexturePayloadInterface >();
 
 	if (!ensure(TextureTranslator))
@@ -769,7 +769,7 @@ TOptional<UE::Interchange::FImportLightProfile> UInterchangeDatasmithTranslator:
 	}
 
 	UInterchangeSourceData* PayloadSourceData = UInterchangeManager::GetInterchangeManager().CreateSourceData(PayloadKey);
-	UE::Interchange::Private::FScopedTranslator ScopedTranslator(PayloadKey, Results, AnalyticsHelper);
+	UE::Interchange::Private::FScopedTranslator ScopedTranslator(PayloadKey, Results, AnalyticsHandler);
 	const IInterchangeTextureLightProfilePayloadInterface* TextureTranslator = ScopedTranslator.GetPayLoadInterface< IInterchangeTextureLightProfilePayloadInterface >();
 	if (!ensure(TextureTranslator))
 	{
diff --git a/Engine/Plugins/Interchange/Extensions/OpenUSD/Source/Import/Private/InterchangeUsdTranslator.cpp b/Engine/Plugins/Interchange/Extensions/OpenUSD/Source/Import/Private/InterchangeUsdTranslator.cpp
--- a/Engine/Plugins/Interchange/Extensions/OpenUSD/Source/Import/Private/InterchangeUsdTranslator.cpp
+++ b/Engine/Plugins/Interchange/Extensions/OpenUSD/Source/Import/Private/InterchangeUsdTranslator.cpp
@@ -30,7 +30,7 @@
 #include "UsdWrappers/UsdTyped.h"
 #include "UsdWrappers/UsdVariantSets.h"
 
-#include "InterchangeAnalyticsHelper.h"
+#include "InterchangeAnalyticsHandlerBase.h"
 #include "InterchangeManager.h"
 #include "InterchangeMaterialInstanceNode.h"
 #include "InterchangeMeshNode.h"
@@ -127,10 +127,10 @@ namespace UE::InterchangeUsdTranslator::Private
 #endif	  // USE_USD_SDK
 	}
 
-	void AddAnalytics(UE::Interchange::FAnalyticsHelper* AnalyticsHelper, UInterchangeUsdContext* Context)
+	void AddAnalytics(TObjectPtr<UInterchangeAnalyticsHandlerBase> AnalyticsHandler, UInterchangeUsdContext* Context)
 	{
 #if USE_USD_SDK
-		if (!AnalyticsHelper || !Context)
+		if (!AnalyticsHandler || !Context)
 		{
 			return;
 		}
@@ -171,8 +171,8 @@ namespace UE::InterchangeUsdTranslator::Private
 		// Clear before adding the events because we may be doing multiple translations on the same import while tweaking the
 		// translator settings or etc., and we don't want to pool duplicate attributes onto the same event
 		const static FString BaseEventIdentifier = TEXT("Interchange.Usage.Import.USD");
-		AnalyticsHelper->ClearAnalyticsEventData(BaseEventIdentifier);
-		AnalyticsHelper->AppendThreadSafe(BaseEventIdentifier, Attributes);
+		AnalyticsHandler->Clear(BaseEventIdentifier);
+		AnalyticsHandler->Append(BaseEventIdentifier, Attributes);
 
 		if (UE::FUsdStage Stage = Context->GetUsdStage())
 		{
@@ -180,8 +180,8 @@ namespace UE::InterchangeUsdTranslator::Private
 			UsdUtils::CollectSchemaAnalytics(Stage, CustomSchemaAttributes);
 
 			const static FString CustomSchemaEventIdentifier = TEXT("Interchange.Usage.USD.CustomSchemaCount");
-			AnalyticsHelper->ClearAnalyticsEventData(CustomSchemaEventIdentifier);
-			AnalyticsHelper->AppendThreadSafe(CustomSchemaEventIdentifier, CustomSchemaAttributes);
+			AnalyticsHandler->Clear(CustomSchemaEventIdentifier);
+			AnalyticsHandler->Append(CustomSchemaEventIdentifier, CustomSchemaAttributes);
 		}
 #endif	  // USE_USD_SDK
 	}
@@ -514,7 +514,7 @@ bool UInterchangeUSDTranslator::Translate(UInterchangeBaseNodeContainer& NodeCon
 		TranslatePrimSubtree(StageToImport.GetPseudoRoot(), Info);
 	}
 
-	AddAnalytics(AnalyticsHelper.Get(), Context);
+	AddAnalytics(AnalyticsHandler, Context);
 
 	return true;
 #else
diff --git a/Engine/Plugins/Interchange/Extensions/OpenUSD/Source/Import/Private/SchemaHandlers/MaterialXSchemaHandler.cpp b/Engine/Plugins/Interchange/Extensions/OpenUSD/Source/Import/Private/SchemaHandlers/MaterialXSchemaHandler.cpp
--- a/Engine/Plugins/Interchange/Extensions/OpenUSD/Source/Import/Private/SchemaHandlers/MaterialXSchemaHandler.cpp
+++ b/Engine/Plugins/Interchange/Extensions/OpenUSD/Source/Import/Private/SchemaHandlers/MaterialXSchemaHandler.cpp
@@ -251,7 +251,7 @@ namespace UE::Interchange::USD
 				PayloadKey,
 				AlternateTexturePath,
 				UsdContext.GetResultsContainer(),
-				UsdContext.GetTranslator()->AnalyticsHelper
+				UsdContext.GetTranslator()->AnalyticsHandler
 			);
 		}
 
@@ -290,7 +290,7 @@ namespace UE::Interchange::USD
 				PayloadKey,
 				AlternateTexturePath,
 				UsdContext.GetResultsContainer(),
-				UsdContext.GetTranslator()->AnalyticsHelper
+				UsdContext.GetTranslator()->AnalyticsHandler
 			);
 		}
 
diff --git a/Engine/Plugins/Interchange/Extensions/OpenUSD/Source/Import/Private/SchemaHandlers/UniversalMaterialSchemaHandler.cpp b/Engine/Plugins/Interchange/Extensions/OpenUSD/Source/Import/Private/SchemaHandlers/UniversalMaterialSchemaHandler.cpp
--- a/Engine/Plugins/Interchange/Extensions/OpenUSD/Source/Import/Private/SchemaHandlers/UniversalMaterialSchemaHandler.cpp
+++ b/Engine/Plugins/Interchange/Extensions/OpenUSD/Source/Import/Private/SchemaHandlers/UniversalMaterialSchemaHandler.cpp
@@ -456,7 +456,7 @@ namespace UE::Interchange::USD
 		UE::Interchange::Private::FScopedTranslator ScopedTranslator{
 			FilePath,
 			UsdContext.GetResultsContainer(),
-			UsdContext.GetTranslator()->AnalyticsHelper
+			UsdContext.GetTranslator()->AnalyticsHandler
 		};
 		const IInterchangeTexturePayloadInterface* TextureTranslator = ScopedTranslator.GetPayLoadInterface<IInterchangeTexturePayloadInterface>();
 		if (TextureTranslator)
@@ -545,7 +545,7 @@ namespace UE::Interchange::USD
 			UE::Interchange::Private::FScopedTranslator ScopedTranslator(
 				TileFilePath,
 				UsdContext.GetResultsContainer(),
-				UsdContext.GetTranslator()->AnalyticsHelper
+				UsdContext.GetTranslator()->AnalyticsHandler
 			);
 			const IInterchangeTexturePayloadInterface* TextureTranslator = ScopedTranslator.GetPayLoadInterface<IInterchangeTexturePayloadInterface>(
 			);
diff --git a/Engine/Plugins/Interchange/Runtime/Interchange.uplugin b/Engine/Plugins/Interchange/Runtime/Interchange.uplugin
--- a/Engine/Plugins/Interchange/Runtime/Interchange.uplugin
+++ b/Engine/Plugins/Interchange/Runtime/Interchange.uplugin
@@ -66,6 +66,11 @@
 			"Type": "Runtime",
 			"LoadingPhase": "Default"
 		},
+		{
+			"Name": "InterchangeAnalytics",
+			"Type": "Runtime",
+			"LoadingPhase": "Default"
+		},
 		{
 			"Name": "InterchangeCommonParser",
 			"Type": "Runtime",
diff --git a/Engine/Plugins/Interchange/Runtime/Source/Analytics/InterchangeAnalytics.Build.cs b/Engine/Plugins/Interchange/Runtime/Source/Analytics/InterchangeAnalytics.Build.cs
new file mode 100644
--- /dev/null
+++ b/Engine/Plugins/Interchange/Runtime/Source/Analytics/InterchangeAnalytics.Build.cs
@@ -0,0 +1,33 @@
+// Copyright Epic Games, Inc. All Rights Reserved.
+
+namespace UnrealBuildTool.Rules
+{
+	public class InterchangeAnalytics : ModuleRules
+	{
+		public InterchangeAnalytics(ReadOnlyTargetRules Target) : base(Target)
+		{
+			PublicDependencyModuleNames.AddRange(
+				new string[]
+				{
+					"Core",
+					"CoreUObject",
+					"Engine",
+					"InterchangeCore",
+					"InterchangeCommon",
+					"InterchangeEngine",
+					"InterchangeFactoryNodes",
+					"InterchangeNodes",
+				}
+			);
+
+			PrivateDependencyModuleNames.AddRange(
+				new string[]
+				{
+					"GeometryCache",
+					"HairStrandsCore",
+					"LevelSequence",
+				}
+			);
+		}
+	}
+}
diff --git a/Engine/Plugins/Interchange/Runtime/Source/Analytics/Private/InterchangeAnalyticsAssetTypeTracker.cpp b/Engine/Plugins/Interchange/Runtime/Source/Analytics/Private/InterchangeAnalyticsAssetTypeTracker.cpp
new file mode 100644
--- /dev/null
+++ b/Engine/Plugins/Interchange/Runtime/Source/Analytics/Private/InterchangeAnalyticsAssetTypeTracker.cpp
@@ -0,0 +1,51 @@
+// Copyright Epic Games, Inc. All Rights Reserved.
+
+#include "InterchangeAnalyticsAssetTypeTracker.h"
+
+TSet<UClass*> FInterchangeAnalyticsAssetTypeTracker::RegisteredAssetClasses;
+TMap<UClass*, FString> FInterchangeAnalyticsAssetTypeTracker::AssetClassPathToDisplayNameMap;
+
+void FInterchangeAnalyticsAssetTypeTracker::RegisterAssetType(UClass* AssetClass, const FString& DisplayName)
+{
+	if (!AssetClass)
+	{
+		return;
+	}
+
+	if (!RegisteredAssetClasses.Contains(AssetClass))
+	{
+		RegisteredAssetClasses.Emplace(AssetClass);
+		AssetClassPathToDisplayNameMap.Emplace(AssetClass, DisplayName);
+	}
+}
+
+TMap<FString, int32> FInterchangeAnalyticsAssetTypeTracker::GetAssetTypeFrequenceMap(const TArray<TObjectPtr<UObject>>& InObjects, bool bIncludeUntrackedClasses /*= false*/)
+{
+	TMap<FString, int32> FrequncyMap;
+	AppendAssetTypeFrequenceMap(InObjects, FrequncyMap, bIncludeUntrackedClasses);
+	return FrequncyMap;
+}
+
+void FInterchangeAnalyticsAssetTypeTracker::AppendAssetTypeFrequenceMap(const TArray<TObjectPtr<UObject>>& InObjects, TMap<FString, int32>& OutFreqMap, bool bIncludeUntrackedClasses /*= false*/)
+{
+	for (const TObjectPtr<UObject>& Object : InObjects)
+	{
+		bool bAssetClassFound = false;
+		for (const UClass* AssetClass : RegisteredAssetClasses)
+		{
+			if (Object->GetClass()->IsChildOf(AssetClass))
+			{
+				int32& Frequncy = OutFreqMap.FindOrAdd(AssetClassPathToDisplayNameMap.FindChecked(AssetClass));
+				Frequncy++;
+				bAssetClassFound = true;
+				break;
+			}
+		}
+
+		if (!bAssetClassFound && bIncludeUntrackedClasses)
+		{
+			int32& Frequency = OutFreqMap.FindOrAdd(TEXT("UntrackedAssets"));
+			Frequency++;
+		}
+	}
+}
diff --git a/Engine/Plugins/Interchange/Runtime/Source/Analytics/Private/InterchangeAnalyticsHandlerDefault.cpp b/Engine/Plugins/Interchange/Runtime/Source/Analytics/Private/InterchangeAnalyticsHandlerDefault.cpp
new file mode 100644
--- /dev/null
+++ b/Engine/Plugins/Interchange/Runtime/Source/Analytics/Private/InterchangeAnalyticsHandlerDefault.cpp
@@ -0,0 +1,318 @@
+// Copyright Epic Games, Inc. All Rights Reserved.
+
+#include "InterchangeAnalyticsHandlerDefault.h"
+#include "EngineAnalytics.h"
+
+#include "InterchangePipelineBase.h"
+#include "InterchangeSourceData.h"
+#include "Nodes/InterchangeSourceNode.h"
+
+#include "Logging/LogCategory.h"
+
+#include "Policies/CondensedJsonPrintPolicy.h"
+#include "Serialization/JsonSerializer.h"
+#include "Serialization/JsonWriter.h"
+#include "Serialization/JsonSerializerWriter.h"
+
+#include "InterchangeAnalyticsAssetTypeTracker.h"
+
+DEFINE_LOG_CATEGORY_STATIC(LogInterchangeAnalytics, Log, All);
+
+namespace UE::Interchange::AnalyticsHelpers
+{
+	// Json writer subclass to allow us to avoid using a SharedPtr to write basic Json.
+	class FAnalyticsJsonWriter : public TJsonStringWriter<TCondensedJsonPrintPolicy<TCHAR>>
+	{
+	public:
+		explicit FAnalyticsJsonWriter(FString* Out) : TJsonStringWriter<TCondensedJsonPrintPolicy<TCHAR>>(Out, 0)
+		{
+		}
+	};
+
+	FJsonFragment GetWarningsErrorsMapAsJsonFragment(const TMap<FString, int32>& InFreqMap)
+	{
+		FString ReturnValue;
+		FAnalyticsJsonWriter JsonWriter(&ReturnValue);
+		JsonWriter.WriteArrayStart();
+		for (const TPair<FString, int32>& FreqMapPair : InFreqMap)
+		{
+			JsonWriter.WriteObjectStart();
+			JsonWriter.WriteValue(TEXT("MessageKey"), FreqMapPair.Key);
+			JsonWriter.WriteValue(TEXT("MessageCount"), FreqMapPair.Value);
+			JsonWriter.WriteObjectEnd();
+		}
+		JsonWriter.WriteArrayEnd();
+		JsonWriter.Close();
+		return FJsonFragment(MoveTemp(ReturnValue));
+	}
+
+	FJsonFragment GetPrimarySourceDataMetadataAsJsonFragment(const UInterchangeSourceData* InSourceData, const UInterchangeSourceNode* InSourceNode)
+	{
+		FString ReturnValue;
+		FAnalyticsJsonWriter JsonWriter(&ReturnValue);
+
+		TMap<FString, FString> ExtraInfoMap;
+		InSourceNode->GetExtraInformation(ExtraInfoMap);
+
+		const FString DefaultValue("N/A");
+
+		const FString ApplicationVendor = ExtraInfoMap.FindRef(FSourceNodeExtraInfoStaticData::GetApplicationVendorExtraInfoKey(), DefaultValue);
+		const FString ApplicationName = ExtraInfoMap.FindRef(FSourceNodeExtraInfoStaticData::GetApplicationNameExtraInfoKey(), DefaultValue);
+		const FString ApplicationVersion = ExtraInfoMap.FindRef(FSourceNodeExtraInfoStaticData::GetApplicationVersionExtraInfoKey(), DefaultValue);
+
+		JsonWriter.WriteObjectStart();
+		JsonWriter.WriteValue(TEXT("SourceExtension"), FPaths::GetExtension(InSourceData->GetFilename()));
+		JsonWriter.WriteValue(TEXT("ApplicationVendor"), ApplicationVendor);
+		JsonWriter.WriteValue(TEXT("ApplicationName"), ApplicationName);
+		JsonWriter.WriteValue(TEXT("ApplicationVersion"), ApplicationVersion);
+		JsonWriter.WriteObjectEnd();
+		JsonWriter.Close();
+		return FJsonFragment(MoveTemp(ReturnValue));
+	}
+
+	FJsonFragment GetAssetTypeFrequencyMapAsJsonFragment(const TMap<FString, int32>& InFreqMap)
+	{
+		FString ReturnValue;
+		FAnalyticsJsonWriter JsonWriter(&ReturnValue);
+
+		JsonWriter.WriteObjectStart();
+		for (const TPair<FString, int32>& AssetTypeFreqPair : InFreqMap)
+		{
+			JsonWriter.WriteValue(AssetTypeFreqPair.Key, AssetTypeFreqPair.Value);
+		}
+		JsonWriter.WriteObjectEnd();
+		JsonWriter.Close();
+		return FJsonFragment(MoveTemp(ReturnValue));
+	}
+
+	bool ExtractNamespace(const FText& Text, FString& OutTextNamespaceId)
+	{
+		FText TextToUse = Text;
+
+		TArray<FHistoricTextFormatData> TextHistory;
+		FTextInspector::GetHistoricFormatData(Text, TextHistory);
+		if (TextHistory.Num() > 0)
+		{
+			const FHistoricTextFormatData& FmtData = TextHistory[0];
+			TextToUse = FmtData.SourceFmt.GetSourceText();
+		}
+
+		const TOptional<FString> TextNamespace = FTextInspector::GetNamespace(TextToUse);
+		const TOptional<FString> TextKey = FTextInspector::GetKey(TextToUse);
+
+		if (TextNamespace.IsSet() && TextKey.IsSet())
+		{
+			OutTextNamespaceId = FString::Printf(TEXT("%s_%s"), *TextNamespace.GetValue(), *TextKey.GetValue());
+			return true;
+		}
+
+		OutTextNamespaceId = FString("UnknownError");
+		return false;
+	}
+
+	void CollectResultsContainerWarningAndErrors(const UInterchangeResultsContainer* ResultContainer, TMap<FString, int32>&WarningTypes, TMap<FString, int32>&ErrorTypes)
+	{
+		TArray<UInterchangeResult*> InterchangeResults = ResultContainer->GetResults();
+
+		for (const UInterchangeResult* InterchangeResult : InterchangeResults)
+		{
+			using namespace UE::Interchange::Private;
+
+			switch (InterchangeResult->GetResultType())
+			{
+			case EInterchangeResultType::Success:
+				break;
+			case EInterchangeResultType::Warning:
+			{
+				FString OutWarningAttribValue;
+				if (ExtractNamespace(InterchangeResult->GetText(), OutWarningAttribValue))
+				{
+					int32& Frequency = WarningTypes.FindOrAdd(OutWarningAttribValue);
+					Frequency++;
+				}
+				else
+				{
+					UE_LOG(LogInterchangeAnalytics, Error, TEXT("Failed to extract Analytic Attribute Value from %s"), *(InterchangeResult->GetText().ToString()));
+				}
+				break;
+			}
+			case EInterchangeResultType::Error:
+			{
+				FString OutErrorAttribValue;
+				if (ExtractNamespace(InterchangeResult->GetText(), OutErrorAttribValue))
+				{
+					int32& Frequency = ErrorTypes.FindOrAdd(OutErrorAttribValue);
+					Frequency++;
+				}
+				else
+				{
+					UE_LOG(LogInterchangeAnalytics, Error, TEXT("Failed to extract Analytic Attribute Value from %s"), *(InterchangeResult->GetText().ToString()));
+				}
+				break;
+			}
+			}
+		}
+	}
+}
+
+void UInterchangeAnalyticsHandlerDefault::Execute_Append(const FString& Identifier, const TArray<FAnalyticsEventAttribute>& ToAdd)
+{
+	TArray<FAnalyticsEventAttribute>& Analytics = AnalyticsAttributes.FindOrAdd(Identifier);
+	Analytics.Append(ToAdd);
+}
+
+void UInterchangeAnalyticsHandlerDefault::Execute_Add(const FString& Identifier, const FAnalyticsEventAttribute& Entry)
+{
+	TArray<FAnalyticsEventAttribute>& Analytics = AnalyticsAttributes.FindOrAdd(Identifier);
+	Analytics.Add(Entry);
+}
+
+void UInterchangeAnalyticsHandlerDefault::Execute_Clear(const FString& Identifier)
+{
+	AnalyticsAttributes.Remove(Identifier);
+}
+
+void UInterchangeAnalyticsHandlerDefault::Execute_Flush()
+{
+	if (FEngineAnalytics::IsAvailable())
+	{
+		const FAnalyticsEventAttribute InterchangeAnalyticsIDAttribute = GetAnalyticsIDAttribute();
+
+		for (TPair<FString, TArray<FAnalyticsEventAttribute>>& AnalyticsPerIdentifier : AnalyticsAttributes)
+		{
+			TArray<FAnalyticsEventAttribute> ModifiedAttributes;
+			ModifiedAttributes.Reserve(AnalyticsPerIdentifier.Value.Num() + 1);
+			ModifiedAttributes = MoveTemp(AnalyticsPerIdentifier.Value);
+			ModifiedAttributes.Add(InterchangeAnalyticsIDAttribute);
+
+			FEngineAnalytics::GetProvider().RecordEvent(AnalyticsPerIdentifier.Key, ModifiedAttributes);
+		}
+
+		AnalyticsAttributes.Empty();
+	}
+}
+
+void UInterchangeAnalyticsHandlerDefault::SendImmediate(const FString& Identifier, TArray<FAnalyticsEventAttribute> Attributes) const
+{
+	if (FEngineAnalytics::IsAvailable())
+	{
+		Attributes.Add(GetAnalyticsIDAttribute());
+		FEngineAnalytics::GetProvider().RecordEvent(Identifier, Attributes);
+	}
+}
+
+void UInterchangeAnalyticsHandlerDefault::SendPipelineAnalytics(UInterchangePipelineBase& Pipeline, const int32 AsyncHelperUniqueId, const FString& ParentPipeline) const
+{
+	int32 PortFlags = 0;
+	UClass* Class = Pipeline.GetClass();
+	FString PipelineChainName = ParentPipeline.IsEmpty() ? Pipeline.GetName() : ParentPipeline + TEXT(".") + Pipeline.GetName();
+
+	TArray<FAnalyticsEventAttribute> PipelineAttribs;
+	PipelineAttribs.Add(FAnalyticsEventAttribute(TEXT("UniqueId"), AsyncHelperUniqueId));
+	PipelineAttribs.Add(FAnalyticsEventAttribute(TEXT("Name"), PipelineChainName));
+	PipelineAttribs.Add(FAnalyticsEventAttribute(TEXT("Class"), Class->GetName()));
+
+	for (FProperty* Property = Class->PropertyLink; Property; Property = Property->PropertyLinkNext)
+	{
+		if (!Pipeline.CanAddPropertyToAnalytics(Property))
+		{
+			continue;
+		}
+
+		FObjectProperty* SubObject = CastField<FObjectProperty>(Property);
+		UInterchangePipelineBase* SubPipeline = SubObject ? Cast<UInterchangePipelineBase>(SubObject->GetObjectPropertyValue_InContainer(&Pipeline)) : nullptr;
+
+		if (SubPipeline)
+		{
+			// Save the settings if the referenced pipeline is a sub-object of ours
+			if (SubPipeline->IsInOuter(&Pipeline))
+			{
+				//Go recursive with subObject, like if they are part of the same object
+				SendPipelineAnalytics(*SubPipeline, AsyncHelperUniqueId, PipelineChainName);
+			}
+		}
+		else
+		{
+			Pipeline.AddPropertyAnalytics(Property, PipelineAttribs);
+		}
+	}
+
+	SendImmediate(GetPipelineAnalyticsEventIdentifier(), PipelineAttribs);
+}
+
+void UInterchangeAnalyticsHandlerDefault::Send(const TArray<UInterchangePipelineBase*>& Pipelines, const int32 AsyncHelperUniqueId)
+{
+	if (!FEngineAnalytics::IsAvailable())
+	{
+		return;
+	}
+
+	for (UInterchangePipelineBase* Pipeline : Pipelines)
+	{
+		if (!Pipeline)
+		{
+			continue;
+		}
+
+		SendPipelineAnalytics(*Pipeline, AsyncHelperUniqueId, FString());
+	}
+}
+
+void UInterchangeAnalyticsHandlerDefault::Send(const FInterchangeImportResultAnalyticsInfo& ImportResultAnalyticsInfo)
+{
+	if (ImportResultAnalyticsInfo.bCancel || !FEngineAnalytics::IsAvailable())
+	{
+		return;
+	}
+
+	using namespace UE::Interchange::AnalyticsHelpers;
+
+	TArray<FAnalyticsEventAttribute> Attribs;
+	//Set the unique id of this import
+	Attribs.Add(FAnalyticsEventAttribute(TEXT("UniqueId"), ImportResultAnalyticsInfo.UniqueId));
+	Attribs.Add(FAnalyticsEventAttribute(TEXT("IsCanceled"), ImportResultAnalyticsInfo.bCancel));
+
+	{
+		const UInterchangeSourceData* PrimarySourceData = ImportResultAnalyticsInfo.PrimarySourceData; 
+		const UInterchangeSourceNode* PrimarySourceNode = ImportResultAnalyticsInfo.PrimarySourceNode;
+		Attribs.Add(FAnalyticsEventAttribute(TEXT("PrimarySourceDataMetadata"), GetPrimarySourceDataMetadataAsJsonFragment(PrimarySourceData, PrimarySourceNode)));
+	}
+
+	int32 ImportedObjectCount = 0;
+	TMap<FString, int32> AssetTypeFrequencyMap;
+	constexpr bool bIncludeUntrackedAssets = true;
+	for (const TPair<int32, TArray<TObjectPtr<UObject>>>& SourceIndexAndImportedAssets : ImportResultAnalyticsInfo.ImportedAssetsPerSourceIndex)
+	{
+		ImportedObjectCount += SourceIndexAndImportedAssets.Value.Num();
+		FInterchangeAnalyticsAssetTypeTracker::AppendAssetTypeFrequenceMap(SourceIndexAndImportedAssets.Value, AssetTypeFrequencyMap, bIncludeUntrackedAssets);
+	}
+	
+	for (const TPair<int32, TArray<TObjectPtr<UObject>>>& SourceIndexAndImportedSceneObjects : ImportResultAnalyticsInfo.ImportedSceneObjectsPerSourceIndex)
+	{
+		ImportedObjectCount += SourceIndexAndImportedSceneObjects.Value.Num();
+		FInterchangeAnalyticsAssetTypeTracker::AppendAssetTypeFrequenceMap(SourceIndexAndImportedSceneObjects.Value, AssetTypeFrequencyMap, bIncludeUntrackedAssets);
+	}
+
+	Attribs.Add(FAnalyticsEventAttribute(TEXT("ImportObjectCount"), ImportedObjectCount));
+	Attribs.Add(FAnalyticsEventAttribute(TEXT("ImportObjectDetails"), GetAssetTypeFrequencyMapAsJsonFragment(AssetTypeFrequencyMap)));
+
+	//Report any warning or error message
+	{
+		TMap<FString, int32> WarningTypes;
+		TMap<FString, int32> ErrorTypes;
+		if (ImportResultAnalyticsInfo.AssetImportResultsContainer.IsSet())
+		{
+			CollectResultsContainerWarningAndErrors(ImportResultAnalyticsInfo.AssetImportResultsContainer.GetValue(), WarningTypes, ErrorTypes);
+		}
+		
+		if (ImportResultAnalyticsInfo.SceneImportResultsContainer.IsSet())
+		{
+			CollectResultsContainerWarningAndErrors(ImportResultAnalyticsInfo.SceneImportResultsContainer.GetValue(), WarningTypes, ErrorTypes);
+		}
+
+		Attribs.Add(FAnalyticsEventAttribute(TEXT("WarningTypes"), GetWarningsErrorsMapAsJsonFragment(WarningTypes)));
+		Attribs.Add(FAnalyticsEventAttribute(TEXT("ErrorTypes"), GetWarningsErrorsMapAsJsonFragment(ErrorTypes)));
+	}
+	
+	SendImmediate(GetImportResultEventIdentifier(), Attribs);
+}
diff --git a/Engine/Plugins/Interchange/Runtime/Source/Analytics/Private/InterchangeAnalyticsModule.cpp b/Engine/Plugins/Interchange/Runtime/Source/Analytics/Private/InterchangeAnalyticsModule.cpp
new file mode 100644
--- /dev/null
+++ b/Engine/Plugins/Interchange/Runtime/Source/Analytics/Private/InterchangeAnalyticsModule.cpp
@@ -0,0 +1,75 @@
+// Copyright Epic Games, Inc. All Rights Reserved.
+
+#include "InterchangeAnalyticsModule.h"
+
+#include "Engine/Engine.h"
+#include "Misc/CoreDelegates.h"
+
+#include "InterchangeAnalyticsHandlerDefault.h"
+#include "InterchangeManager.h"
+#include "InterchangeAnalyticsAssetTypeTracker.h"
+
+#include "Animation/AnimSequence.h"
+#include "Animation/Skeleton.h"
+#include "Engine/Blueprint.h"
+#include "Engine/SkeletalMesh.h"
+#include "Engine/StaticMesh.h"
+#include "Engine/Texture.h"
+#include "Engine/World.h"
+#include "GameFramework/Actor.h"
+#include "GeometryCache.h"
+#include "GroomAsset.h"
+#include "LevelSequence.h"
+#include "Materials/MaterialInterface.h"
+#include "Sound/SoundWave.h"
+
+
+class FInterchangeAnalyticsModule : public IInterchangeAnalyticsModule
+{
+	/** IModuleInterface implementation */
+	virtual void StartupModule() override 
+	{
+		auto RegisterAnalyticsClass = []() 
+		{
+			UInterchangeManager& InterchangeManager = UInterchangeManager::GetInterchangeManager();
+			InterchangeManager.SetAnalyticsHandlerClass(UInterchangeAnalyticsHandlerDefault::StaticClass());
+
+			// Track Asset Types with Interchange Analytics
+			FInterchangeAnalyticsAssetTypeTracker::RegisterAssetType(UStaticMesh::StaticClass(), TEXT("StaticMeshes"));
+			FInterchangeAnalyticsAssetTypeTracker::RegisterAssetType(USkeletalMesh::StaticClass(), TEXT("SkeletalMeshes"));
+			FInterchangeAnalyticsAssetTypeTracker::RegisterAssetType(UMaterialInterface::StaticClass(), TEXT("Materials"));
+			FInterchangeAnalyticsAssetTypeTracker::RegisterAssetType(UAnimSequence::StaticClass(), TEXT("Animations"));
+			FInterchangeAnalyticsAssetTypeTracker::RegisterAssetType(USkeleton::StaticClass(), TEXT("Skeletons"));
+			FInterchangeAnalyticsAssetTypeTracker::RegisterAssetType(USoundWave::StaticClass(), TEXT("Sounds"));
+			FInterchangeAnalyticsAssetTypeTracker::RegisterAssetType(UTexture::StaticClass(), TEXT("Textures"));
+
+			FInterchangeAnalyticsAssetTypeTracker::RegisterAssetType(UGroomAsset::StaticClass(), TEXT("Grooms"));
+			FInterchangeAnalyticsAssetTypeTracker::RegisterAssetType(UGeometryCache::StaticClass(), TEXT("GeometryCaches"));
+			
+			FInterchangeAnalyticsAssetTypeTracker::RegisterAssetType(AActor::StaticClass(), TEXT("Actors"));
+			FInterchangeAnalyticsAssetTypeTracker::RegisterAssetType(ULevelSequence::StaticClass(), TEXT("LevelSequences"));
+			
+			FInterchangeAnalyticsAssetTypeTracker::RegisterAssetType(UWorld::StaticClass(), TEXT("Levels"));
+			FInterchangeAnalyticsAssetTypeTracker::RegisterAssetType(UBlueprint::StaticClass(), TEXT("PackedLevelActorBlueprints"));
+		};
+
+		if (GEngine)
+		{
+			RegisterAnalyticsClass();
+		}
+		else
+		{
+			PostEngineInitHandle = FCoreDelegates::OnPostEngineInit.AddLambda(RegisterAnalyticsClass);
+		}
+	}
+
+	virtual void ShutdownModule() override 
+	{
+		FCoreDelegates::OnPostEngineInit.Remove(PostEngineInitHandle);
+	}
+
+private:
+	FDelegateHandle PostEngineInitHandle;
+};
+
+IMPLEMENT_MODULE(FInterchangeAnalyticsModule, InterchangeAnalytics)
diff --git a/Engine/Plugins/Interchange/Runtime/Source/Analytics/Public/InterchangeAnalyticsAssetTypeTracker.h b/Engine/Plugins/Interchange/Runtime/Source/Analytics/Public/InterchangeAnalyticsAssetTypeTracker.h
new file mode 100644
--- /dev/null
+++ b/Engine/Plugins/Interchange/Runtime/Source/Analytics/Public/InterchangeAnalyticsAssetTypeTracker.h
@@ -0,0 +1,29 @@
+// Copyright Epic Games, Inc. All Rights Reserved.
+
+#pragma once
+#include "UObject/Class.h"
+#include "UObject/ObjectPtr.h"
+
+#include "Containers/Array.h"
+#include "Containers/Map.h"
+#include "Containers/Set.h"
+#include "Containers/UnrealString.h"
+
+#define UE_API INTERCHANGEANALYTICS_API
+
+class FInterchangeAnalyticsAssetTypeTracker
+{
+public:
+	UE_API static void RegisterAssetType(UClass* AssetClass, const FString& DisplayName);
+
+	UE_API static TMap<FString, int32> GetAssetTypeFrequenceMap(const TArray<TObjectPtr<UObject>>& InObjects, bool bIncludeUntrackedClasses = false);
+
+	UE_API static void AppendAssetTypeFrequenceMap(const TArray<TObjectPtr<UObject>>& InObjects, TMap<FString, int32>& OutFreqMap, bool bIncludeUntrackedClasses = false);
+
+private:
+    static TSet<UClass*> RegisteredAssetClasses;
+    static TMap<UClass*, FString> AssetClassPathToDisplayNameMap;
+};
+
+
+#undef UE_API
\ No newline at end of file
diff --git a/Engine/Plugins/Interchange/Runtime/Source/Analytics/Public/InterchangeAnalyticsHandlerDefault.h b/Engine/Plugins/Interchange/Runtime/Source/Analytics/Public/InterchangeAnalyticsHandlerDefault.h
new file mode 100644
--- /dev/null
+++ b/Engine/Plugins/Interchange/Runtime/Source/Analytics/Public/InterchangeAnalyticsHandlerDefault.h
@@ -0,0 +1,52 @@
+// Copyright Epic Games, Inc. All Rights Reserved.
+
+#pragma once
+
+#include "UObject/Object.h"
+#include "InterchangeAnalyticsHandlerBase.h"
+
+#include "Containers/Array.h"
+#include "Containers/Map.h"
+#include "Containers/UnrealString.h"
+
+#include "Misc/Guid.h"
+#include "AnalyticsEventAttribute.h"
+
+#include "InterchangeAnalyticsHandlerDefault.generated.h"
+
+#define UE_API INTERCHANGEANALYTICS_API
+
+class UInterchangePipelineBase;
+
+UCLASS(BlueprintType, Blueprintable, MinimalAPI)
+class UInterchangeAnalyticsHandlerDefault : public UInterchangeAnalyticsHandlerBase
+{
+    GENERATED_BODY()
+    
+public:
+	virtual void Send(const TArray<UInterchangePipelineBase*>& Pipelines, const int32 AsyncHelperUniqueId) override;
+
+	virtual void Send(const FInterchangeImportResultAnalyticsInfo& ImportResultAnalyticsInfo) override;
+
+protected:
+	void SendImmediate(const FString& Identifier, TArray<FAnalyticsEventAttribute> Attributes) const;
+
+    virtual void Execute_Append(const FString& Identifier, const TArray<FAnalyticsEventAttribute>& ToAdd) override;
+
+    virtual void Execute_Add(const FString& Identifier, const FAnalyticsEventAttribute& Entry) override;
+
+    virtual void Execute_Clear(const FString& Identifier) override;
+
+	virtual void Execute_Flush() override;
+
+private:
+	void SendPipelineAnalytics(UInterchangePipelineBase& Pipeline, const int32 AsyncHelperUniqueId, const FString& ParentPipeline) const;
+
+protected:
+    //Analytics Attributes to share with the Translators. (to be sent when Import is finished)
+    TMap<FString, TArray<FAnalyticsEventAttribute>> AnalyticsAttributes;
+    FCriticalSection CriticalSection;
+};
+
+
+#undef UE_API
\ No newline at end of file
diff --git a/Engine/Plugins/Interchange/Runtime/Source/Analytics/Public/InterchangeAnalyticsModule.h b/Engine/Plugins/Interchange/Runtime/Source/Analytics/Public/InterchangeAnalyticsModule.h
new file mode 100644
--- /dev/null
+++ b/Engine/Plugins/Interchange/Runtime/Source/Analytics/Public/InterchangeAnalyticsModule.h
@@ -0,0 +1,35 @@
+// Copyright Epic Games, Inc. All Rights Reserved.
+
+#pragma once
+
+#include "Modules/ModuleInterface.h"
+#include "Modules/ModuleManager.h"
+
+/**
+ * The public interface to this module
+ */
+class IInterchangeAnalyticsModule : public IModuleInterface
+{
+public:
+	/**
+	 * Singleton-like access to this module's interface.  This is just for convenience!
+	 * Beware of calling this during the shutdown phase, though.  Your module might have been unloaded already.
+	 *
+	 * @return Returns singleton instance, loading the module on demand if needed
+	 */
+	static inline IInterchangeAnalyticsModule& Get()
+	{
+		return FModuleManager::LoadModuleChecked< IInterchangeAnalyticsModule >( "InterchangeAnalytics" );
+	}
+
+	/**
+	 * Checks to see if this module is loaded and ready.  It is only valid to call Get() if IsAvailable() returns true.
+	 *
+	 * @return True if the module is loaded and ready to use
+	 */
+	static inline bool IsAvailable()
+	{
+		return FModuleManager::Get().IsModuleLoaded( "InterchangeAnalytics" );
+	}
+};
+
diff --git a/Engine/Plugins/Interchange/Runtime/Source/Import/Private/Fbx/InterchangeFbxTranslator.cpp b/Engine/Plugins/Interchange/Runtime/Source/Import/Private/Fbx/InterchangeFbxTranslator.cpp
--- a/Engine/Plugins/Interchange/Runtime/Source/Import/Private/Fbx/InterchangeFbxTranslator.cpp
+++ b/Engine/Plugins/Interchange/Runtime/Source/Import/Private/Fbx/InterchangeFbxTranslator.cpp
@@ -6,7 +6,7 @@
 #include "Dom/JsonObject.h"
 #include "Dom/JsonValue.h"
 #include "HAL/PlatformFileManager.h"
-#include "InterchangeAnalyticsHelper.h"
+#include "InterchangeAnalyticsHandlerBase.h"
 #include "InterchangeCommonAnimationPayload.h"
 #include "InterchangeDispatcher.h"
 #include "FBX/InterchangeDispatcherFBXTasks.h"
@@ -73,9 +73,9 @@ namespace UE::Interchange::Private
 		uFBX,
 	};
 
-	void  AddAnalytics (const EFBXTranslatorAnalyticsResult& TranslatorAnalyticsResult, TSharedPtr<UE::Interchange::FAnalyticsHelper> FBXAnalyticsHelper, bool bIsUsingInterchangeWorker = false, EFBXParserType FBXParser = EFBXParserType::Unspecified)
+	void  AddAnalytics (const EFBXTranslatorAnalyticsResult& TranslatorAnalyticsResult, TObjectPtr<UInterchangeAnalyticsHandlerBase> FBXAnalyticsHandler, bool bIsUsingInterchangeWorker = false, EFBXParserType FBXParser = EFBXParserType::Unspecified)
 	{
-		if (!FBXAnalyticsHelper.IsValid())
+		if (!FBXAnalyticsHandler)
 		{
 			return;
 		}
@@ -127,7 +127,7 @@ namespace UE::Interchange::Private
 			break;
 		}
 
-		FBXAnalyticsHelper->AppendThreadSafe(FBXAnalyticsIdentifier, FBXAnalyticsAttribs);
+		FBXAnalyticsHandler->Append(FBXAnalyticsIdentifier, FBXAnalyticsAttribs);
 	}
 
 	void ApplyTranslatorMessage(const UInterchangeFbxTranslator* Translator, const FString& JsonMessage)
@@ -276,11 +276,11 @@ bool UInterchangeFbxTranslator::Translate(UInterchangeBaseNodeContainer& BaseNod
 	TRACE_CPUPROFILER_EVENT_SCOPE(UInterchangeFbxTranslator::Translate);
 	using namespace UE::Interchange::Private;
 
-	if (AnalyticsHelper)
+	if (AnalyticsHandler)
 	{
 		// This is to ensure that the keys are not duplicated when the Analytics Event is actually sent.
 		// Changing the Translator Settings in the import dialog will call the Translate function on the Translators again.
-		AnalyticsHelper->ClearAnalyticsEventData(FBXAnalyticsIdentifier);
+		AnalyticsHandler->Clear(FBXAnalyticsIdentifier);
 	}
 
 	if (HasAnyFlags(RF_ClassDefaultObject | RF_ArchetypeObject))
@@ -290,7 +290,7 @@ bool UInterchangeFbxTranslator::Translate(UInterchangeBaseNodeContainer& BaseNod
 		UInterchangeResultError_Generic* ErrorResult = AddMessage<UInterchangeResultError_Generic>();
 		ErrorResult->SourceAssetName = SourceData ? SourceData->GetFilename() : FString();
 		ErrorResult->Text = NSLOCTEXT("UInterchangeFbxTranslator", "Translate_DoNotUseClassDefault", "Class default should not be use has an active translator.");
-		AddAnalytics(EFBXTranslatorAnalyticsResult::TranslatorFailed, AnalyticsHelper, bUseWorkerImport);
+		AddAnalytics(EFBXTranslatorAnalyticsResult::TranslatorFailed, AnalyticsHandler, bUseWorkerImport);
 		return false;
 	}
 
@@ -299,7 +299,7 @@ bool UInterchangeFbxTranslator::Translate(UInterchangeBaseNodeContainer& BaseNod
 	FString FilePath = GetSourceData()->GetFilename();
 	if (!FPaths::FileExists(FilePath))
 	{
-		AddAnalytics(EFBXTranslatorAnalyticsResult::InputFileNotFound, AnalyticsHelper, bUseWorkerImport);
+		AddAnalytics(EFBXTranslatorAnalyticsResult::InputFileNotFound, AnalyticsHandler, bUseWorkerImport);
 		return false;
 	}
 
@@ -327,7 +327,7 @@ bool UInterchangeFbxTranslator::Translate(UInterchangeBaseNodeContainer& BaseNod
 	{
 		if (!Dispatcher.IsValid())
 		{
-			AddAnalytics(InterchangeDispatcherFailed, AnalyticsHelper, bUseWorkerImport);
+			AddAnalytics(InterchangeDispatcherFailed, AnalyticsHandler, bUseWorkerImport);
 			return false;
 		}
 		//Create a json command to read the fbx file
@@ -341,7 +341,7 @@ bool UInterchangeFbxTranslator::Translate(UInterchangeBaseNodeContainer& BaseNod
 		if (!WorkerFatalError.IsEmpty())
 		{
 			AddMessage(UInterchangeResult::FromJson(WorkerFatalError));
-			AddAnalytics(EFBXTranslatorAnalyticsResult::InterchangeWorkerFailed, AnalyticsHelper, bUseWorkerImport);
+			AddAnalytics(EFBXTranslatorAnalyticsResult::InterchangeWorkerFailed, AnalyticsHandler, bUseWorkerImport);
 		}
 
 		UE::Interchange::ETaskState TaskState;
@@ -366,7 +366,7 @@ bool UInterchangeFbxTranslator::Translate(UInterchangeBaseNodeContainer& BaseNod
 
 				if (!bEncounteredParserErrors && InterchangeResult->GetResultType() == EInterchangeResultType::Error)
 				{
-					AddAnalytics(EFBXTranslatorAnalyticsResult::FBXParserFailed, AnalyticsHelper, bUseWorkerImport);
+					AddAnalytics(EFBXTranslatorAnalyticsResult::FBXParserFailed, AnalyticsHandler, bUseWorkerImport);
 					bEncounteredParserErrors = true;
 				}
 			}
@@ -374,7 +374,7 @@ bool UInterchangeFbxTranslator::Translate(UInterchangeBaseNodeContainer& BaseNod
 
 		if (TaskState != UE::Interchange::ETaskState::ProcessOk)
 		{
-			AddAnalytics(EFBXTranslatorAnalyticsResult::InterchangeWorkerFailed, AnalyticsHelper, bUseWorkerImport);
+			AddAnalytics(EFBXTranslatorAnalyticsResult::InterchangeWorkerFailed, AnalyticsHandler, bUseWorkerImport);
 			return false;
 		}
 		//Grab the result file and fill the BaseNodeContainer
@@ -402,7 +402,7 @@ bool UInterchangeFbxTranslator::Translate(UInterchangeBaseNodeContainer& BaseNod
 				{
 					AddAnalytics(
 						EFBXTranslatorAnalyticsResult::FBXParserFailed, 
-						AnalyticsHelper, 
+						AnalyticsHandler, 
 						bUseWorkerImport, 
 						bUseUfbxParser ? EFBXParserType::uFBX : EFBXParserType::FBX_SDK
 					);
@@ -428,7 +428,7 @@ bool UInterchangeFbxTranslator::Translate(UInterchangeBaseNodeContainer& BaseNod
 	{
 		AddAnalytics(
 			EFBXTranslatorAnalyticsResult::Successful, 
-			AnalyticsHelper, 
+			AnalyticsHandler, 
 			bUseWorkerImport, 
 			bUseUfbxParser ? EFBXParserType::uFBX : EFBXParserType::FBX_SDK
 		);
@@ -542,7 +542,7 @@ TOptional<UE::Interchange::FImportImage> UInterchangeFbxTranslator::GetTexturePa
 	}
 #endif
 
-	UE::Interchange::Private::FScopedTranslator ScopedTranslator(PayLoadKey, Results, AnalyticsHelper);
+	UE::Interchange::Private::FScopedTranslator ScopedTranslator(PayLoadKey, Results, AnalyticsHandler);
 	const IInterchangeTexturePayloadInterface* TextureTranslator = ScopedTranslator.GetPayLoadInterface<IInterchangeTexturePayloadInterface>();
 	if (!TextureTranslator)
 	{
diff --git a/Engine/Plugins/Interchange/Runtime/Source/Import/Private/Gltf/InterchangeGltfTranslator.cpp b/Engine/Plugins/Interchange/Runtime/Source/Import/Private/Gltf/InterchangeGltfTranslator.cpp
--- a/Engine/Plugins/Interchange/Runtime/Source/Import/Private/Gltf/InterchangeGltfTranslator.cpp
+++ b/Engine/Plugins/Interchange/Runtime/Source/Import/Private/Gltf/InterchangeGltfTranslator.cpp
@@ -11,7 +11,7 @@
 #include "GLTFReader.h"
 #include "GLTFTexture.h"
 
-#include "InterchangeAnalyticsHelper.h"
+#include "InterchangeAnalyticsHandlerBase.h"
 #include "InterchangeAnimationTrackSetNode.h"
 #include "InterchangeCameraNode.h"
 #include "InterchangeImportLog.h"
@@ -154,11 +154,11 @@ namespace UE::Interchange::Gltf::Private
 		NOTSUPPORTED_EXTENSION_FOUND
 	};
 	void AddAnalytics(const TranslationResult& TranslationResult,
-		TSharedPtr<UE::Interchange::FAnalyticsHelper> AnalyticsHelper,
+		TObjectPtr<UInterchangeAnalyticsHandlerBase> AnalyticsHandler,
 		const GLTF::FAsset& Asset = GLTF::FAsset(),
 		const FText& GLTFReaderLogMessage = FText::GetEmpty())
 	{
-		if (!AnalyticsHelper.IsValid())
+		if (!AnalyticsHandler)
 		{
 			return;
 		}
@@ -218,7 +218,7 @@ namespace UE::Interchange::Gltf::Private
 			break;
 		}
 
-		AnalyticsHelper->AppendThreadSafe(GLTFAnalyticsIdentifier, GLTFAnalytics);
+		AnalyticsHandler->Append(GLTFAnalyticsIdentifier, GLTFAnalytics);
 	};
 }
 
@@ -503,7 +503,7 @@ bool UInterchangeGLTFTranslator::Translate( UInterchangeBaseNodeContainer& NodeC
 	FString FilePath = GetSourceData()->GetFilename();
 	if ( !FPaths::FileExists( FilePath ) )
 	{
-		AddAnalytics(TranslationResult::INPUT_FILE_NOTFOUND, AnalyticsHelper);
+		AddAnalytics(TranslationResult::INPUT_FILE_NOTFOUND, AnalyticsHandler);
 		return false;
 	}
 
@@ -585,7 +585,7 @@ bool UInterchangeGLTFTranslator::Translate( UInterchangeBaseNodeContainer& NodeC
 			ErrorResult->SourceAssetName = FileName;
 			ErrorResult->Text = LogMessage.Value;
 
-			AddAnalytics(TranslationResult::GLTFREADER_FAILED, AnalyticsHelper, GltfAsset, LogMessage.Value);
+			AddAnalytics(TranslationResult::GLTFREADER_FAILED, AnalyticsHandler, GltfAsset, LogMessage.Value);
 			return false;
 		}
 	}
@@ -609,7 +609,7 @@ bool UInterchangeGLTFTranslator::Translate( UInterchangeBaseNodeContainer& NodeC
 			LOCTEXT("UnsupportedRequiredExtensions", "Not all required extensions are supported. (Unsupported extensions: {0})"),
 			FText::FromString(NotSupportedRequiredExtensionsStringified));
 
-		AddAnalytics(TranslationResult::NOTSUPPORTED_EXTENSION_FOUND, AnalyticsHelper, GltfAsset);
+		AddAnalytics(TranslationResult::NOTSUPPORTED_EXTENSION_FOUND, AnalyticsHandler, GltfAsset);
 		return false;
 	}
 
@@ -981,7 +981,7 @@ bool UInterchangeGLTFTranslator::Translate( UInterchangeBaseNodeContainer& NodeC
 		UE_LOG(LogInterchangeImport, Warning, TEXT("glTF Mesh Import Warning: glTF mesh usage expectations are not met."));
 	}
 
-	AddAnalytics(TranslationResult::SUCCESSFULL, AnalyticsHelper, GltfAsset);
+	AddAnalytics(TranslationResult::SUCCESSFULL, AnalyticsHandler, GltfAsset);
 	return true;
 }
 
@@ -1019,7 +1019,7 @@ TOptional< UE::Interchange::FImportImage > UInterchangeGLTFTranslator::GetTextur
 	else
 	{
 		const FString TextureFilePath = FPaths::ConvertRelativePathToFull(GltfTexture.Source.FilePath);
-		UE::Interchange::Private::FScopedTranslator ScopedTranslator(TextureFilePath, Results, AnalyticsHelper);
+		UE::Interchange::Private::FScopedTranslator ScopedTranslator(TextureFilePath, Results, AnalyticsHandler);
 		const IInterchangeTexturePayloadInterface* TextureTranslator = ScopedTranslator.GetPayLoadInterface<IInterchangeTexturePayloadInterface>();
 		if (!ensure(TextureTranslator))
 		{
diff --git a/Engine/Plugins/Interchange/Runtime/Source/Import/Private/InterchangeTranslatorHelper.cpp b/Engine/Plugins/Interchange/Runtime/Source/Import/Private/InterchangeTranslatorHelper.cpp
--- a/Engine/Plugins/Interchange/Runtime/Source/Import/Private/InterchangeTranslatorHelper.cpp
+++ b/Engine/Plugins/Interchange/Runtime/Source/Import/Private/InterchangeTranslatorHelper.cpp
@@ -11,7 +11,7 @@
 
 namespace UE::Interchange::Private
 {
-	FScopedTranslator::FScopedTranslator(const FString& PayLoadKey, UInterchangeResultsContainer* Results, TSharedPtr<UE::Interchange::FAnalyticsHelper> AnalyticsHelper)
+	FScopedTranslator::FScopedTranslator(const FString& PayLoadKey, UInterchangeResultsContainer* Results, TObjectPtr<UInterchangeAnalyticsHandlerBase> AnalyticsHandler)
 	{
 		PayloadSourceData = UInterchangeManager::GetInterchangeManager().CreateSourceData(PayLoadKey);
 		if (PayloadSourceData)
@@ -21,7 +21,7 @@ namespace UE::Interchange::Private
 			SourceTranslator = UInterchangeManager::GetInterchangeManager().GetTranslatorForSourceData(PayloadSourceData);
 			if (SourceTranslator)
 			{
-				SourceTranslator->AnalyticsHelper = AnalyticsHelper;
+				SourceTranslator->AnalyticsHandler = AnalyticsHandler;
 				SourceTranslator->SetInternalFlags(EInternalObjectFlags::Async);
 				if (Results)
 				{
diff --git a/Engine/Plugins/Interchange/Runtime/Source/Import/Private/MaterialX/InterchangeMaterialXTranslator.cpp b/Engine/Plugins/Interchange/Runtime/Source/Import/Private/MaterialX/InterchangeMaterialXTranslator.cpp
--- a/Engine/Plugins/Interchange/Runtime/Source/Import/Private/MaterialX/InterchangeMaterialXTranslator.cpp
+++ b/Engine/Plugins/Interchange/Runtime/Source/Import/Private/MaterialX/InterchangeMaterialXTranslator.cpp
@@ -61,10 +61,10 @@ bool UInterchangeMaterialXTranslator::Translate(UInterchangeBaseNodeContainer& B
 
 TOptional<UE::Interchange::FImportImage> UInterchangeMaterialXTranslator::GetTexturePayloadData(const FString& PayloadKey, TOptional<FString>& AlternateTexturePath) const
 {
-	return GetTexturePayloadData(PayloadKey, AlternateTexturePath, Results, AnalyticsHelper);
+	return GetTexturePayloadData(PayloadKey, AlternateTexturePath, Results, AnalyticsHandler);
 }
 
-TOptional<UE::Interchange::FImportImage> UInterchangeMaterialXTranslator::GetTexturePayloadData(const FString& PayloadKey, TOptional<FString>& AlternateTexturePath, UInterchangeResultsContainer* InResults, TSharedPtr<UE::Interchange::FAnalyticsHelper> InAnalyticsHelper)
+TOptional<UE::Interchange::FImportImage> UInterchangeMaterialXTranslator::GetTexturePayloadData(const FString& PayloadKey, TOptional<FString>& AlternateTexturePath, UInterchangeResultsContainer* InResults, TObjectPtr<UInterchangeAnalyticsHandlerBase> InAnalyticsHandler)
 {
 	FString Filename = PayloadKey;
 	TextureCompressionSettings CompressionSettings = TextureCompressionSettings::TC_Default;
@@ -82,7 +82,7 @@ TOptional<UE::Interchange::FImportImage> UInterchangeMaterialXTranslator::GetTex
 		CompressionSettings = TextureCompressionSettings::TC_HalfFloat;
 	}
 
-	UE::Interchange::Private::FScopedTranslator ScopedTranslator(Filename, InResults, InAnalyticsHelper);
+	UE::Interchange::Private::FScopedTranslator ScopedTranslator(Filename, InResults, InAnalyticsHandler);
 	const IInterchangeTexturePayloadInterface* TextureTranslator = ScopedTranslator.GetPayLoadInterface<IInterchangeTexturePayloadInterface>();
 
 	if (!TextureTranslator)
@@ -104,10 +104,10 @@ TOptional<UE::Interchange::FImportImage> UInterchangeMaterialXTranslator::GetTex
 
 TOptional<UE::Interchange::FImportBlockedImage> UInterchangeMaterialXTranslator::GetBlockedTexturePayloadData(const FString& PayloadKey, TOptional<FString>& AlternateTexturePath) const
 {
-	return GetBlockedTexturePayloadData(PayloadKey, AlternateTexturePath, Results, AnalyticsHelper);
+	return GetBlockedTexturePayloadData(PayloadKey, AlternateTexturePath, Results, AnalyticsHandler);
 }
 
-TOptional<UE::Interchange::FImportBlockedImage> UInterchangeMaterialXTranslator::GetBlockedTexturePayloadData(const FString& PayloadKey, TOptional<FString>& AlternateTexturePath, UInterchangeResultsContainer* InResults, TSharedPtr<UE::Interchange::FAnalyticsHelper> InAnalyticsHelper)
+TOptional<UE::Interchange::FImportBlockedImage> UInterchangeMaterialXTranslator::GetBlockedTexturePayloadData(const FString& PayloadKey, TOptional<FString>& AlternateTexturePath, UInterchangeResultsContainer* InResults, TObjectPtr<UInterchangeAnalyticsHandlerBase> InAnalyticsHandler)
 {
 	using namespace UE::Interchange;
 	using namespace UE::TextureUtilitiesCommon;
@@ -141,7 +141,7 @@ TOptional<UE::Interchange::FImportBlockedImage> UInterchangeMaterialXTranslator:
 		const int32 IndexUDIM = Pair.Get<0>();
 		const FString& FilenameUDIM = Pair.Get<1>();
 
-		UE::Interchange::Private::FScopedTranslator ScopedTranslator(FilenameUDIM, InResults, InAnalyticsHelper);
+		UE::Interchange::Private::FScopedTranslator ScopedTranslator(FilenameUDIM, InResults, InAnalyticsHandler);
 		const IInterchangeTexturePayloadInterface* TextureTranslator = ScopedTranslator.GetPayLoadInterface<IInterchangeTexturePayloadInterface>();
 
 		// should we stop if one of the UDIM image is broken?
diff --git a/Engine/Plugins/Interchange/Runtime/Source/Import/Private/MaterialX/MaterialXUtils/MaterialXManager.cpp b/Engine/Plugins/Interchange/Runtime/Source/Import/Private/MaterialX/MaterialXUtils/MaterialXManager.cpp
--- a/Engine/Plugins/Interchange/Runtime/Source/Import/Private/MaterialX/MaterialXUtils/MaterialXManager.cpp
+++ b/Engine/Plugins/Interchange/Runtime/Source/Import/Private/MaterialX/MaterialXUtils/MaterialXManager.cpp
@@ -23,7 +23,7 @@
 #include "MaterialX/MaterialXUtils/MaterialXVolumeShader.h"
 #include "MaterialX/MaterialXUtils/MaterialXVolumeMaterial.h"
 
-#include "InterchangeAnalyticsHelper.h"
+#include "InterchangeAnalyticsHandlerBase.h"
 #include "InterchangeHelper.h"
 #include "InterchangeImportLog.h"
 #include "InterchangeMaterialDefinitions.h"
@@ -607,9 +607,9 @@ bool FMaterialXManager::Translate(MaterialX::DocumentPtr Document, UInterchangeB
 		}
 	}
 
-	if (Translator && Translator->AnalyticsHelper)
+	if (Translator && Translator->AnalyticsHandler)
 	{
-		Translator->AnalyticsHelper->AppendThreadSafe(TEXT("Interchange.Usage.Import.MaterialX"), AnalyticsAttributes);
+		Translator->AnalyticsHandler->Append(TEXT("Interchange.Usage.Import.MaterialX"), AnalyticsAttributes);
 	}
 
 	return bIsDocumentValid;
diff --git a/Engine/Plugins/Interchange/Runtime/Source/Import/Private/Mesh/InterchangeOBJTranslator.cpp b/Engine/Plugins/Interchange/Runtime/Source/Import/Private/Mesh/InterchangeOBJTranslator.cpp
--- a/Engine/Plugins/Interchange/Runtime/Source/Import/Private/Mesh/InterchangeOBJTranslator.cpp
+++ b/Engine/Plugins/Interchange/Runtime/Source/Import/Private/Mesh/InterchangeOBJTranslator.cpp
@@ -1563,7 +1563,7 @@ TOptional<UE::Interchange::FImportImage> UInterchangeOBJTranslator::GetTexturePa
 	// Here, we have a texture source referenced by the .obj which is ideally handled by the existing texture translators.
 	// In this case, we can just implement the texture payload interface here and use a temporary texture translator instance
 	// to generate the payload, but what if the translator needed to do non-trivial work during the Translate() phase?
-	UE::Interchange::Private::FScopedTranslator ScopedTranslator(PayLoadKey, Results, AnalyticsHelper);
+	UE::Interchange::Private::FScopedTranslator ScopedTranslator(PayLoadKey, Results, AnalyticsHandler);
 	const IInterchangeTexturePayloadInterface* TextureTranslator = ScopedTranslator.GetPayLoadInterface<IInterchangeTexturePayloadInterface>();
 	if (!ensure(TextureTranslator))
 	{
diff --git a/Engine/Plugins/Interchange/Runtime/Source/Import/Public/InterchangeTranslatorHelper.h b/Engine/Plugins/Interchange/Runtime/Source/Import/Public/InterchangeTranslatorHelper.h
--- a/Engine/Plugins/Interchange/Runtime/Source/Import/Public/InterchangeTranslatorHelper.h
+++ b/Engine/Plugins/Interchange/Runtime/Source/Import/Public/InterchangeTranslatorHelper.h
@@ -8,18 +8,14 @@
 #define UE_API INTERCHANGEIMPORT_API
 
 class UInterchangeResultsContainer;
-
-namespace UE::Interchange
-{
-	struct FAnalyticsHelper;
-}
+class UInterchangeAnalyticsHandlerBase;
 
 namespace UE::Interchange::Private
 {
 	struct FScopedTranslator
 	{
 	public:
-		UE_API FScopedTranslator(const FString& PayLoadKey, UInterchangeResultsContainer* Results, TSharedPtr<UE::Interchange::FAnalyticsHelper> AnalyticsHelper);
+		UE_API FScopedTranslator(const FString& PayLoadKey, UInterchangeResultsContainer* Results, TObjectPtr<UInterchangeAnalyticsHandlerBase> AnalyticsHandler);
 		UE_API ~FScopedTranslator();
 
 		template <class T>
diff --git a/Engine/Plugins/Interchange/Runtime/Source/Import/Public/MaterialX/InterchangeMaterialXTranslator.h b/Engine/Plugins/Interchange/Runtime/Source/Import/Public/MaterialX/InterchangeMaterialXTranslator.h
--- a/Engine/Plugins/Interchange/Runtime/Source/Import/Public/MaterialX/InterchangeMaterialXTranslator.h
+++ b/Engine/Plugins/Interchange/Runtime/Source/Import/Public/MaterialX/InterchangeMaterialXTranslator.h
@@ -42,11 +42,11 @@ public:
 
 	using IInterchangeTexturePayloadInterface::GetTexturePayloadData;
 	UE_API virtual TOptional<UE::Interchange::FImportImage> GetTexturePayloadData(const FString& PayloadKey, TOptional<FString>& AlternateTexturePath) const override;
-	UE_API static TOptional<UE::Interchange::FImportImage> GetTexturePayloadData(const FString& PayloadKey, TOptional<FString>& AlternateTexturePath, UInterchangeResultsContainer* InResults, TSharedPtr<UE::Interchange::FAnalyticsHelper> InAnalyticsHelper);
+	UE_API static TOptional<UE::Interchange::FImportImage> GetTexturePayloadData(const FString& PayloadKey, TOptional<FString>& AlternateTexturePath, UInterchangeResultsContainer* InResults, TObjectPtr<UInterchangeAnalyticsHandlerBase> InAnalyticsHandler);
 
 	using IInterchangeBlockedTexturePayloadInterface::GetBlockedTexturePayloadData;
 	UE_API virtual TOptional<UE::Interchange::FImportBlockedImage> GetBlockedTexturePayloadData(const FString& PayloadKey, TOptional<FString>& AlternateTexturePath) const override;
-	UE_API static TOptional<UE::Interchange::FImportBlockedImage> GetBlockedTexturePayloadData(const FString& PayloadKey, TOptional<FString>& AlternateTexturePath, UInterchangeResultsContainer* InResults, TSharedPtr<UE::Interchange::FAnalyticsHelper> InAnalyticsHelper);
+	UE_API static TOptional<UE::Interchange::FImportBlockedImage> GetBlockedTexturePayloadData(const FString& PayloadKey, TOptional<FString>& AlternateTexturePath, UInterchangeResultsContainer* InResults, TObjectPtr<UInterchangeAnalyticsHandlerBase> InAnalyticsHandler);
 
 };
 
diff --git a/Engine/Source/Runtime/Interchange/Core/Public/InterchangeTranslatorBase.h b/Engine/Source/Runtime/Interchange/Core/Public/InterchangeTranslatorBase.h
--- a/Engine/Source/Runtime/Interchange/Core/Public/InterchangeTranslatorBase.h
+++ b/Engine/Source/Runtime/Interchange/Core/Public/InterchangeTranslatorBase.h
@@ -19,11 +19,7 @@
 #include "InterchangeTranslatorBase.generated.h"
 
 class UInterchangeBaseNodeContainer;
-
-namespace UE::Interchange
-{
-	struct FAnalyticsHelper;
-}
+class UInterchangeAnalyticsHandlerBase;
 
 /**
  * Base class for translator settings, the UInterchangeAssetImportData will store this data if the source translator exist
@@ -180,5 +176,5 @@ public:
 	UPROPERTY()
 	TObjectPtr<const UInterchangeSourceData> SourceData;
 
-	TSharedPtr<UE::Interchange::FAnalyticsHelper> AnalyticsHelper = nullptr;
+	TObjectPtr<UInterchangeAnalyticsHandlerBase> AnalyticsHandler = nullptr;
 };
diff --git a/Engine/Source/Runtime/Interchange/Engine/Private/InterchangeAnalyticsHandlerBase.cpp b/Engine/Source/Runtime/Interchange/Engine/Private/InterchangeAnalyticsHandlerBase.cpp
new file mode 100644
--- /dev/null
+++ b/Engine/Source/Runtime/Interchange/Engine/Private/InterchangeAnalyticsHandlerBase.cpp
@@ -0,0 +1,93 @@
+// Copyright Epic Games, Inc. All Rights Reserved.
+
+#include "InterchangeAnalyticsHandlerBase.h"
+#include "InterchangePipelineBase.h"
+#include "Misc/ScopeLock.h"
+
+FString UInterchangeAnalyticsHandlerBase::GetDefaultAnalyticsCategoryName(EInterchangeAnalyticsDefaultCategory DefaultCategory)
+{
+	static const FString CategoryNames[(size_t)EInterchangeAnalyticsDefaultCategory::InterchangeAnalyticsCategory_COUNT]
+	{
+	#define Category(Name) TEXT("."#Name)
+
+		TEXT(""), // Not Specified
+		Category(Models),
+		Category(Scenes),
+		Category(Materials),
+		Category(Textures),
+		Category(Sounds),
+		Category(Meshes),
+		Category(StaticMeshes),
+		Category(SkeletalMeshesAndAnimations),
+		Category(GeometryCaches),
+		Category(Grooms),
+
+	#undef Category
+	};
+
+	if (DefaultCategory > EInterchangeAnalyticsDefaultCategory::NotSpecified
+		&& DefaultCategory < EInterchangeAnalyticsDefaultCategory::InterchangeAnalyticsCategory_COUNT)
+	{
+		return CategoryNames[(size_t)DefaultCategory];
+	}
+
+	return CategoryNames[(size_t)EInterchangeAnalyticsDefaultCategory::NotSpecified];
+}
+
+FAnalyticsEventAttribute UInterchangeAnalyticsHandlerBase::CreateAnalyticsIDAttribute(FGuid AnalyticsID)
+{
+	ensure(AnalyticsID.IsValid());
+	return FAnalyticsEventAttribute(TEXT("InterchangeAnalyticsGUID"), AnalyticsID.ToString());
+}
+
+FString UInterchangeAnalyticsHandlerBase::GetPipelineAnalyticsEventIdentifier()
+{
+	const static FString PipelineAnalyticsEventIdentifier = TEXT("Interchange.Usage.Import.Pipeline");
+	return PipelineAnalyticsEventIdentifier;
+}
+
+FString UInterchangeAnalyticsHandlerBase::GetImportResultEventIdentifier()
+{
+	const static FString ImportResultEventIdentifier = TEXT("Interchange.Usage.ImportResult");
+	return ImportResultEventIdentifier;
+}
+
+void UInterchangeAnalyticsHandlerBase::SetGUID(const FGuid& InGUID)
+{
+	if (AnalyticsID.IsValid())
+	{
+		return;
+	}
+
+	AnalyticsID = InGUID;
+}
+
+void UInterchangeAnalyticsHandlerBase::Append(const FString& Identifier, const TArray<FAnalyticsEventAttribute>& ToAdd)
+{
+	FScopeLock ScopeLock(&CriticalSection);
+	Execute_Append(Identifier, ToAdd);
+}
+
+void UInterchangeAnalyticsHandlerBase::Add(const FString& Identifier, const FAnalyticsEventAttribute& Entry)
+{
+	FScopeLock ScopeLock(&CriticalSection);
+	Execute_Add(Identifier, Entry);
+}
+
+void UInterchangeAnalyticsHandlerBase::Clear(const FString& Identifier)
+{
+	FScopeLock ScopeLock(&CriticalSection);
+	Execute_Clear(Identifier);
+}
+
+void UInterchangeAnalyticsHandlerBase::Flush()
+{
+	FScopeLock ScopeLock(&CriticalSection);
+	Execute_Flush();
+}
+
+FAnalyticsEventAttribute UInterchangeAnalyticsHandlerBase::GetAnalyticsIDAttribute() const
+{
+	ensure(AnalyticsID.IsValid());
+	return CreateAnalyticsIDAttribute(AnalyticsID);
+}
\ No newline at end of file
diff --git a/Engine/Source/Runtime/Interchange/Engine/Private/InterchangeAnalyticsHelper.cpp b/Engine/Source/Runtime/Interchange/Engine/Private/InterchangeAnalyticsHelper.cpp
deleted file mode 100644
diff --git a/Engine/Source/Runtime/Interchange/Engine/Private/InterchangeManager.cpp b/Engine/Source/Runtime/Interchange/Engine/Private/InterchangeManager.cpp
--- a/Engine/Source/Runtime/Interchange/Engine/Private/InterchangeManager.cpp
+++ b/Engine/Source/Runtime/Interchange/Engine/Private/InterchangeManager.cpp
@@ -15,7 +15,6 @@
 #include "CoreMinimal.h"
 #include "Engine/Blueprint.h"
 #include "Engine/World.h"
-#include "EngineAnalytics.h"
 #include "EngineUtils.h"
 #include "Framework/Notifications/NotificationManager.h"
 #include "InterchangeAssetImportData.h"
@@ -24,7 +23,7 @@
 #include "InterchangeEngineLogPrivate.h"
 #include "InterchangeFactoryBase.h"
 #include "InterchangeHelper.h"
-#include "InterchangeAnalyticsHelper.h"
+#include "InterchangeAnalyticsHandlerBase.h"
 #include "InterchangeProjectSettings.h"
 #include "InterchangePythonPipelineBase.h"
 #include "InterchangeSceneImportAsset.h"
@@ -43,10 +42,6 @@
 #include "Nodes/InterchangeBaseNodeContainer.h"
 #include "Nodes/InterchangeSourceNode.h"
 #include "PackageUtils/PackageUtils.h"
-#include "Policies/CondensedJsonPrintPolicy.h"
-#include "Serialization/JsonSerializer.h"
-#include "Serialization/JsonWriter.h"
-#include "Serialization/JsonSerializerWriter.h"
 #include "Tasks/InterchangeTaskParsing.h"
 #include "Tasks/InterchangeTaskPipeline.h"
 #include "Tasks/InterchangeTaskTranslator.h"
@@ -98,130 +93,47 @@ namespace UE::Interchange::Private
 #endif
 	}
 
-	void FillPipelineAnalyticData(FAnalyticsHelper& AnalyticsHelper, UInterchangePipelineBase* Pipeline, const int32 UniqueId, const FString& ParentPipeline)
+	class FInterchangeAnalyticsHelper
 	{
-		if (!Pipeline || !FEngineAnalytics::IsAvailable())
+	public:
+		FInterchangeAnalyticsHelper(UInterchangeAnalyticsHandlerBase* InHandler)
+			:Handler(InHandler)
 		{
-			return;
 		}
-		
-		int32 PortFlags = 0;
-		UClass* Class = Pipeline->GetClass();
-		FString PipelineChainName = ParentPipeline.IsEmpty() ? Pipeline->GetName() : ParentPipeline + TEXT(".") + Pipeline->GetName();
 
-		TArray<FAnalyticsEventAttribute> PipelineAttribs;
-		PipelineAttribs.Add(FAnalyticsEventAttribute(TEXT("UniqueId"), UniqueId));
-		PipelineAttribs.Add(FAnalyticsEventAttribute(TEXT("Name"), PipelineChainName));
-		PipelineAttribs.Add(FAnalyticsEventAttribute(TEXT("Class"), Class->GetName()));
-
-		for (FProperty* Property = Class->PropertyLink; Property; Property = Property->PropertyLinkNext)
+		void Add(const FString& Identifier, const FAnalyticsEventAttribute& Entry)
 		{
-			if (!Pipeline || !Pipeline->CanAddPropertyToAnalytics(Property))
+			if (!Handler)
 			{
-				continue;
+				return;
 			}
 
-			FObjectProperty* SubObject = CastField<FObjectProperty>(Property);
-			UInterchangePipelineBase* SubPipeline = SubObject ? Cast<UInterchangePipelineBase>(SubObject->GetObjectPropertyValue_InContainer(Pipeline)) : nullptr;
-			
-			
-			if (SubPipeline)
-			{
-				// Save the settings if the referenced pipeline is a subobject of ours
-				if (SubPipeline->IsInOuter(Pipeline))
-				{
-					//Go recursive with subObject, like if they are part of the same object
-					FillPipelineAnalyticData(AnalyticsHelper, SubPipeline, UniqueId, PipelineChainName);
-				}
-			}
-			else
-			{
-				Pipeline->AddPropertyAnalytics(Property, PipelineAttribs);
-			}
-		}
-		
-		{
-			const FString PipelineEventIdentifier = "Interchange.Usage.Import.Pipeline";
-			AnalyticsHelper.SendAnalyticsImmediate(PipelineEventIdentifier, PipelineAttribs);
+			Handler->Add(Identifier, Entry);
 		}
-		
-	}
 
-	// Json writer subclass to allow us to avoid using a SharedPtr to write basic Json.
-	class FAnalyticsJsonWriter : public TJsonStringWriter<TCondensedJsonPrintPolicy<TCHAR>>
-	{
-	public:
-		explicit FAnalyticsJsonWriter(FString* Out) : TJsonStringWriter<TCondensedJsonPrintPolicy<TCHAR>>(Out, 0)
+		void Flush()
 		{
-		}
-	};
+			if (!Handler)
+			{
+				return;
+			}
 
-	FJsonFragment AnalyticsConvertMapToJsonFragment(const TMap<FString, int32>& InFreqMap)
-	{
-		FString ReturnValue;
-		FAnalyticsJsonWriter JsonWriter(&ReturnValue);
-		JsonWriter.WriteArrayStart();
-		for (const TPair<FString,int32>& FreqMapPair : InFreqMap)
-		{
-			JsonWriter.WriteObjectStart();
-			JsonWriter.WriteValue(TEXT("MessageKey"), FreqMapPair.Key);
-			JsonWriter.WriteValue(TEXT("MessageCount"), FreqMapPair.Value);
-			JsonWriter.WriteObjectEnd();
+			Handler->Flush();
 		}
-		JsonWriter.WriteArrayEnd();
-		JsonWriter.Close();
-		return FJsonFragment(MoveTemp(ReturnValue));
-	}
 
-	FJsonFragment AnalyticsPrimarySourceDataMetadataToJsonFragment(const UInterchangeSourceData* InSourceData, const UInterchangeSourceNode* InSourceNode)
-	{
-		using namespace UE::Interchange;
-		FString ReturnValue;
-		FAnalyticsJsonWriter JsonWriter(&ReturnValue);
-		
-		TMap<FString, FString> ExtraInfoMap;
-		InSourceNode->GetExtraInformation(ExtraInfoMap);
-
-		const FString DefaultValue("N/A");
-
-		const FString ApplicationVendor = ExtraInfoMap.FindRef(FSourceNodeExtraInfoStaticData::GetApplicationVendorExtraInfoKey(), DefaultValue);
-		const FString ApplicationName = ExtraInfoMap.FindRef(FSourceNodeExtraInfoStaticData::GetApplicationNameExtraInfoKey(), DefaultValue);
-		const FString ApplicationVersion = ExtraInfoMap.FindRef(FSourceNodeExtraInfoStaticData::GetApplicationVersionExtraInfoKey(), DefaultValue);
-
-		JsonWriter.WriteObjectStart();
-		JsonWriter.WriteValue(TEXT("SourceExtension"), FPaths::GetExtension(InSourceData->GetFilename()));
-		JsonWriter.WriteValue(TEXT("ApplicationVendor"), ApplicationVendor);
-		JsonWriter.WriteValue(TEXT("ApplicationName"), ApplicationName);
-		JsonWriter.WriteValue(TEXT("ApplicationVersion"), ApplicationVersion);
-		JsonWriter.WriteObjectEnd();
-		JsonWriter.Close();
-		return FJsonFragment(MoveTemp(ReturnValue));
-	}
-
-	bool ExtractNamespace(const FText& Text, FString& OutTextNamespaceId)
-	{
-		FText TextToUse = Text;
-
-		TArray<FHistoricTextFormatData> TextHistory;
-		FTextInspector::GetHistoricFormatData(Text, TextHistory);
-		if (TextHistory.Num() > 0)
+		void Send(const TArray<UInterchangePipelineBase*>& Pipelines, const int32 AsyncHelperUniqueId)
 		{
-			const FHistoricTextFormatData& FmtData = TextHistory[0];
-			TextToUse = FmtData.SourceFmt.GetSourceText();
-		}
-
-		const TOptional<FString> TextNamespace = FTextInspector::GetNamespace(TextToUse);
-		const TOptional<FString> TextKey = FTextInspector::GetKey(TextToUse);
+			if (!Handler)
+			{
+				return;
+			}
 
-		if (TextNamespace.IsSet() && TextKey.IsSet())
-		{
-			OutTextNamespaceId =  FString::Printf(TEXT("%s_%s"), *TextNamespace.GetValue(), *TextKey.GetValue());
-			return true;
+			Handler->Send(Pipelines, AsyncHelperUniqueId);
 		}
-		
-		OutTextNamespaceId = FString("UnknownError");
-		return false;
-	}
+
+	private:
+		UInterchangeAnalyticsHandlerBase* Handler;
+	};
 }
 
 UE::Interchange::FScopedInterchangeImportEnableState::FScopedInterchangeImportEnableState(const bool bScopeValue)
@@ -288,10 +200,60 @@ UInterchangeBaseNodeContainer* UE::Interchange::FScopedBaseNodeContainer::GetBas
 	return ScopedBaseNodeContainerPtr.Get();
 }
 
-UE::Interchange::FImportAsyncHelper::FImportAsyncHelper(const FGuid& InAnalyticsID)
+FInterchangeImportResultAnalyticsInfo UE::Interchange::FImportAsyncHelper::CreateImportResultAnalyticsInfo(FImportAsyncHelper& AsyncHelper)
+{
+	const UInterchangeBaseNodeContainer* PrimaryBaseNodeContainer = AsyncHelper.BaseNodeContainers[0].Get();
+	const UInterchangeSourceNode* PrimarySourceNode = UInterchangeSourceNode::GetUniqueInstance(PrimaryBaseNodeContainer);
+	const UInterchangeSourceData* PrimarySourceData = AsyncHelper.SourceDatas[0].Get();
+
+	TOptional<const UInterchangeResultsContainer*> AssetImportResultsContainer;
+	TOptional<const UInterchangeResultsContainer*> SceneImportResultsContainer;
+
+	TMap<int32, TArray<TObjectPtr<UObject>>> AssetObjectsPerSourceIndex;
+	TMap<int32, TArray<TObjectPtr<UObject>>> SceneObjectsPerSourceIndex;
+
+	for (const TPair<int32, TArray<FImportedObjectInfo>>& ImportedAssetPerSourceIndex : AsyncHelper.ImportedAssetsPerSourceIndex)
+	{
+		TArray<TObjectPtr<UObject>>& ImportedObjects = AssetObjectsPerSourceIndex.FindOrAdd(ImportedAssetPerSourceIndex.Key);
+		for (const FImportedObjectInfo& ImportObjectInfo : ImportedAssetPerSourceIndex.Value)
+		{
+			ImportedObjects.Add(ImportObjectInfo.GetImportedObject());
+		}
+	}
+	if (const UInterchangeResultsContainer* AssetImportResults = AsyncHelper.AssetImportResult->GetResults())
+	{
+		AssetImportResultsContainer = AssetImportResults;
+	}
+
+
+	for (const TPair<int32, TArray<FImportedObjectInfo>>& ImportedSceneObjectPerSourceIndex : AsyncHelper.ImportedSceneObjectsPerSourceIndex)
+	{
+		TArray<TObjectPtr<UObject>>& ImportedObjects = SceneObjectsPerSourceIndex.FindOrAdd(ImportedSceneObjectPerSourceIndex.Key);
+		for (const FImportedObjectInfo& ImportObjectInfo : ImportedSceneObjectPerSourceIndex.Value)
+		{
+			ImportedObjects.Add(ImportObjectInfo.GetImportedObject());
+		}
+	}
+	if (const UInterchangeResultsContainer* SceneImportResults = AsyncHelper.SceneImportResult->GetResults())
+	{
+		SceneImportResultsContainer = SceneImportResults;
+	}
+
+	return FInterchangeImportResultAnalyticsInfo{
+		.bCancel  = AsyncHelper.bCancel,
+		.UniqueId = AsyncHelper.UniqueId,
+		.PrimarySourceData = PrimarySourceData,
+		.PrimarySourceNode = PrimarySourceNode,
+		.AssetImportResultsContainer = AssetImportResultsContainer,
+		.SceneImportResultsContainer = SceneImportResultsContainer,
+		.ImportedAssetsPerSourceIndex = MoveTemp(AssetObjectsPerSourceIndex),
+		.ImportedSceneObjectsPerSourceIndex = MoveTemp(SceneObjectsPerSourceIndex),
+	};
+}
+
+UE::Interchange::FImportAsyncHelper::FImportAsyncHelper()
 	: AssetImportResult(MakeShared<FImportResult>())
 	, SceneImportResult(MakeShared<FImportResult>())
-	, AnalyticsID(InAnalyticsID)
 {
 	bCancel = false;
 
@@ -305,6 +267,11 @@ UE::Interchange::FImportAsyncHelper::FImportAsyncHelper(const FGuid& InAnalytics
 UE::Interchange::FImportAsyncHelper::~FImportAsyncHelper()
 {
 	CleanUp();
+	
+	if (AnalyticsHandler.IsValid())
+	{
+		AnalyticsHandler.Reset();
+	}
 
 	UE::Interchange::Private::StaticHelpers::AsyncHelperCounter--;
 	if (UE::Interchange::Private::StaticHelpers::AsyncHelperCounter == 0)
@@ -521,108 +488,6 @@ bool UE::Interchange::FImportAsyncHelper::IsImportingObject(UObject* Object) con
 	return bFoundAsset;
 }
 
-void UE::Interchange::FImportAsyncHelper::SendAnalyticImportEndData()
-{
-	if (bCancel || !FEngineAnalytics::IsAvailable())
-	{
-		return;
-	}
-
-	TArray<FAnalyticsEventAttribute> Attribs;
-	//Set the unique id of this import
-	Attribs.Add(FAnalyticsHelper::CreateAnalyticsIDAttribute(AnalyticsID));
-	Attribs.Add(FAnalyticsEventAttribute(TEXT("UniqueId"), UniqueId));
-	Attribs.Add(FAnalyticsEventAttribute(TEXT("IsCanceled"), bCancel));
-	
-
-	const UInterchangeBaseNodeContainer* PrimaryBaseNodeContainer = BaseNodeContainers[0].Get();
-	if(const UInterchangeSourceNode* PrimarySourceNode = UInterchangeSourceNode::GetUniqueInstance(PrimaryBaseNodeContainer))
-	{
-		using namespace UE::Interchange::Private;
-		Attribs.Add(FAnalyticsEventAttribute(TEXT("PrimarySourceDataMetadata"), AnalyticsPrimarySourceDataMetadataToJsonFragment(SourceDatas[0], PrimarySourceNode)));
-	}
-
-	int32 ImportedObjectCount = 0;
-	for (const TPair<int32, TArray<FImportedObjectInfo>>& SourceIndexAndImportedAssets : ImportedAssetsPerSourceIndex)
-	{
-		ImportedObjectCount += SourceIndexAndImportedAssets.Value.Num();
-	}
-
-	for (const TPair<int32, TArray<FImportedObjectInfo>>& SourceIndexAndImportedScneObjects : ImportedSceneObjectsPerSourceIndex)
-	{
-		ImportedObjectCount += SourceIndexAndImportedScneObjects.Value.Num();
-	}
-
-	Attribs.Add(FAnalyticsEventAttribute(TEXT("ImportObjectCount"), ImportedObjectCount));
-
-	//Report any warning or error message
-	TMap<FString, int32> WarningTypes;
-	TMap<FString, int32> ErrorTypes;
-	auto CollectResultContainer = [&WarningTypes, &ErrorTypes](const UInterchangeResultsContainer* ResultContainer)
-	{
-		TArray<UInterchangeResult*> InterchangeResults = ResultContainer->GetResults();
-		
-		for (const UInterchangeResult* InterchangeResult : InterchangeResults)
-		{
-			using namespace UE::Interchange::Private;
-
-			switch (InterchangeResult->GetResultType() )
-			{
-			case EInterchangeResultType::Success:
-				break;
-			case EInterchangeResultType::Warning:
-			{
-				FString OutWarningAttribValue;
-				if (ExtractNamespace(InterchangeResult->GetText(), OutWarningAttribValue))
-				{
-					int32& Frequency = WarningTypes.FindOrAdd(OutWarningAttribValue);
-					Frequency++;
-				}
-				else
-				{
-					UE_LOG(LogInterchangeEngine, Error, TEXT("Failed to extract Analytic Attribute Value from %s"), *(InterchangeResult->GetText().ToString()));
-				}
-				break;
-			}
-			case EInterchangeResultType::Error:
-			{
-				FString OutErrorAttribValue;
-				if (ExtractNamespace(InterchangeResult->GetText(), OutErrorAttribValue))
-				{
-					int32& Frequency = ErrorTypes.FindOrAdd(OutErrorAttribValue);
-					Frequency++;
-				}
-				else
-				{
-					UE_LOG(LogInterchangeEngine, Error, TEXT("Failed to extract Analytic Attribute Value from %s"), *(InterchangeResult->GetText().ToString()));
-				}
-				break;
-			}
-			}
-		}
-	};
-
-	if (const UInterchangeResultsContainer* ResultContainer = AssetImportResult->GetResults())
-	{
-		CollectResultContainer(ResultContainer);
-	}
-	if (const UInterchangeResultsContainer* ResultContainer = SceneImportResult->GetResults())
-	{
-		CollectResultContainer(ResultContainer);
-	}
-
-	{
-		using namespace UE::Interchange::Private;
-		Attribs.Add(FAnalyticsEventAttribute(TEXT("WarningTypes"), AnalyticsConvertMapToJsonFragment(WarningTypes)));
-		Attribs.Add(FAnalyticsEventAttribute(TEXT("ErrorTypes"), AnalyticsConvertMapToJsonFragment(ErrorTypes)));
-	}
-
-	{
-		const FString ImportResultEventIdentifier = "Interchange.Usage.ImportResult";
-		FEngineAnalytics::GetProvider().RecordEvent(ImportResultEventIdentifier, Attribs);
-	}
-}
-
 void UE::Interchange::FImportAsyncHelper::ReleaseTranslatorsSource()
 {
 	//Make sure Executing translator task are done before releasing the resource
@@ -1801,13 +1666,23 @@ UInterchangeManager::ImportInternal(const FString& ContentPath, const UInterchan
 	static int32 GeneratedUniqueID = 0;
 	int32 UniqueId = ++GeneratedUniqueID;
 
-	const FGuid AnalyticsGuid = FGuid::NewDeterministicGuid(ContentPath, static_cast<uint64>(FDateTime::UtcNow().GetTicks()));
+	
+	TSharedRef<UE::Interchange::FImportAsyncHelper, ESPMode::ThreadSafe> AsyncHelper = MakeShared<UE::Interchange::FImportAsyncHelper, ESPMode::ThreadSafe>();
+	AsyncHelper->UniqueId = UniqueId;
+	AsyncHelper->bRunSynchronous = ImportAssetParameters.bRunSynchronous;
+
+	// Do not collect Analytics while automation testing.
+	if (!GIsAutomationTesting)
+	{
+		AsyncHelper->AnalyticsHandler = CreateAnalyticsHandler(ContentPath);
+	}
+
+	UE::Interchange::Private::FInterchangeAnalyticsHelper AnalyticsHelper(AsyncHelper->AnalyticsHandler.Get());
 
-	TSharedPtr<UE::Interchange::FAnalyticsHelper, ESPMode::ThreadSafe> AnalyticsHelper = MakeShared< UE::Interchange::FAnalyticsHelper, ESPMode::ThreadSafe>(AnalyticsGuid);
 	const FString UsageImportAnalayticsId = "Interchange.Usage.Import";
 	auto PreReturn = [&AnalyticsHelper]()
 	{
-		AnalyticsHelper->SendAnalytics();
+		AnalyticsHelper.Flush();	
 	};
 
 	auto EarlyExit = [&PreReturn]()
@@ -1819,21 +1694,22 @@ UInterchangeManager::ImportInternal(const FString& ContentPath, const UInterchan
 		return ImportResult;
 	};
 
-	AnalyticsHelper->Add(UsageImportAnalayticsId, FAnalyticsEventAttribute(TEXT("UniqueId"), UniqueId));
-	AnalyticsHelper->Add(UsageImportAnalayticsId, FAnalyticsEventAttribute(TEXT("ImportType"), ImportType == UE::Interchange::EImportType::ImportType_Asset ? TEXT("Asset") : TEXT("Scene")));
-	AnalyticsHelper->Add(UsageImportAnalayticsId, FAnalyticsEventAttribute(TEXT("Parameters.IsAutomated"), ImportAssetParameters.bIsAutomated));
-	AnalyticsHelper->Add(UsageImportAnalayticsId, FAnalyticsEventAttribute(TEXT("Parameters.IsReimport"), ImportAssetParameters.ReimportAsset != nullptr));
+	AnalyticsHelper.Add(UsageImportAnalayticsId, FAnalyticsEventAttribute(TEXT("UniqueId"), UniqueId));
+	AnalyticsHelper.Add(UsageImportAnalayticsId, FAnalyticsEventAttribute(TEXT("ImportType"), ImportType == UE::Interchange::EImportType::ImportType_Asset ? TEXT("Asset") : TEXT("Scene")));
+	AnalyticsHelper.Add(UsageImportAnalayticsId, FAnalyticsEventAttribute(TEXT("Parameters.IsAutomated"), ImportAssetParameters.bIsAutomated));
+	AnalyticsHelper.Add(UsageImportAnalayticsId, FAnalyticsEventAttribute(TEXT("Parameters.IsReimport"), ImportAssetParameters.ReimportAsset != nullptr));
 	if (ImportAssetParameters.ReimportAsset != nullptr)
 	{
-		AnalyticsHelper->Add(UsageImportAnalayticsId, FAnalyticsEventAttribute(TEXT("Parameters.ReimportSourceIndex"), ImportAssetParameters.ReimportSourceIndex));
+		AnalyticsHelper.Add(UsageImportAnalayticsId, FAnalyticsEventAttribute(TEXT("Parameters.ReimportSourceIndex"), ImportAssetParameters.ReimportSourceIndex));
 	}
 	const bool bIsPipelineOverride = ImportAssetParameters.OverridePipelines.Num() > 0;
-	AnalyticsHelper->Add(UsageImportAnalayticsId, FAnalyticsEventAttribute(TEXT("Parameters.IsPipelineOverrided"), bIsPipelineOverride));
-	AnalyticsHelper->Add(UsageImportAnalayticsId, FAnalyticsEventAttribute(TEXT("Parameters.bReplaceExisting"), ImportAssetParameters.bReplaceExisting));
+	AnalyticsHelper.Add(UsageImportAnalayticsId, FAnalyticsEventAttribute(TEXT("Parameters.IsPipelineOverrided"), bIsPipelineOverride));
+	AnalyticsHelper.Add(UsageImportAnalayticsId, FAnalyticsEventAttribute(TEXT("Parameters.bReplaceExisting"), ImportAssetParameters.bReplaceExisting));
 	if (!ImportAssetParameters.DestinationName.IsEmpty())
 	{
-		AnalyticsHelper->Add(UsageImportAnalayticsId, FAnalyticsEventAttribute(TEXT("Parameters.DestinationName"), ImportAssetParameters.DestinationName));
+		AnalyticsHelper.Add(UsageImportAnalayticsId, FAnalyticsEventAttribute(TEXT("Parameters.DestinationName"), ImportAssetParameters.DestinationName));
 	}
+	
 
 	if (this != &GetInterchangeManager())
 	{
@@ -1857,7 +1733,7 @@ UInterchangeManager::ImportInternal(const FString& ContentPath, const UInterchan
 		Translator->ReleaseSource();
 	}
 
-	AnalyticsHelper->Add(UsageImportAnalayticsId, FAnalyticsEventAttribute(TEXT("SourceExtension"), FPaths::GetExtension(SourceData->GetFilename())));
+	AnalyticsHelper.Add(UsageImportAnalayticsId, FAnalyticsEventAttribute(TEXT("SourceExtension"), FPaths::GetExtension(SourceData->GetFilename())));
 
 	bool bImportScene = ImportType == UE::Interchange::EImportType::ImportType_Scene;
 	const FInterchangeImportSettings& InterchangeImportSettings = FInterchangeProjectSettingsUtils::GetDefaultImportSettings(bImportScene);
@@ -1953,7 +1829,7 @@ UInterchangeManager::ImportInternal(const FString& ContentPath, const UInterchan
 							}
 
 							//Add an analytic so we can know if a not re-import was turn into a re-import
-							AnalyticsHelper->Add(UsageImportAnalayticsId, FAnalyticsEventAttribute(TEXT("Parameters.ConvertToReimport"), true));
+							AnalyticsHelper.Add(UsageImportAnalayticsId, FAnalyticsEventAttribute(TEXT("Parameters.ConvertToReimport"), true));
 						}
 					}
 				}
@@ -2026,11 +1902,9 @@ UInterchangeManager::ImportInternal(const FString& ContentPath, const UInterchan
 	TaskData.ImportLevel = ImportAssetParameters.ImportLevel;
 	TaskData.DestinationName = ImportAssetParameters.DestinationName;
 	TaskData.bReplaceExisting = ImportAssetParameters.bReplaceExisting;
-
-	TSharedRef<UE::Interchange::FImportAsyncHelper, ESPMode::ThreadSafe> AsyncHelper = CreateAsyncHelper(TaskData, ImportAssetParameters, AnalyticsGuid);
-	AsyncHelper->UniqueId = UniqueId;
-	AsyncHelper->bRunSynchronous = ImportAssetParameters.bRunSynchronous;
-
+	
+	InitAsyncHelper(AsyncHelper, TaskData, ImportAssetParameters);
+	
 	//We support only one source currently
 
 	//Create a duplicate of the source data, we need to be multithread safe so we copy it to control the life cycle. The async helper will hold it and delete it when the import task will be completed.
@@ -2050,11 +1924,11 @@ UInterchangeManager::ImportInternal(const FString& ContentPath, const UInterchan
 			return EarlyExit();
 		}
 
-		SourceTranslator->AnalyticsHelper = AnalyticsHelper;
+		SourceTranslator->AnalyticsHandler = AsyncHelper->AnalyticsHandler.Get();
 
 		if (bIsReimport)
 		{
-			//Set translator settings if we are doing a reimport
+			//Set translator settings if we are doing a re-import
 			if (const UInterchangeTranslatorSettings* InterchangeTranslatorSettings = OriginalAssetImportData->GetTranslatorSettings())
 			{
 				SourceTranslator->SetSettings(InterchangeTranslatorSettings);
@@ -2140,9 +2014,11 @@ UInterchangeManager::ImportInternal(const FString& ContentPath, const UInterchan
 			{
 				AsyncHelper->Pipelines.Add(Pipeline);
 				AsyncHelper->OriginalPipelines.Add(Pipeline);
-				UE::Interchange::Private::FillPipelineAnalyticData(*AnalyticsHelper, Pipeline, UniqueId, FString());
 			}
 		}
+
+		AnalyticsHelper.Send(LastImportPipelines, UniqueId);
+		
 	}
 	else
 	{
@@ -2154,8 +2030,10 @@ UInterchangeManager::ImportInternal(const FString& ContentPath, const UInterchan
 #else
 		bool bShowPipelineStacksConfigurationDialog = false;
 #endif
+
 		
-		AnalyticsHelper->AddThreadSafe(UsageImportAnalayticsId, FAnalyticsEventAttribute(TEXT("ShowImportDialog"), bShowPipelineStacksConfigurationDialog));
+		AnalyticsHelper.Add(UsageImportAnalayticsId, FAnalyticsEventAttribute(TEXT("ShowImportDialog"), bShowPipelineStacksConfigurationDialog));
+
 
 		//If we need to display the dialog we want to translate the source file before showing it
 		if (RegisteredPipelineConfiguration && bShowPipelineStacksConfigurationDialog)
@@ -2229,7 +2107,7 @@ UInterchangeManager::ImportInternal(const FString& ContentPath, const UInterchan
 							SourcePipeline = PythonPipelineAsset->GeneratedPipeline;
 						}
 					}
-					if (SourcePipeline && SourcePipeline->SupportReimport()) //Its possible a pipeline doesnt exist anymore so it wont load into memory when we loading the outer asset
+					if (SourcePipeline && SourcePipeline->SupportReimport()) //Its possible a pipeline doesn't exist anymore so it wont load into memory when we loading the outer asset
 					{
 						//Duplicate the pipeline saved in the asset import data
 						UInterchangePipelineBase* GeneratedPipeline = Cast<UInterchangePipelineBase>(StaticDuplicateObject(SourcePipeline, GetTransientPackage()));
@@ -2414,16 +2292,18 @@ UInterchangeManager::ImportInternal(const FString& ContentPath, const UInterchan
 				{
 					AsyncHelper->OriginalPipelines.Add(Pipeline);
 				}
-				UE::Interchange::Private::FillPipelineAnalyticData(*AnalyticsHelper, Pipeline, UniqueId, FString());
+
 			}
+
+			AnalyticsHelper.Send(AsyncHelper->Pipelines, UniqueId);
 		}
 	}
 
 	//Cancel the import do not queue task
 	if (bImportCanceled || bImportAborted)
 	{
-		AnalyticsHelper->AddThreadSafe(UsageImportAnalayticsId, FAnalyticsEventAttribute(TEXT("Canceled"), bImportCanceled));
-
+		
+		AnalyticsHelper.Add(UsageImportAnalayticsId, FAnalyticsEventAttribute(TEXT("Canceled"), bImportCanceled));
 		AsyncHelper->InitCancel();
 
 		//if it is a re-import, re-instate the backedup source data:
@@ -2485,6 +2365,29 @@ UInterchangeManager::ImportInternal(const FString& ContentPath, const UInterchan
 	return TTuple<UE::Interchange::FAssetImportResultRef, UE::Interchange::FSceneImportResultRef>{ AsyncHelper->AssetImportResult, AsyncHelper->SceneImportResult };
 }
 
+void UInterchangeManager::SetAnalyticsHandlerClass(const TSoftClassPtr<UInterchangeAnalyticsHandlerBase>& InAnalyticsHandlerClass)
+{
+	AnalyticsHandlerClass = InAnalyticsHandlerClass;
+}
+
+TStrongObjectPtr<UInterchangeAnalyticsHandlerBase> UInterchangeManager::CreateAnalyticsHandler(const FString& ContentPath)
+{
+	TStrongObjectPtr<UInterchangeAnalyticsHandlerBase> AnalyticsHandler = nullptr;
+	if (!AnalyticsHandlerClass.IsNull())
+	{
+		if (UClass* LoadedAnalyticsHandlerClass = AnalyticsHandlerClass.LoadSynchronous())
+		{
+			if (UInterchangeAnalyticsHandlerBase* AnalyticsHandlerBase = NewObject<UInterchangeAnalyticsHandlerBase>(GetTransientPackage(), LoadedAnalyticsHandlerClass, NAME_None, RF_NoFlags))
+			{
+				AnalyticsHandler = TStrongObjectPtr<UInterchangeAnalyticsHandlerBase>(AnalyticsHandlerBase);
+				AnalyticsHandler->SetGUID(FGuid::NewDeterministicGuid(ContentPath, static_cast<uint64>(FDateTime::UtcNow().GetTicks())));
+			}
+		}
+	}
+
+	return AnalyticsHandler;
+}
+
 void UInterchangeManager::SetEditorUtilities(UClass* EditorUtilitiesClass)
 {
 	//Create the object
@@ -2678,9 +2581,8 @@ const UClass* UInterchangeManager::GetRegisteredFactoryClass(const UClass* Class
 	return Result;
 }
 
-TSharedRef<UE::Interchange::FImportAsyncHelper, ESPMode::ThreadSafe> UInterchangeManager::CreateAsyncHelper(const UE::Interchange::FImportAsyncHelperData& Data, const FImportAssetParameters& ImportAssetParameters, const FGuid& InAnalyticsID)
+void UInterchangeManager::InitAsyncHelper(TSharedRef<UE::Interchange::FImportAsyncHelper, ESPMode::ThreadSafe>& AsyncHelper, const UE::Interchange::FImportAsyncHelperData& Data, const FImportAssetParameters& ImportAssetParameters)
 {
-	TSharedRef<UE::Interchange::FImportAsyncHelper, ESPMode::ThreadSafe> AsyncHelper = MakeShared<UE::Interchange::FImportAsyncHelper, ESPMode::ThreadSafe>(InAnalyticsID);
 	//Copy the task data
 	AsyncHelper->TaskData = Data;
 	
@@ -2701,8 +2603,6 @@ TSharedRef<UE::Interchange::FImportAsyncHelper, ESPMode::ThreadSafe> UInterchang
 	AsyncHelper->bRuntimeOrPIE = GetEditorUtilities() ? GetEditorUtilities()->IsRuntimeOrPIE() : false;
 
 	AsyncHelper->TaskExecuteStarted = FPlatformTime::Seconds();
-
-	return AsyncHelper;
 }
 
 void UInterchangeManager::ReleaseAsyncHelper(TWeakPtr<UE::Interchange::FImportAsyncHelper, ESPMode::ThreadSafe> AsyncHelper)
diff --git a/Engine/Source/Runtime/Interchange/Engine/Private/Tasks/InterchangeTaskCompletion.cpp b/Engine/Source/Runtime/Interchange/Engine/Private/Tasks/InterchangeTaskCompletion.cpp
--- a/Engine/Source/Runtime/Interchange/Engine/Private/Tasks/InterchangeTaskCompletion.cpp
+++ b/Engine/Source/Runtime/Interchange/Engine/Private/Tasks/InterchangeTaskCompletion.cpp
@@ -6,6 +6,7 @@
 #include "CoreMinimal.h"
 #include "Engine/World.h"
 #include "GameFramework/Actor.h"
+#include "InterchangeAnalyticsHandlerBase.h"
 #include "InterchangeEngineLogPrivate.h"
 #include "InterchangeFactoryBase.h"
 #include "InterchangeManager.h"
@@ -20,7 +21,6 @@
 #include "UObject/ObjectMacros.h"
 #include "UObject/WeakObjectPtrTemplates.h"
 
-
 namespace UE::Interchange::Private::ObjectDeletionUtils
 {
 	static void PurgeObject(UObject* Object)
@@ -206,7 +206,12 @@ void UE::Interchange::FTaskCompletion_GameThread::Execute()
 	TSharedPtr<FImportAsyncHelper, ESPMode::ThreadSafe> AsyncHelper = WeakAsyncHelper.Pin();
 	check(AsyncHelper.IsValid());
 
-	AsyncHelper->SendAnalyticImportEndData();
+	if (AsyncHelper->AnalyticsHandler.IsValid())
+	{
+		FInterchangeImportResultAnalyticsInfo ImportResultAnalyticsInfo = FImportAsyncHelper::CreateImportResultAnalyticsInfo(*AsyncHelper.Get());
+		AsyncHelper->AnalyticsHandler->Send(ImportResultAnalyticsInfo);
+	}
+	
 	//No need anymore of the translators sources
 	AsyncHelper->ReleaseTranslatorsSource();
 
diff --git a/Engine/Source/Runtime/Interchange/Engine/Public/InterchangeAnalyticsHandlerBase.h b/Engine/Source/Runtime/Interchange/Engine/Public/InterchangeAnalyticsHandlerBase.h
new file mode 100644
--- /dev/null
+++ b/Engine/Source/Runtime/Interchange/Engine/Public/InterchangeAnalyticsHandlerBase.h
@@ -0,0 +1,146 @@
+// Copyright Epic Games, Inc. All Rights Reserved.
+
+#pragma once
+
+#include "UObject/Object.h"
+#include "UObject/ObjectMacros.h"
+#include "UObject/WeakObjectPtrTemplates.h"
+
+#include "Containers/Array.h"
+#include "Containers/Map.h"
+#include "Containers/UnrealString.h"
+
+#include "HAL/CriticalSection.h"
+
+#include "Misc/Guid.h"
+#include "Misc/Optional.h"
+#include "Misc/NotNull.h"
+#include "AnalyticsEventAttribute.h"
+
+#include "InterchangeAnalyticsHandlerBase.generated.h"
+
+#define UE_API INTERCHANGEENGINE_API
+
+class UInterchangePipelineBase;
+
+class UInterchangeSourceData;
+class UInterchangeSourceNode;
+
+class UInterchangeResultsContainer;
+
+/**
+ * If you add or remove from this enum, make sure to update the corresponding strings
+ * in UInterchangeAnalyticsHandlerBase::GetDefaultAnalyticsCategoryName.
+ */
+enum class EInterchangeAnalyticsDefaultCategory : uint8
+{
+	NotSpecified = 0,
+	Models,
+	Scenes,
+	Materials,
+	Textures,
+	Sounds,
+	Meshes,
+	StaticMeshes,
+	SkeletalMeshesAndAnimations,
+	GeometryCaches,
+	Grooms,
+	InterchangeAnalyticsCategory_COUNT
+};
+
+struct FInterchangeImportResultAnalyticsInfo
+{
+	bool bCancel;
+	int32 UniqueId;
+
+	TNotNull<const UInterchangeSourceData*> PrimarySourceData;
+	TNotNull<const UInterchangeSourceNode*> PrimarySourceNode;
+
+	TOptional<const UInterchangeResultsContainer*> AssetImportResultsContainer;
+	TOptional<const UInterchangeResultsContainer*> SceneImportResultsContainer;
+
+	TMap<int32, TArray<TObjectPtr<UObject>>> ImportedAssetsPerSourceIndex;
+	TMap<int32, TArray<TObjectPtr<UObject>>> ImportedSceneObjectsPerSourceIndex;
+};
+
+/**
+ * This is the base class implementation for the Interchange Analytics. This will NOT collect any analytics BY DESIGN.
+ * DO NOT MAKE CHANGES TO THE IMPLEMENTATIONS OF THE FUNCTIONS FOR THIS CLASS. See UInterchangeAnalyticsHandlerDefault instead 
+ * for the default implementation.
+ * 
+ * Interchange Manager exposes API to set the TSoftClassPtr for a class to use to collect Analytics.
+ * By default InterchangeAnalyticsModule registers UInterchangeAnalyticsHandlerDefault as the Analytics Handler Class.
+ */
+UCLASS(BlueprintType, Blueprintable, Abstract, MinimalAPI)
+class UInterchangeAnalyticsHandlerBase : public UObject
+{
+    GENERATED_BODY()
+
+public:
+	UE_API static FAnalyticsEventAttribute CreateAnalyticsIDAttribute(FGuid AnalyticsID);
+
+	UE_API static FString GetDefaultAnalyticsCategoryName(EInterchangeAnalyticsDefaultCategory DefaultCategory);
+
+protected:
+	UE_API static FString GetPipelineAnalyticsEventIdentifier();
+
+	UE_API static FString GetImportResultEventIdentifier();
+
+public:
+	/**
+	 * Can only be set once per instance of the Analytics Handler.
+	 */
+	UE_API void SetGUID(const FGuid& InGUID);
+
+	/**
+	 * Adds an array of attributes to the event identifier in a thread safe manner.
+	 */
+    UE_API void Append(const FString& Identifier, const TArray<FAnalyticsEventAttribute>& ToAdd);
+
+	/**
+	 * Add the analytics attribute to an event identifier in a thread safe manner.
+	 */
+	UE_API void Add(const FString& Identifier, const FAnalyticsEventAttribute& Entry);
+
+	/**
+	 * Clears accumulated attributes for an event identifier. Useful in Translator Analytic Events context.
+	 */
+	UE_API void Clear(const FString& Identifier);
+
+	/**
+	 * Send Accumulated Analytics.
+	 */
+    UE_API void Flush();
+
+	/**
+	 * Send Pipelines Analytics.
+	 */
+	virtual void Send(const TArray<UInterchangePipelineBase*>& Pipelines, const int32 AsyncHelperUniqueId) PURE_VIRTUAL(UInterchangeAnalyticsHandlerBase::Send, )
+	
+	/**
+	 * Send Import Result Analytics.
+	 */
+	virtual void Send(const FInterchangeImportResultAnalyticsInfo& ImportResultAnalyticsInfo) PURE_VIRTUAL(UInterchangeAnalyticsHandlerBase::Send, )
+	
+	/**
+	* Creates an Analytics Event Attribute for the current instance of the analytics helper.
+	* Helpful to be used on external analytics events not sent by AnalyticsHelper.
+	*/
+    UE_API virtual FAnalyticsEventAttribute GetAnalyticsIDAttribute() const;
+
+protected:
+	virtual void Execute_Append(const FString& Identifier, const TArray<FAnalyticsEventAttribute>& ToAdd) PURE_VIRTUAL(UInterchangeAnalyticsHandlerBase::Execute_Append, )
+		
+	virtual void Execute_Add(const FString& Identifier, const FAnalyticsEventAttribute& Entry) PURE_VIRTUAL(UInterchangeAnalyticsHandlerBase::Execute_Add, )
+
+	virtual void Execute_Clear(const FString& Identifier) PURE_VIRTUAL(UInterchangeAnalyticsHandlerBase::Execute_Clear, )
+	
+	virtual void Execute_Flush() PURE_VIRTUAL(UInterchangeAnalyticsHandlerBase::Execute_Flush, )
+
+private:
+	FCriticalSection CriticalSection;
+    FGuid AnalyticsID;
+};
+
+
+#undef UE_API
\ No newline at end of file
diff --git a/Engine/Source/Runtime/Interchange/Engine/Public/InterchangeAnalyticsHelper.h b/Engine/Source/Runtime/Interchange/Engine/Public/InterchangeAnalyticsHelper.h
deleted file mode 100644
diff --git a/Engine/Source/Runtime/Interchange/Engine/Public/InterchangeManager.h b/Engine/Source/Runtime/Interchange/Engine/Public/InterchangeManager.h
--- a/Engine/Source/Runtime/Interchange/Engine/Public/InterchangeManager.h
+++ b/Engine/Source/Runtime/Interchange/Engine/Public/InterchangeManager.h
@@ -13,6 +13,7 @@
 #include "HAL/Thread.h"
 #include "HAL/ThreadSafeBool.h"
 #include "InterchangeAssetImportData.h"
+#include "InterchangeAnalyticsHandlerBase.h"
 #include "InterchangeEditorUtilitiesBase.h"
 #include "InterchangeFactoryBase.h"
 #include "InterchangePipelineConfigurationBase.h"
@@ -33,6 +34,8 @@
 
 #include "InterchangeManager.generated.h"
 
+#define UE_API INTERCHANGEENGINE_API
+
 class FAsyncTaskNotification;
 class UInterchangeBlueprintPipelineBase;
 class UInterchangeFactoryBaseNode;
@@ -42,11 +45,6 @@ class ULevel;
 
 struct FAnalyticsEventAttribute;
 
-namespace UE::Interchange
-{
-	struct FAnalyticsHelper;
-}
-
 /** Some utility delegates for automating Interchange. */
 DECLARE_DYNAMIC_DELEGATE_OneParam(FOnObjectImportDoneDynamic, UObject*, Object);
 DECLARE_DELEGATE_OneParam(FOnObjectImportDoneNative, UObject*);
@@ -79,8 +77,8 @@ namespace UE
 		class FScopedInterchangeImportEnableState
 		{
 		public:
-			INTERCHANGEENGINE_API explicit FScopedInterchangeImportEnableState(const bool bScopeValue);
-			INTERCHANGEENGINE_API ~FScopedInterchangeImportEnableState();
+			UE_API explicit FScopedInterchangeImportEnableState(const bool bScopeValue);
+			UE_API ~FScopedInterchangeImportEnableState();
 		private:
 			bool bOriginalInterchangeImportEnableState;
 		};
@@ -88,9 +86,9 @@ namespace UE
 		class FScopedSourceData
 		{
 		public:
-			INTERCHANGEENGINE_API explicit FScopedSourceData(const FString& Filename);
-			INTERCHANGEENGINE_API ~FScopedSourceData();
-			INTERCHANGEENGINE_API UInterchangeSourceData* GetSourceData() const;
+			UE_API explicit FScopedSourceData(const FString& Filename);
+			UE_API ~FScopedSourceData();
+			UE_API UInterchangeSourceData* GetSourceData() const;
 		private:
 			TStrongObjectPtr<UInterchangeSourceData> SourceDataPtr = nullptr;
 		};
@@ -98,9 +96,9 @@ namespace UE
 		class FScopedTranslator
 		{
 		public:
-			INTERCHANGEENGINE_API explicit FScopedTranslator(const UInterchangeSourceData* SourceData);
-			INTERCHANGEENGINE_API ~FScopedTranslator();
-			INTERCHANGEENGINE_API UInterchangeTranslatorBase* GetTranslator();
+			UE_API explicit FScopedTranslator(const UInterchangeSourceData* SourceData);
+			UE_API ~FScopedTranslator();
+			UE_API UInterchangeTranslatorBase* GetTranslator();
 
 		private:
 			TStrongObjectPtr<UInterchangeTranslatorBase> ScopedTranslatorPtr = nullptr;
@@ -109,9 +107,9 @@ namespace UE
 		class FScopedBaseNodeContainer
 		{
 		public:
-			INTERCHANGEENGINE_API FScopedBaseNodeContainer();
-			INTERCHANGEENGINE_API ~FScopedBaseNodeContainer();
-			INTERCHANGEENGINE_API UInterchangeBaseNodeContainer* GetBaseNodeContainer();
+			UE_API FScopedBaseNodeContainer();
+			UE_API ~FScopedBaseNodeContainer();
+			UE_API UInterchangeBaseNodeContainer* GetBaseNodeContainer();
 
 		private:
 			TStrongObjectPtr<UInterchangeBaseNodeContainer> ScopedBaseNodeContainerPtr = nullptr;
@@ -151,7 +149,7 @@ namespace UE
 		class FImportResult : protected FGCObject
 		{
 		public:
-			INTERCHANGEENGINE_API FImportResult();
+			UE_API FImportResult();
 
 			FImportResult(FImportResult&&) = delete;
 			FImportResult& operator=(FImportResult&&) = delete;
@@ -169,33 +167,33 @@ namespace UE
 				Done
 			};
 
-			INTERCHANGEENGINE_API EStatus GetStatus() const;
+			UE_API EStatus GetStatus() const;
 
-			INTERCHANGEENGINE_API bool IsValid() const;
+			UE_API bool IsValid() const;
 
-			INTERCHANGEENGINE_API void SetInProgress();
-			INTERCHANGEENGINE_API void SetDone();
-			INTERCHANGEENGINE_API void WaitUntilDone(bool bSynchronous = false);
+			UE_API void SetInProgress();
+			UE_API void SetDone();
+			UE_API void WaitUntilDone(bool bSynchronous = false);
 
 			// Assets are only made available once they have been completely imported (passed through the entire import pipeline).
 			// While the status isn't EStatus::Done, the list can grow between subsequent calls.
 			// FAssetImportResult holds a reference to the assets so that they aren't garbage collected.
-			INTERCHANGEENGINE_API const TArray< UObject* >& GetImportedObjects() const;
+			UE_API const TArray< UObject* >& GetImportedObjects() const;
 
 			// Helper to get the first asset of a certain class. Use when expecting a single asset of that class to be imported, because the order isn't deterministic.
-			INTERCHANGEENGINE_API UObject* GetFirstAssetOfClass(UClass* InClass) const;
+			UE_API UObject* GetFirstAssetOfClass(UClass* InClass) const;
 
 			// Return the results of this asset import operation.
 			UInterchangeResultsContainer* GetResults() const { return Results; }
 
 			// Adds an asset to the list of imported assets.
-			INTERCHANGEENGINE_API void AddImportedObject(UObject* ImportedObject);
+			UE_API void AddImportedObject(UObject* ImportedObject);
 
 			// Callback when the status switches to done.
-			INTERCHANGEENGINE_API void OnDone(TFunction< void(FImportResult&) > Callback);
+			UE_API void OnDone(TFunction< void(FImportResult&) > Callback);
 
 			//Set the async helper that own this import result
-			INTERCHANGEENGINE_API void SetAsyncHelper(TWeakPtr<class FImportAsyncHelper> InAsyncHelper);
+			UE_API void SetAsyncHelper(TWeakPtr<class FImportAsyncHelper> InAsyncHelper);
 
 			// Internal delegates. To set these, use the FImportAssetParameters when calling the Interchange import functions.
 			FOnObjectImportDoneDynamic OnObjectDone;
@@ -206,7 +204,7 @@ namespace UE
 
 		protected:
 			/* FGCObject interface */
-			INTERCHANGEENGINE_API virtual void AddReferencedObjects(FReferenceCollector& Collector) override;
+			UE_API virtual void AddReferencedObjects(FReferenceCollector& Collector) override;
 			virtual FString GetReferencerName() const override
 			{
 				return TEXT("UE::Interchange::FImportResult");
@@ -232,7 +230,10 @@ namespace UE
 		class FImportAsyncHelper : protected FGCObject
 		{
 		public:
-			FImportAsyncHelper(const FGuid& AnalyticsID);
+			static FInterchangeImportResultAnalyticsInfo CreateImportResultAnalyticsInfo(FImportAsyncHelper& AsyncHelper);
+
+		public:
+			FImportAsyncHelper();
 
 			~FImportAsyncHelper();
 
@@ -311,7 +312,7 @@ namespace UE
 				mutable bool bPostEditChangeCalled = false; //This field is set by the PreCompletionTask and need to be mutable
 			};
 
-			INTERCHANGEENGINE_API FImportedObjectInfo& AddDefaultImportedAssetGetRef(int32 SourceIndex);
+			UE_API FImportedObjectInfo& AddDefaultImportedAssetGetRef(int32 SourceIndex);
 			const FImportedObjectInfo* FindImportedAssets(int32 SourceIndex, TFunction< bool(const FImportedObjectInfo& ImportedObjects) > Predicate) const;
 			void IterateImportedAssets(int32 SourceIndex, TFunction< void(const TArray<FImportedObjectInfo>& ImportedObjects) > Callback) const;
 			void IterateImportedAssetsPerSourceIndex(TFunction< void(int32 SourceIndex, const TArray<FImportedObjectInfo>& ImportedObjects) > Callback) const;
@@ -334,7 +335,6 @@ namespace UE
 			//If we cancel the tasks, we set this Boolean to true.
 			std::atomic<bool> bCancel;
 
-			void SendAnalyticImportEndData();
 			void ReleaseTranslatorsSource();
 
 			/**
@@ -348,6 +348,10 @@ namespace UE
 
 			double TaskExecuteStarted;
 
+			// This class owns the Analytics Handler and manages its lifetime.
+			// Analytics Handler will be released at the end of import when the ImportAsyncHandler is deleted.
+			TStrongObjectPtr<UInterchangeAnalyticsHandlerBase> AnalyticsHandler;
+
 		private:
 			FCriticalSection ClassPermissionLock;
 			// Set of classes whose creation has been denied.
@@ -368,12 +372,10 @@ namespace UE
 
 			mutable FCriticalSection ImportedSceneObjectsPerSourceIndexLock;
 			TMap<int32, TArray<FImportedObjectInfo>> ImportedSceneObjectsPerSourceIndex;
-
-			FGuid AnalyticsID;
 		};
 
 		/* This function takes an asset that represents a pipeline and generates a UInterchangePipelineBase asset. */
-		INTERCHANGEENGINE_API UInterchangePipelineBase* GeneratePipelineInstance(const FSoftObjectPath& PipelineInstance);
+		UE_API UInterchangePipelineBase* GeneratePipelineInstance(const FSoftObjectPath& PipelineInstance);
 	} //ns Interchange
 } //ns UE
 
@@ -390,13 +392,13 @@ public:
 	TArray<FSoftObjectPath> OverridePipelines;
 
 	UFUNCTION(BlueprintCallable, Category = "Interchange | PipelineStackOverride")
-	INTERCHANGEENGINE_API void AddPythonPipeline(UInterchangePythonPipelineBase* PipelineBase);
+	UE_API void AddPythonPipeline(UInterchangePythonPipelineBase* PipelineBase);
 
 	UFUNCTION(BlueprintCallable, Category = "Interchange | PipelineStackOverride")
-	INTERCHANGEENGINE_API void AddBlueprintPipeline(UInterchangeBlueprintPipelineBase* PipelineBase);
+	UE_API void AddBlueprintPipeline(UInterchangeBlueprintPipelineBase* PipelineBase);
 
 	UFUNCTION(BlueprintCallable, Category = "Interchange | PipelineStackOverride")
-	INTERCHANGEENGINE_API void AddPipeline(UInterchangePipelineBase* PipelineBase);
+	UE_API void AddPipeline(UInterchangePipelineBase* PipelineBase);
 };
 
 USTRUCT(BlueprintType)
@@ -487,15 +489,15 @@ public:
 	}
 
 	/** Return the Interchange Manager singleton. */
-	static INTERCHANGEENGINE_API UInterchangeManager& GetInterchangeManager();
+	static UE_API UInterchangeManager& GetInterchangeManager();
 
 	/** Return the CVar that enables or disables Interchange. */
-	static INTERCHANGEENGINE_API bool IsInterchangeImportEnabled();
+	static UE_API bool IsInterchangeImportEnabled();
 	/** Set the CVar that enables or disables Interchange. */
-	static INTERCHANGEENGINE_API void SetInterchangeImportEnabled(bool bEnabled);
+	static UE_API void SetInterchangeImportEnabled(bool bEnabled);
 
 	/** Checks if there are any imports in progress.*/
-	static INTERCHANGEENGINE_API bool IsImporting();
+	static UE_API bool IsImporting();
 
 	/** Delegate type that is fired when new assets have been imported. Note: InCreatedObject can be NULL if the import failed. Params: UFactory* InFactory, UObject* InCreatedObject. */
 	DECLARE_MULTICAST_DELEGATE_OneParam(FInterchangeOnAssetPostImport, UObject*);
@@ -520,7 +522,7 @@ public:
 	// Called when before the application is exiting.
 	FSimpleMulticastDelegate OnPreDestroyInterchangeManager;
 
-	INTERCHANGEENGINE_API void SanitizeNameInline(FString& NameToSanitize, const ESanitizeNameTypeFlags NameType);
+	UE_API void SanitizeNameInline(FString& NameToSanitize, const ESanitizeNameTypeFlags NameType);
 	/**
 	 * All translators must be registered with the manager.
 	 * @Param Translator - The UClass of the translator you want to register.
@@ -529,7 +531,7 @@ public:
 	 * @Note If you register the same class multiple times, this returns true for every call.
 	 * @Note The order in which the translators are registered will be the same as the order used to select a translator to import a file.
 	 */
-	INTERCHANGEENGINE_API bool RegisterTranslator(const UClass* TranslatorClass);
+	UE_API bool RegisterTranslator(const UClass* TranslatorClass);
 
 	/**
 	 * All factories must be registered with the manager.
@@ -538,7 +540,7 @@ public:
 	 *
 	 * @Note If you register the same class multiple time, this returns true for every call.
 	 */
-	INTERCHANGEENGINE_API bool RegisterFactory(const UClass* Factory);
+	UE_API bool RegisterFactory(const UClass* Factory);
 
 	/**
 	 * All writers must be registered with the manager,
@@ -547,7 +549,7 @@ public:
 	 *
 	 * @Note If you register the same class multiple time, this returns true for every call.
 	 */
-	INTERCHANGEENGINE_API bool RegisterWriter(const UClass* Writer);
+	UE_API bool RegisterWriter(const UClass* Writer);
 
 	/**
 	 * All converters must be registered with the manager.
@@ -556,7 +558,7 @@ public:
 	 *
 	 * @Note If you register the same class multiple time, this returns true for every call.
 	 */
-	INTERCHANGEENGINE_API bool RegisterImportDataConverter(const UClass* Converter);
+	UE_API bool RegisterImportDataConverter(const UClass* Converter);
 
 	/**
 	 * Call all the registered converters to see if any converter can convert the data.
@@ -564,7 +566,7 @@ public:
 	 * @Param Extension - The file extension we want to import.
 	 * @return true if one of the converter has converted the data, or false otherwise.
 	 */
-	INTERCHANGEENGINE_API bool ConvertImportData(UObject* Asset, const FString& Extension) const;
+	UE_API bool ConvertImportData(UObject* Asset, const FString& Extension) const;
 
 	/**
 	 * Call all the registered converter, until one can convert the source import data to UInterchangeAssetImportData
@@ -572,7 +574,7 @@ public:
 	 * @Param ImportAssetParameters - The interchange import asset parameters.
 	 * @return true if one of the converter has convert the data. False otherwise.
 	 */
-	INTERCHANGEENGINE_API bool ConvertImportData(const UObject* SourceImportData, FImportAssetParameters& ImportAssetParameters) const;
+	UE_API bool ConvertImportData(const UObject* SourceImportData, FImportAssetParameters& ImportAssetParameters) const;
 
 	/**
 	 * Call all the registered converter, until one can convert the source data to the destination
@@ -581,26 +583,26 @@ public:
 	 * @Param DestinationClass - The class representing the DestinationImportData
 	 * @return true if one of the converter has convert the data. False otherwise.
 	 */
-	INTERCHANGEENGINE_API bool ConvertImportData(const UObject* SourceImportData, const UClass* DestinationClass, UObject** DestinationImportData) const;
+	UE_API bool ConvertImportData(const UObject* SourceImportData, const UClass* DestinationClass, UObject** DestinationImportData) const;
 
 	/**
 	 * Returns the list of supported formats for a given translator type.
 	 */
 	UFUNCTION(BlueprintCallable, Category = "Interchange | Import Manager")
-	INTERCHANGEENGINE_API TArray<FString> GetSupportedFormats(const EInterchangeTranslatorType ForTranslatorType) const;
+	UE_API TArray<FString> GetSupportedFormats(const EInterchangeTranslatorType ForTranslatorType) const;
 
 	/**
 	 * Returns the list of formats supporting the specified translator asset type.
 	 */
 	UFUNCTION(BlueprintCallable, Category = "Interchange | Import Manager")
-	INTERCHANGEENGINE_API TArray<FString> GetSupportedAssetTypeFormats(const EInterchangeTranslatorAssetType ForTranslatorAssetType, const EInterchangeTranslatorType ForTranslatorType = EInterchangeTranslatorType::Invalid, bool bStrictMatchTranslatorType = false) const;
+	UE_API TArray<FString> GetSupportedAssetTypeFormats(const EInterchangeTranslatorAssetType ForTranslatorAssetType, const EInterchangeTranslatorType ForTranslatorType = EInterchangeTranslatorType::Invalid, bool bStrictMatchTranslatorType = false) const;
 
 	/**
 	 * Returns the list of supported formats for a given Object.
 	 * 
 	 */
 	UFUNCTION(BlueprintCallable, Category = "Interchange | Import Manager")
-	INTERCHANGEENGINE_API TArray<FString> GetSupportedFormatsForObject(const UObject* Object, int32 SourceFileIndex) const;
+	UE_API TArray<FString> GetSupportedFormatsForObject(const UObject* Object, int32 SourceFileIndex) const;
 
 	/**
 	 * Check whether there is a registered translator for this source data.
@@ -611,7 +613,7 @@ public:
 	 * @Note: Temporary until FBX level import is production ready
 	 */
 	UFUNCTION(BlueprintCallable, Category = "Interchange | Import Manager")
-	INTERCHANGEENGINE_API bool CanTranslateSourceData(const UInterchangeSourceData* SourceData, bool bSceneImportOnly = false) const;
+	UE_API bool CanTranslateSourceData(const UInterchangeSourceData* SourceData, bool bSceneImportOnly = false) const;
 
 	/**
 	 * Returns true if Interchange can create this type of asset and is able to translate its source files.
@@ -619,7 +621,7 @@ public:
 	 * @Param OutFilenames - An array that is filled with the object's source filenames if the operation is successful.
 	 */
 	UFUNCTION(BlueprintCallable, Category = "Interchange | Import Manager")
-	INTERCHANGEENGINE_API bool CanReimport(const UObject* Object, TArray<FString>& OutFilenames) const;
+	UE_API bool CanReimport(const UObject* Object, TArray<FString>& OutFilenames) const;
 
 	/**
 	 * Call this to start a synchronous asset import process.
@@ -633,7 +635,7 @@ public:
 	 * @Note - In blueprint depending on the event you use to start the import its possible to have a deadlock, use the async function if its what you are experimenting
 	 */
 	UFUNCTION(BlueprintCallable, Category = "Interchange | Import Manager")
-	INTERCHANGEENGINE_API bool ImportAsset(const FString& ContentPath, const UInterchangeSourceData* SourceData, const FImportAssetParameters& ImportAssetParameters, TArray<UObject*>& OutImportedObjects);
+	UE_API bool ImportAsset(const FString& ContentPath, const UInterchangeSourceData* SourceData, const FImportAssetParameters& ImportAssetParameters, TArray<UObject*>& OutImportedObjects);
 
 	/**
 	 * Call this to start a synchronous asset import process.
@@ -645,7 +647,7 @@ public:
 	 * @return true if the import succeeds, or false otherwise.
 	 *
 	 */
-	INTERCHANGEENGINE_API bool ImportAsset(const FString& ContentPath, const UInterchangeSourceData* SourceData, const FImportAssetParameters& ImportAssetParameters);
+	UE_API bool ImportAsset(const FString& ContentPath, const UInterchangeSourceData* SourceData, const FImportAssetParameters& ImportAssetParameters);
 
 	/**
 	 * Call this to start a synchronous asset import process.
@@ -656,7 +658,7 @@ public:
 	 * @param ImportAssetParameters - All parameters that need to be passed to the import asset function.
 	 * @return return an import result which can be use to know when the asynchronous import is terminate.
 	 */
-	INTERCHANGEENGINE_API UE::Interchange::FAssetImportResultRef ImportAssetWithResult(const FString& ContentPath, const UInterchangeSourceData* SourceData, const FImportAssetParameters& ImportAssetParameters);
+	UE_API UE::Interchange::FAssetImportResultRef ImportAssetWithResult(const FString& ContentPath, const UInterchangeSourceData* SourceData, const FImportAssetParameters& ImportAssetParameters);
 
 	/**
 	 * Call this to start an asynchronous asset import process.
@@ -667,7 +669,7 @@ public:
 	 * @param ImportAssetParameters - All parameters that need to be passed to the import asset function.
 	 * @return return an import result which can be use to know when the asynchronous import is terminate.
 	 */
-	INTERCHANGEENGINE_API UE::Interchange::FAssetImportResultRef ImportAssetAsync(const FString& ContentPath, const UInterchangeSourceData* SourceData, const FImportAssetParameters& ImportAssetParameters);
+	UE_API UE::Interchange::FAssetImportResultRef ImportAssetAsync(const FString& ContentPath, const UInterchangeSourceData* SourceData, const FImportAssetParameters& ImportAssetParameters);
 	
 	/**
 	 * Call this from blueprint or python to start an asynchronous asset import process.
@@ -679,7 +681,7 @@ public:
 	 * @return true if the import was started, or false otherwise.
 	 */
 	UFUNCTION(BlueprintCallable, Category = "Interchange | Import Manager")
-	INTERCHANGEENGINE_API bool ScriptedImportAssetAsync(const FString& ContentPath, const UInterchangeSourceData* SourceData, const FImportAssetParameters& ImportAssetParameters);
+	UE_API bool ScriptedImportAssetAsync(const FString& ContentPath, const UInterchangeSourceData* SourceData, const FImportAssetParameters& ImportAssetParameters);
 
 	/**
 	 * Call this to start a synchronous asset re-import process.
@@ -696,7 +698,7 @@ public:
 	 * @Note - In blueprint depending on the event you use to start the import its possible to have a deadlock, use the async function if its what you are experimenting
 	 */
  	UFUNCTION(BlueprintCallable, Category = "Interchange | Import Manager")
- 	INTERCHANGEENGINE_API bool ReimportAsset(UObject* ObjectToReimport, const FImportAssetParameters& ImportAssetParameters, TArray<UObject*>& OutImportedObjects);
+ 	UE_API bool ReimportAsset(UObject* ObjectToReimport, const FImportAssetParameters& ImportAssetParameters, TArray<UObject*>& OutImportedObjects);
 
 	/**
 	 * Call this to start an asynchronous asset import process.
@@ -707,7 +709,7 @@ public:
 	 * @param ImportAssetParameters - All parameters that need to be passed to the import asset function.
 	 * @return return an import result which can be use to know when the asynchronous import is terminate.
 	 */
-	INTERCHANGEENGINE_API UE::Interchange::FAssetImportResultRef ReimportAssetAsync(UObject* ObjectToReimport, const FImportAssetParameters& ImportAssetParameters);
+	UE_API UE::Interchange::FAssetImportResultRef ReimportAssetAsync(UObject* ObjectToReimport, const FImportAssetParameters& ImportAssetParameters);
 
 	/**
 	 * Call this from blueprint or python to start an asynchronous asset import process.
@@ -724,7 +726,7 @@ public:
 	 * @Note - If you need an event for when the import is done use the ImportAssetParameters events.
 	 */
 	UFUNCTION(BlueprintCallable, Category = "Interchange | Import Manager")
-	INTERCHANGEENGINE_API bool ScriptedReimportAssetAsync(UObject* ObjectToReimport, const FImportAssetParameters& ImportAssetParameters);
+	UE_API bool ScriptedReimportAssetAsync(UObject* ObjectToReimport, const FImportAssetParameters& ImportAssetParameters);
 
 	/**
 	 * Call this to start a synchronous scene import process.
@@ -738,7 +740,7 @@ public:
 	 * @Note - In blueprint depending on the event you use to start the import its possible to have a deadlock, use the async function if its what you are experimenting
 	 */
 	UFUNCTION(BlueprintCallable, Category = "Interchange | Import Manager")
-	INTERCHANGEENGINE_API bool ImportScene(const FString& ContentPath, const UInterchangeSourceData* SourceData, const FImportAssetParameters& ImportAssetParameters);
+	UE_API bool ImportScene(const FString& ContentPath, const UInterchangeSourceData* SourceData, const FImportAssetParameters& ImportAssetParameters);
 
 	/**
 	 * Call this to start a asynchronous scene import process.
@@ -750,7 +752,7 @@ public:
 	 * @return true if the import was started, or false otherwise.
 	 */
 	UFUNCTION(BlueprintCallable, Category = "Interchange | Import Manager")
-	INTERCHANGEENGINE_API bool ScriptedImportSceneAsync(const FString& ContentPath, const UInterchangeSourceData* SourceData, const FImportAssetParameters& ImportAssetParameters);
+	UE_API bool ScriptedImportSceneAsync(const FString& ContentPath, const UInterchangeSourceData* SourceData, const FImportAssetParameters& ImportAssetParameters);
 
 	/**
 	 * Call this to start a asynchronous scene import process.
@@ -761,7 +763,7 @@ public:
 	 * @param ImportAssetParameters - All parameters that need to be passed to the import asset function.
 	 * @return return a pair of import result which can be use to know when the asynchronous import is terminate.
 	 */
-	INTERCHANGEENGINE_API TTuple<UE::Interchange::FAssetImportResultRef, UE::Interchange::FSceneImportResultRef>
+	UE_API TTuple<UE::Interchange::FAssetImportResultRef, UE::Interchange::FSceneImportResultRef>
 	ImportSceneAsync(const FString& ContentPath, const UInterchangeSourceData* SourceData, const FImportAssetParameters& ImportAssetParameters);
 
 	/**
@@ -772,7 +774,7 @@ public:
 	 * @return true if the export succeeds, or false otherwise.
 	 */
 	UFUNCTION(BlueprintCallable, Category = "Interchange | Export Manager")
-	INTERCHANGEENGINE_API bool ExportAsset(const UObject* Asset, bool bIsAutomated = false);
+	UE_API bool ExportAsset(const UObject* Asset, bool bIsAutomated = false);
 
 	/**
 	 * Call this to start a scene export process. The caller must specify a source data.
@@ -781,7 +783,7 @@ public:
 	 * @return true if the export succeeds, or false otherwise.
 	 */
 	UFUNCTION(BlueprintCallable, Category = "Interchange | Export Manager")
-	INTERCHANGEENGINE_API bool ExportScene(const UObject* World, bool bIsAutomated = false);
+	UE_API bool ExportScene(const UObject* World, bool bIsAutomated = false);
 
 	/*
 	* Script helper to create a source data object that points to a file on disk.
@@ -789,7 +791,7 @@ public:
 	* @return: A new UInterchangeSourceData.
 	*/
 	UFUNCTION(BlueprintCallable, Category = "Interchange | Import Manager")
-	static INTERCHANGEENGINE_API UInterchangeSourceData* CreateSourceData(const FString& InFileName);
+	static UE_API UInterchangeSourceData* CreateSourceData(const FString& InFileName);
 
 	/**
 	* Script helper to get a registered factory for a specified class.
@@ -797,22 +799,22 @@ public:
 	* @return: The registered factory class if found, or NULL if no registered factory was found.
 	*/
 	UFUNCTION(BlueprintCallable, Category = "Interchange | Import Manager")
-	INTERCHANGEENGINE_API const UClass* GetRegisteredFactoryClass(const UClass* ClassToMake) const;
+	UE_API const UClass* GetRegisteredFactoryClass(const UClass* ClassToMake) const;
 
 	/**
-	 * Return a pointer to an FImportAsyncHelper. The pointer is deleted when ReleaseAsyncHelper is called.
+	 * Initializes an FImportAsyncHelper with correct task data and import asset parameters. The Async Helper is deleted when ReleaseAsyncHelper is called.
 	 * @param Data - The data to pass to the different import tasks.
 	 */
-	INTERCHANGEENGINE_API TSharedRef<UE::Interchange::FImportAsyncHelper, ESPMode::ThreadSafe> CreateAsyncHelper(const UE::Interchange::FImportAsyncHelperData& Data, const FImportAssetParameters& ImportAssetParameters, const FGuid& InAnalyticsID);
+	UE_API void InitAsyncHelper(TSharedRef<UE::Interchange::FImportAsyncHelper, ESPMode::ThreadSafe>& AsyncHelper, const UE::Interchange::FImportAsyncHelperData& Data, const FImportAssetParameters& ImportAssetParameters);
 
 	/** Delete the specified AsyncHelper and remove it from the array that was holding it. */
-	INTERCHANGEENGINE_API void ReleaseAsyncHelper(TWeakPtr<UE::Interchange::FImportAsyncHelper, ESPMode::ThreadSafe> AsyncHelper);
+	UE_API void ReleaseAsyncHelper(TWeakPtr<UE::Interchange::FImportAsyncHelper, ESPMode::ThreadSafe> AsyncHelper);
 
 	/**
 	 * Return a pointer to an FImportAsyncHelper that is the specified UniqueId. Return an empty shared pointer if the unique id is not matching any async helper
 	 * The return asynchelper should be release before the end of the import process to make sure all resource are release when the import is completed.
 	 */
-	INTERCHANGEENGINE_API TSharedPtr<UE::Interchange::FImportAsyncHelper, ESPMode::ThreadSafe> GetAsyncHelper(int32 UniqueId);
+	UE_API TSharedPtr<UE::Interchange::FImportAsyncHelper, ESPMode::ThreadSafe> GetAsyncHelper(int32 UniqueId);
 
 	/*
 	 * Return the first translator that can translate the source data.
@@ -820,13 +822,13 @@ public:
 	 * @return return a matching translator, or nullptr if no translators exist for the source data.
 	 */
 	UFUNCTION(BlueprintCallable, Category = "Interchange | Import Manager")
-	INTERCHANGEENGINE_API UInterchangeTranslatorBase* GetTranslatorForSourceData(const UInterchangeSourceData* SourceData) const;
+	UE_API UInterchangeTranslatorBase* GetTranslatorForSourceData(const UInterchangeSourceData* SourceData) const;
 
 	/**
 	 * Return true if Interchange is actively importing or exporting, or false otherwise.
 	 */
 	UFUNCTION(BlueprintCallable, Category = "Interchange | Import Manager")
-	INTERCHANGEENGINE_API bool IsInterchangeActive();
+	UE_API bool IsInterchangeActive();
 
 	/**
 	 * Return false if Interchange is not actively importing or exporting.
@@ -834,20 +836,20 @@ public:
 	 * to be able to complete the operation they requested. (This is called by the exit editor operation.)
 	 */
 	UFUNCTION(BlueprintCallable, Category = "Interchange | Import Manager")
-	INTERCHANGEENGINE_API bool WarnIfInterchangeIsActive();
+	UE_API bool WarnIfInterchangeIsActive();
 
 	/**
 	 * Returns the list of supported formats for a given translator type.
 	 */
 	UFUNCTION(BlueprintCallable, Category = "Interchange | Import Manager")
-	INTERCHANGEENGINE_API UInterchangeAssetImportData* GetAssetImportData(UObject* Asset) const;
+	UE_API UInterchangeAssetImportData* GetAssetImportData(UObject* Asset) const;
 
 	/**
 	 * Check whether there is a translator registered that can translate the source data with the specified PayloadInterface.
 	 * @Param SourceData - The source data input you want to translate to Uod.
 	 * @return true if the source data can be translated using the specified PayloadInterface, or false otherwise.
 	 */
-	INTERCHANGEENGINE_API bool CanTranslateSourceDataWithPayloadInterface(const UInterchangeSourceData* SourceData, const UClass* PayloadInterfaceClass) const;
+	UE_API bool CanTranslateSourceDataWithPayloadInterface(const UInterchangeSourceData* SourceData, const UClass* PayloadInterfaceClass) const;
 
 	/*
 	 * Return the first translator that can translate the source data with the specified PayloadInterface.
@@ -855,65 +857,70 @@ public:
 	 * @Param PayloadInterfaceClass - The PayloadInterface that the translator must implement.
 	 * @return return a matching translator that implements the specified PayloadInterface, or nullptr if none exists.
 	 */
-	INTERCHANGEENGINE_API UInterchangeTranslatorBase* GetTranslatorSupportingPayloadInterfaceForSourceData(const UInterchangeSourceData* SourceData, const UClass* PayloadInterfaceClass) const;
+	UE_API UInterchangeTranslatorBase* GetTranslatorSupportingPayloadInterfaceForSourceData(const UInterchangeSourceData* SourceData, const UClass* PayloadInterfaceClass) const;
 
 	/**
 	 * Return true if the object is being imported, or false otherwise. If the user imports multiple file in the same folder, it's possible to
 	 * have the same asset name in two different files.
 	 */
 	UFUNCTION(BlueprintCallable, Category = "Interchange | Import Manager")
-	INTERCHANGEENGINE_API bool IsObjectBeingImported(UObject* Object) const;
+	UE_API bool IsObjectBeingImported(UObject* Object) const;
 
 	/**
 	 * Queue task that are not directly a import or re-import of assets.
 	 * The post import tasks are execute only when the QueuedTasks is empty (no more import task running)
 	 * Example: if a skeletal mesh re-import cannot apply the existing alternate skinning data, it will enqueue a post import task to re-import those alternate skinning files.
 	 */
-	INTERCHANGEENGINE_API bool EnqueuePostImportTask(TSharedPtr<FInterchangePostImportTask> PostImportTask);
+	UE_API bool EnqueuePostImportTask(TSharedPtr<FInterchangePostImportTask> PostImportTask);
 
 
 	/**
 	* Sets the bReplaceExistingAllDialogAnswer which is responsible for:
 	* if the Import process will override or not all existing assets this Import stack tries to override.
 	*/
-	INTERCHANGEENGINE_API static void SetReplaceExistingAlldialogAnswer(bool bReplaceExistingAllDialogAnswer);
+	UE_API static void SetReplaceExistingAlldialogAnswer(bool bReplaceExistingAllDialogAnswer);
 
 	/**
 	* Checks if the bReplaceExistingAllDialogAnswer.
 	*/
-	INTERCHANGEENGINE_API static void ResetReplaceExistingAlldialogAnswerSet();
+	UE_API static void ResetReplaceExistingAlldialogAnswerSet();
 
 	/**
 	* Gets the bReplaceExistingAllDialogAnswer.
 	*/
-	INTERCHANGEENGINE_API static TOptional<bool> GetReplaceExistingAlldialogAnswer();
+	UE_API static TOptional<bool> GetReplaceExistingAlldialogAnswer();
 
 	/**
 	* Set the editor utilities, those are use for editor operation like saving an asset.
 	*/
-	INTERCHANGEENGINE_API void SetEditorUtilities(UClass* EditorUtilitiesClass);
+	UE_API void SetEditorUtilities(UClass* EditorUtilitiesClass);
 	
 	/**
 	* Get the editor utilities, those are use for editor operation like saving an asset.
 	*/
-	INTERCHANGEENGINE_API UInterchangeEditorUtilitiesBase* GetEditorUtilities() const;
+	UE_API UInterchangeEditorUtilitiesBase* GetEditorUtilities() const;
 
 	/**
 	 * Check if a non-parallel translator is unlocked for using.
 	 */
-	INTERCHANGEENGINE_API bool CanUseTranslator(UInterchangeTranslatorBase* Translator) const;
+	UE_API bool CanUseTranslator(UInterchangeTranslatorBase* Translator) const;
+
+	/**
+	* Sets class to use for analytics. This will create a new object during import and use for the Analytics.
+	*/
+	UE_API void SetAnalyticsHandlerClass(const TSoftClassPtr<UInterchangeAnalyticsHandlerBase>& InAnalyticsHandlerClass);
 
 protected:
 
 	/** Return true if Interchange can show UI. */
-	static INTERCHANGEENGINE_API bool IsAttended();
+	static UE_API bool IsAttended();
 
 	/*
 	 * Find all pipeline candidates (C++, Blueprint, and Python).
 	 * @Param SourceData - The source data you need a translator for.
 	 * @return return a matching translator, or nullptr if there is no match.
 	 */
-	INTERCHANGEENGINE_API void FindPipelineCandidate(TArray<UClass*>& PipelineCandidates);
+	UE_API void FindPipelineCandidate(TArray<UClass*>& PipelineCandidates);
 
 	/**
 	 * This function cancels all tasks and finishes them as fast as possible.
@@ -921,31 +928,31 @@ protected:
 	 * @note - This is a asynchronous call. Tasks will be completed (canceled) soon.
 	 */
 	UFUNCTION(BlueprintCallable, Category = "Interchange | Import Manager")
-	INTERCHANGEENGINE_API void CancelAllTasks();
+	UE_API void CancelAllTasks();
 
 	/**
 	 * Wait synchronously until all tasks are done.
 	 */
 	UFUNCTION(BlueprintCallable, Category = "Interchange | Import Manager")
-	INTERCHANGEENGINE_API void WaitUntilAllTasksDone(bool bCancel);
+	UE_API void WaitUntilAllTasksDone(bool bCancel);
 
 
 	/**
 	 * If we set the mode to active, we will setup the timer and add the thread that will block the GC.
 	 * If the we set the mode to inactive, we will remove the timer and finish the thread that blocks the GC.
 	 */
-	INTERCHANGEENGINE_API void SetActiveMode(bool IsActive);
+	UE_API void SetActiveMode(bool IsActive);
 
 	/**
 	 * Start tasks until we reach the taskgraph worker number.
 	 * @param bCancelAllTasks - If true, we will start all tasks but with the cancel state set, so tasks will complete fast and call the completion task.
 	 */
-	INTERCHANGEENGINE_API void StartQueuedTasks(bool bCancelAllTasks = false);
+	UE_API void StartQueuedTasks(bool bCancelAllTasks = false);
 
 	/**
 	 * Called by the public Import functions.
 	 */
-	INTERCHANGEENGINE_API TTuple<UE::Interchange::FAssetImportResultRef, UE::Interchange::FSceneImportResultRef>
+	UE_API TTuple<UE::Interchange::FAssetImportResultRef, UE::Interchange::FSceneImportResultRef>
 	ImportInternal(const FString& ContentPath, const UInterchangeSourceData* SourceData, const FImportAssetParameters& ImportAssetParameters, const UE::Interchange::EImportType ImportType);
 
 	TOptional<FString> ValidateReimportParameter(UObject* ObjectToReimport, const FImportAssetParameters& ImportAssetParameters, bool bRunSynchronous);
@@ -958,6 +965,9 @@ protected:
 	void RemoveImportTask(TSharedPtr<UE::Interchange::FImportAsyncHelper, ESPMode::ThreadSafe> ImportTask);
 	void RemoveImportTask(int32 AsyncHelperUniqueId);
 
+private:
+	TStrongObjectPtr<UInterchangeAnalyticsHandlerBase> CreateAnalyticsHandler(const FString& ContentPath);
+
 private:
 
 	static bool bIsCreatingSingleton;
@@ -968,6 +978,8 @@ private:
 		UClass* TranslatorClass = nullptr;
 	};
 
+	TSoftClassPtr<UInterchangeAnalyticsHandlerBase> AnalyticsHandlerClass;
+
 	TMap<UClass*, bool> NonParallelTranslatorLocks;
 	TMap<UClass*, TArray<FQueuedTaskData>> NonParallelTranslatorQueueTasks;
 	
@@ -1040,3 +1052,5 @@ private:
 	// #todo_interchange: Remove when Interchange utility method for pipelines is done
 	friend class UInterchangeAssetImportData;
 };
+
+#undef UE_API
\ No newline at end of file
