[Interchange Analytics] [CRASH FIX]
- Changed the Interchange Source Node pointer from TNonNull to a const raw pointer with nullptr check.

#rnx

#Synced-CL 50825332

diff --git a/Engine/Plugins/Interchange/Runtime/Source/Analytics/Private/InterchangeAnalyticsHandlerDefault.cpp b/Engine/Plugins/Interchange/Runtime/Source/Analytics/Private/InterchangeAnalyticsHandlerDefault.cpp
--- a/Engine/Plugins/Interchange/Runtime/Source/Analytics/Private/InterchangeAnalyticsHandlerDefault.cpp
+++ b/Engine/Plugins/Interchange/Runtime/Source/Analytics/Private/InterchangeAnalyticsHandlerDefault.cpp
@@ -52,7 +52,10 @@ namespace UE::Interchange::AnalyticsHelpers
 		FAnalyticsJsonWriter JsonWriter(&ReturnValue);
 
 		TMap<FString, FString> ExtraInfoMap;
-		InSourceNode->GetExtraInformation(ExtraInfoMap);
+		if (InSourceNode)
+		{
+			InSourceNode->GetExtraInformation(ExtraInfoMap);
+		}
 
 		const FString DefaultValue("N/A");
 
diff --git a/Engine/Source/Runtime/Interchange/Engine/Public/InterchangeAnalyticsHandlerBase.h b/Engine/Source/Runtime/Interchange/Engine/Public/InterchangeAnalyticsHandlerBase.h
--- a/Engine/Source/Runtime/Interchange/Engine/Public/InterchangeAnalyticsHandlerBase.h
+++ b/Engine/Source/Runtime/Interchange/Engine/Public/InterchangeAnalyticsHandlerBase.h
@@ -54,7 +54,7 @@ struct FInterchangeImportResultAnalyticsInfo
 	int32 UniqueId;
 
 	TNotNull<const UInterchangeSourceData*> PrimarySourceData;
-	TNotNull<const UInterchangeSourceNode*> PrimarySourceNode;
+	const UInterchangeSourceNode* PrimarySourceNode;
 
 	TOptional<const UInterchangeResultsContainer*> AssetImportResultsContainer;
 	TOptional<const UInterchangeResultsContainer*> SceneImportResultsContainer;
