[Interchange] Added TSetAttributeHelper

The idea behind this helper implementation is to use same keys as a TArrayAttributeHelper making both interchangeable without breaking backward compatibility.

AddItem and HasItem are both O(1)
RemoveItem is O(N) because attributes key needs be re-aligned to avoid gaps inside the storage. RemoveItem is not called except during testing.

Some code should still be converted into TSetAttributeHelper instead of TArrayAttributeHelper but will be done slowly during future changes.

For the tested model, a sphere with 300 material slots - the reduction is around 92% for ApplySlotMaterialDependencies (7.8s to 550ms)

#tests Added new tests. Ran all Interchange tests locally
[RN] Added new TSetAttributeHelper

#Synced-CL 50859516

diff --git a/Engine/Plugins/Interchange/Runtime/Source/FactoryNodes/Private/InterchangeAnimSequenceFactoryNode.cpp b/Engine/Plugins/Interchange/Runtime/Source/FactoryNodes/Private/InterchangeAnimSequenceFactoryNode.cpp
--- a/Engine/Plugins/Interchange/Runtime/Source/FactoryNodes/Private/InterchangeAnimSequenceFactoryNode.cpp
+++ b/Engine/Plugins/Interchange/Runtime/Source/FactoryNodes/Private/InterchangeAnimSequenceFactoryNode.cpp
@@ -119,14 +119,7 @@ FString UInterchangeAnimSequenceFactoryNode::GetKeyDisplayName(const UE::Interch
 	}
 	else if (KeyDisplayName.StartsWith(UE::Interchange::FAnimSequenceNodeStaticData::GetAnimatedMorphTargetDependenciesKey().ToString()))
 	{
-		KeyDisplayName = TEXT("Animated Morph Targets Index ");
-		const FString IndexKey = UE::Interchange::TArrayAttributeHelper<FString>::IndexKey();
-		int32 IndexPosition = KeyDisplayName.Find(IndexKey) + IndexKey.Len();
-		if (IndexPosition < KeyDisplayName.Len())
-		{
-			KeyDisplayName += KeyDisplayName.RightChop(IndexPosition);
-		}
-		return KeyDisplayName;
+		return KeyDisplayName = UE::Interchange::BuildDisplayNameForArrayKey(TEXT("Animated Morph Targets Index "), KeyDisplayName);
 	}
 	else if (NodeAttributeKey == UE::Interchange::FAnimSequenceNodeStaticData::GetAnimatedAttributeCurveNamesKey())
 	{
@@ -135,14 +128,7 @@ FString UInterchangeAnimSequenceFactoryNode::GetKeyDisplayName(const UE::Interch
 	}
 	else if (KeyDisplayName.StartsWith(UE::Interchange::FAnimSequenceNodeStaticData::GetAnimatedAttributeCurveNamesKey().ToString()))
 	{
-		KeyDisplayName = TEXT("Animated Attribute Curves Index ");
-		const FString IndexKey = UE::Interchange::TArrayAttributeHelper<FString>::IndexKey();
-		int32 IndexPosition = KeyDisplayName.Find(IndexKey) + IndexKey.Len();
-		if (IndexPosition < KeyDisplayName.Len())
-		{
-			KeyDisplayName += KeyDisplayName.RightChop(IndexPosition);
-		}
-		return KeyDisplayName;
+		return KeyDisplayName = UE::Interchange::BuildDisplayNameForArrayKey(TEXT("Animated Attribute Curves Index "), KeyDisplayName);
 	}
 	else if (NodeAttributeKey == UE::Interchange::FAnimSequenceNodeStaticData::GetAnimatedAttributeStepCurveNamesKey())
 	{
@@ -151,14 +137,7 @@ FString UInterchangeAnimSequenceFactoryNode::GetKeyDisplayName(const UE::Interch
 	}
 	else if (KeyDisplayName.StartsWith(UE::Interchange::FAnimSequenceNodeStaticData::GetAnimatedAttributeStepCurveNamesKey().ToString()))
 	{
-		KeyDisplayName = TEXT("Animated Attribute Step Curves Index ");
-		const FString IndexKey = UE::Interchange::TArrayAttributeHelper<FString>::IndexKey();
-		int32 IndexPosition = KeyDisplayName.Find(IndexKey) + IndexKey.Len();
-		if (IndexPosition < KeyDisplayName.Len())
-		{
-			KeyDisplayName += KeyDisplayName.RightChop(IndexPosition);
-		}
-		return KeyDisplayName;
+		return KeyDisplayName = UE::Interchange::BuildDisplayNameForArrayKey(TEXT("Animated Attribute Step Curves Index "), KeyDisplayName);
 	}
 	else if (NodeAttributeKey == UE::Interchange::FAnimSequenceNodeStaticData::GetAnimatedMaterialCurveSuffixesKey())
 	{
@@ -167,14 +146,7 @@ FString UInterchangeAnimSequenceFactoryNode::GetKeyDisplayName(const UE::Interch
 	}
 	else if (KeyDisplayName.StartsWith(UE::Interchange::FAnimSequenceNodeStaticData::GetAnimatedMaterialCurveSuffixesKey().ToString()))
 	{
-		KeyDisplayName = TEXT("Animated Material Curve Suffixes Index ");
-		const FString IndexKey = UE::Interchange::TArrayAttributeHelper<FString>::IndexKey();
-		int32 IndexPosition = KeyDisplayName.Find(IndexKey) + IndexKey.Len();
-		if (IndexPosition < KeyDisplayName.Len())
-		{
-			KeyDisplayName += KeyDisplayName.RightChop(IndexPosition);
-		}
-		return KeyDisplayName;
+		return KeyDisplayName = UE::Interchange::BuildDisplayNameForArrayKey(TEXT("Animated Material Curve Suffixes Index "), KeyDisplayName);
 	}
 	else if (NodeAttributeKey == UE::Interchange::FAnimSequenceNodeStaticData::GetSceneNodeAnimationPayloadKeyUidMapKey())
 	{
diff --git a/Engine/Plugins/Interchange/Runtime/Source/FactoryNodes/Private/InterchangeMeshFactoryNode.cpp b/Engine/Plugins/Interchange/Runtime/Source/FactoryNodes/Private/InterchangeMeshFactoryNode.cpp
--- a/Engine/Plugins/Interchange/Runtime/Source/FactoryNodes/Private/InterchangeMeshFactoryNode.cpp
+++ b/Engine/Plugins/Interchange/Runtime/Source/FactoryNodes/Private/InterchangeMeshFactoryNode.cpp
@@ -109,24 +109,15 @@ FString UInterchangeMeshFactoryNode::GetKeyDisplayName(const UE::Interchange::FA
 	FString NodeAttributeKeyString = NodeAttributeKey.ToString();
 	if (NodeAttributeKey == UE::Interchange::FMeshFactoryNodeStaticData::GetLodDependenciesBaseKey())
 	{
-		KeyDisplayName = TEXT("LOD Dependencies Count");
-		return KeyDisplayName;
+		return KeyDisplayName = TEXT("LOD Dependencies Count");
 	}
 	else if (NodeAttributeKeyString.StartsWith(UE::Interchange::FMeshFactoryNodeStaticData::GetLodDependenciesBaseKey().ToString()))
 	{
-		KeyDisplayName = TEXT("LOD Dependencies Index ");
-		const FString IndexKey = UE::Interchange::TArrayAttributeHelper<FString>::IndexKey();
-		int32 IndexPosition = NodeAttributeKeyString.Find(IndexKey) + IndexKey.Len();
-		if (IndexPosition < NodeAttributeKeyString.Len())
-		{
-			KeyDisplayName += NodeAttributeKeyString.RightChop(IndexPosition);
-		}
-		return KeyDisplayName;
+		return KeyDisplayName = UE::Interchange::BuildDisplayNameForArrayKey(TEXT("LOD Dependencies Index "), NodeAttributeKeyString);
 	}
 	else if (NodeAttributeKey == UE::Interchange::FMeshFactoryNodeStaticData::GetSlotMaterialDependencyBaseKey())
 	{
-		KeyDisplayName = TEXT("Material Dependencies Count");
-		return KeyDisplayName;
+		return KeyDisplayName = TEXT("Material Dependencies Count");
 	}
 	else if (NodeAttributeKeyString.StartsWith(UE::Interchange::FMeshFactoryNodeStaticData::GetSlotMaterialDependencyBaseKey().ToString()))
 	{
@@ -153,18 +144,11 @@ FString UInterchangeMeshFactoryNode::GetKeyDisplayName(const UE::Interchange::FA
 	}
 	else if (NodeAttributeKey == UE::Interchange::FMeshFactoryNodeStaticData::GetAssemblyPartDependencyBaseKey())
 	{
-		KeyDisplayName = TEXT("Nanite Assembly Part Dependencies Count");
+		return KeyDisplayName = TEXT("Nanite Assembly Part Dependencies Count");
 	}
 	else if (NodeAttributeKeyString.StartsWith(UE::Interchange::FMeshFactoryNodeStaticData::GetAssemblyPartDependencyBaseKey().ToString()))
 	{
-		KeyDisplayName = TEXT("Nanite Assembly Part Dependencies Index ");
-		const FString IndexKey = UE::Interchange::TArrayAttributeHelper<FString>::IndexKey();
-		int32 IndexPosition = NodeAttributeKeyString.Find(IndexKey) + IndexKey.Len();
-		if (IndexPosition < NodeAttributeKeyString.Len())
-		{
-			KeyDisplayName += NodeAttributeKeyString.RightChop(IndexPosition);
-		}
-		return KeyDisplayName;
+		return KeyDisplayName = UE::Interchange::BuildDisplayNameForArrayKey(TEXT("Nanite Assembly Part Dependencies Index "), NodeAttributeKeyString);
 	}
 	return Super::GetKeyDisplayName(NodeAttributeKey);
 }
@@ -324,7 +308,6 @@ bool UInterchangeMeshFactoryNode::RemoveSlotMaterialDependencyUid(const FString&
 bool UInterchangeMeshFactoryNode::ResetSlotMaterialDependencies()
 {
 	SlotMaterialDependencies.Empty();
-	SlotMaterialDependencies.RebuildCache();
 	return true;
 }
 
@@ -356,7 +339,6 @@ bool UInterchangeMeshFactoryNode::RemoveAssemblyPartDependencyUid(const FString&
 bool UInterchangeMeshFactoryNode::ResetAssemblyDependencies()
 {
 	AssemblyPartDependencies.Empty();
-	AssemblyPartDependencies.RebuildCache();
 	return true;
 }
 
diff --git a/Engine/Plugins/Interchange/Runtime/Source/FactoryNodes/Private/InterchangeSkeletalMeshLodDataNode.cpp b/Engine/Plugins/Interchange/Runtime/Source/FactoryNodes/Private/InterchangeSkeletalMeshLodDataNode.cpp
--- a/Engine/Plugins/Interchange/Runtime/Source/FactoryNodes/Private/InterchangeSkeletalMeshLodDataNode.cpp
+++ b/Engine/Plugins/Interchange/Runtime/Source/FactoryNodes/Private/InterchangeSkeletalMeshLodDataNode.cpp
@@ -36,24 +36,15 @@ FString UInterchangeSkeletalMeshLodDataNode::GetKeyDisplayName(const UE::Interch
 	const FString NodeAttributeKeyString = KeyDisplayName;
 	if (NodeAttributeKey == UE::Interchange::FSkeletalMeshNodeLodDataStaticData::GetLODMeshDataArrayBaseKey())
 	{
-		KeyDisplayName = TEXT("Mesh count");
-		return KeyDisplayName;
+		return KeyDisplayName = TEXT("Mesh count");
 	}
 	else if (NodeAttributeKeyString.StartsWith(UE::Interchange::FSkeletalMeshNodeLodDataStaticData::GetLODMeshDataArrayBaseKey().ToString()))
 	{
-		KeyDisplayName = TEXT("Mesh index ");
-		const FString IndexKey = UE::Interchange::TArrayAttributeHelper<FString>::IndexKey();
-		int32 IndexPosition = NodeAttributeKeyString.Find(IndexKey) + IndexKey.Len();
-		if (IndexPosition < NodeAttributeKeyString.Len())
-		{
-			KeyDisplayName += NodeAttributeKeyString.RightChop(IndexPosition);
-		}
-		return KeyDisplayName;
+		return KeyDisplayName = UE::Interchange::BuildDisplayNameForArrayKey(TEXT("Mesh index "), NodeAttributeKeyString);
 	}
 	else if (NodeAttributeKey == Macro_CustomSkeletonUidKey)
 	{
-		KeyDisplayName = TEXT("Skeleton factory node");
-		return KeyDisplayName;
+		return KeyDisplayName = TEXT("Skeleton factory node");
 	}
 	return Super::GetKeyDisplayName(NodeAttributeKey);
 }
diff --git a/Engine/Plugins/Interchange/Runtime/Source/FactoryNodes/Private/InterchangeStaticMeshFactoryNode.cpp b/Engine/Plugins/Interchange/Runtime/Source/FactoryNodes/Private/InterchangeStaticMeshFactoryNode.cpp
--- a/Engine/Plugins/Interchange/Runtime/Source/FactoryNodes/Private/InterchangeStaticMeshFactoryNode.cpp
+++ b/Engine/Plugins/Interchange/Runtime/Source/FactoryNodes/Private/InterchangeStaticMeshFactoryNode.cpp
@@ -140,13 +140,7 @@ FString UInterchangeStaticMeshFactoryNode::GetKeyDisplayName(const UE::Interchan
 	}
 	else if (NodeAttributeKey == UE::Interchange::FStaticMeshNodeStaticData::GetSocketUidsBaseKey())
 	{
-		KeyDisplayName = TEXT("Socket Index ");
-		const FString IndexKey = UE::Interchange::TArrayAttributeHelper<FString>::IndexKey();
-		int32 IndexPosition = OriginalKeyName.Find(IndexKey) + IndexKey.Len();
-		if (IndexPosition < OriginalKeyName.Len())
-		{
-			KeyDisplayName += OriginalKeyName.RightChop(IndexPosition);
-		}
+		KeyDisplayName = UE::Interchange::BuildDisplayNameForArrayKey(TEXT("Socket Index "), OriginalKeyName);
 	}
 	else
 	{
diff --git a/Engine/Plugins/Interchange/Runtime/Source/FactoryNodes/Private/InterchangeStaticMeshLodDataNode.cpp b/Engine/Plugins/Interchange/Runtime/Source/FactoryNodes/Private/InterchangeStaticMeshLodDataNode.cpp
--- a/Engine/Plugins/Interchange/Runtime/Source/FactoryNodes/Private/InterchangeStaticMeshLodDataNode.cpp
+++ b/Engine/Plugins/Interchange/Runtime/Source/FactoryNodes/Private/InterchangeStaticMeshLodDataNode.cpp
@@ -112,93 +112,51 @@ FString UInterchangeStaticMeshLodDataNode::GetKeyDisplayName(const UE::Interchan
 	const FString NodeAttributeKeyString = KeyDisplayName;
 	if (NodeAttributeKey == UE::Interchange::FStaticMeshNodeLodDataStaticData::GetLODMeshDataArrayBaseKey())
 	{
-		KeyDisplayName = TEXT("Mesh count");
-		return KeyDisplayName;
+		return KeyDisplayName = TEXT("Mesh count");
 	}
 	else if (NodeAttributeKeyString.StartsWith(UE::Interchange::FStaticMeshNodeLodDataStaticData::GetLODMeshDataArrayBaseKey().ToString()))
 	{
-		KeyDisplayName = TEXT("Mesh index ");
-		const FString IndexKey = UE::Interchange::TArrayAttributeHelper<FString>::IndexKey();
-		int32 IndexPosition = NodeAttributeKeyString.Find(IndexKey) + IndexKey.Len();
-		if (IndexPosition < NodeAttributeKeyString.Len())
-		{
-			KeyDisplayName += NodeAttributeKeyString.RightChop(IndexPosition);
-		}
-		return KeyDisplayName;
+		return KeyDisplayName = UE::Interchange::BuildDisplayNameForArrayKey(TEXT("Mesh index "), NodeAttributeKeyString);
 	}
 	else if (NodeAttributeKey == UE::Interchange::FStaticMeshNodeLodDataStaticData::GetBoxCollisionMeshDataArrayBaseKey())
 	{
-		KeyDisplayName = TEXT("Box collision count");
-		return KeyDisplayName;
+		return KeyDisplayName = TEXT("Box collision count");
 	}
 	else if (NodeAttributeKeyString.StartsWith(UE::Interchange::FStaticMeshNodeLodDataStaticData::GetBoxCollisionMeshDataArrayBaseKey().ToString()))
 	{
-		KeyDisplayName = TEXT("Box collision index ");
-		const FString IndexKey = UE::Interchange::TArrayAttributeHelper<FString>::IndexKey();
-		int32 IndexPosition = NodeAttributeKeyString.Find(IndexKey) + IndexKey.Len();
-		if (IndexPosition < NodeAttributeKeyString.Len())
-		{
-			KeyDisplayName += NodeAttributeKeyString.RightChop(IndexPosition);
-		}
-		return KeyDisplayName;
+		return KeyDisplayName = UE::Interchange::BuildDisplayNameForArrayKey(TEXT("Box collision index "), NodeAttributeKeyString);
 	}
 	else if (NodeAttributeKey == UE::Interchange::FStaticMeshNodeLodDataStaticData::GetCapsuleCollisionMeshDataArrayBaseKey())
 	{
-		KeyDisplayName = TEXT("Capsule collision count");
-		return KeyDisplayName;
+		return KeyDisplayName = TEXT("Capsule collision count");
 	}
 	else if (NodeAttributeKeyString.StartsWith(UE::Interchange::FStaticMeshNodeLodDataStaticData::GetCapsuleCollisionMeshDataArrayBaseKey().ToString()))
 	{
-		KeyDisplayName = TEXT("Capsule collision index ");
-		const FString IndexKey = UE::Interchange::TArrayAttributeHelper<FString>::IndexKey();
-		int32 IndexPosition = NodeAttributeKeyString.Find(IndexKey) + IndexKey.Len();
-		if (IndexPosition < NodeAttributeKeyString.Len())
-		{
-			KeyDisplayName += NodeAttributeKeyString.RightChop(IndexPosition);
-		}
-		return KeyDisplayName;
+		return KeyDisplayName = UE::Interchange::BuildDisplayNameForArrayKey(TEXT("Capsule collision index "), NodeAttributeKeyString);
 	}
 	else if (NodeAttributeKey == UE::Interchange::FStaticMeshNodeLodDataStaticData::GetSphereCollisionMeshDataArrayBaseKey())
 	{
-		KeyDisplayName = TEXT("Sphere collision count");
-		return KeyDisplayName;
+		return KeyDisplayName = TEXT("Sphere collision count");
 	}
 	else if (NodeAttributeKeyString.StartsWith(UE::Interchange::FStaticMeshNodeLodDataStaticData::GetSphereCollisionMeshDataArrayBaseKey().ToString()))
 	{
-		KeyDisplayName = TEXT("Sphere collision index ");
-		const FString IndexKey = UE::Interchange::TArrayAttributeHelper<FString>::IndexKey();
-		int32 IndexPosition = NodeAttributeKeyString.Find(IndexKey) + IndexKey.Len();
-		if (IndexPosition < NodeAttributeKeyString.Len())
-		{
-			KeyDisplayName += NodeAttributeKeyString.RightChop(IndexPosition);
-		}
-		return KeyDisplayName;
+		return KeyDisplayName = UE::Interchange::BuildDisplayNameForArrayKey(TEXT("Sphere collision index "), NodeAttributeKeyString);
 	}
 	else if (NodeAttributeKey == UE::Interchange::FStaticMeshNodeLodDataStaticData::GetConvexCollisionMeshDataArrayBaseKey())
 	{
-		KeyDisplayName = TEXT("Convex collision count");
-		return KeyDisplayName;
+		return KeyDisplayName = TEXT("Convex collision count");
 	}
 	else if (NodeAttributeKeyString.StartsWith(UE::Interchange::FStaticMeshNodeLodDataStaticData::GetConvexCollisionMeshDataArrayBaseKey().ToString()))
 	{
-		KeyDisplayName = TEXT("Convex collision index ");
-		const FString IndexKey = UE::Interchange::TArrayAttributeHelper<FString>::IndexKey();
-		int32 IndexPosition = NodeAttributeKeyString.Find(IndexKey) + IndexKey.Len();
-		if (IndexPosition < NodeAttributeKeyString.Len())
-		{
-			KeyDisplayName += NodeAttributeKeyString.RightChop(IndexPosition);
-		}
-		return KeyDisplayName;
+		return KeyDisplayName = UE::Interchange::BuildDisplayNameForArrayKey(TEXT("Convex collision index "), NodeAttributeKeyString);
 	}
 	else if (NodeAttributeKey == Macro_CustomOneConvexHullPerUCXKey)
 	{
-		KeyDisplayName = TEXT("One Convex Hull Per UCX");
-		return KeyDisplayName;
+		return KeyDisplayName = TEXT("One Convex Hull Per UCX");
 	}
 	else if (NodeAttributeKey == Macro_CustomImportCollisionKey)
 	{
-		KeyDisplayName = TEXT("Import Collision");
-		return KeyDisplayName;
+		return KeyDisplayName = TEXT("Import Collision");
 	}
 
 	return Super::GetKeyDisplayName(NodeAttributeKey);
diff --git a/Engine/Plugins/Interchange/Runtime/Source/Import/Private/Groom/InterchangeGroomCacheFactory.cpp b/Engine/Plugins/Interchange/Runtime/Source/Import/Private/Groom/InterchangeGroomCacheFactory.cpp
--- a/Engine/Plugins/Interchange/Runtime/Source/Import/Private/Groom/InterchangeGroomCacheFactory.cpp
+++ b/Engine/Plugins/Interchange/Runtime/Source/Import/Private/Groom/InterchangeGroomCacheFactory.cpp
@@ -98,12 +98,13 @@ UInterchangeFactoryBase::FImportAssetResult UInterchangeGroomCacheFactory::Begin
 	{
 		// Otherwise, the user wanted to import the groom asset
 		bool bSuccess = false;
-		TArray<FString> Dependencies;
-		GroomCacheFactoryNode->GetFactoryDependencies(Dependencies);
-		if (Dependencies.Num() > 0)
+		TSet<FString> Dependencies = GroomCacheFactoryNode->GetFactoryDependencies();
+		for (const FString& Dependency : Dependencies)
 		{
-			const UInterchangeGroomFactoryNode* GroomFactoryNode = Cast<UInterchangeGroomFactoryNode>(Arguments.NodeContainer->GetNode(Dependencies[0]));
+			const UInterchangeGroomFactoryNode* GroomFactoryNode = Cast<UInterchangeGroomFactoryNode>(Arguments.NodeContainer->GetNode(Dependency));
 			bSuccess = GroomFactoryNode && GroomFactoryNode->GetCustomReferenceObject(GroomAssetPath);
+
+			break; // GroomCacheFactoryNode is supposed to have one UInterchangeGroomFactoryNode dependency
 		}
 
 		if (!bSuccess)
diff --git a/Engine/Plugins/Interchange/Runtime/Source/Nodes/Private/InterchangeMeshBundleNode.cpp b/Engine/Plugins/Interchange/Runtime/Source/Nodes/Private/InterchangeMeshBundleNode.cpp
--- a/Engine/Plugins/Interchange/Runtime/Source/Nodes/Private/InterchangeMeshBundleNode.cpp
+++ b/Engine/Plugins/Interchange/Runtime/Source/Nodes/Private/InterchangeMeshBundleNode.cpp
@@ -72,7 +72,6 @@ bool UInterchangeMeshBundleNode::RemoveMeshNodeUid(const FString& MeshNodeUid)
 void UInterchangeMeshBundleNode::ResetMeshNodeUids()
 {
 	MeshNodeToTransforms.Empty();
-	MeshNodeToTransforms.RebuildCache();
 }
 
 void UInterchangeMeshBundleNode::GetAllMeshNodesAndTransforms(TMap<FString, TArray<FTransform>>& OutMeshNodeToTransforms) const
diff --git a/Engine/Plugins/Interchange/Runtime/Source/Nodes/Private/InterchangeMeshNode.cpp b/Engine/Plugins/Interchange/Runtime/Source/Nodes/Private/InterchangeMeshNode.cpp
--- a/Engine/Plugins/Interchange/Runtime/Source/Nodes/Private/InterchangeMeshNode.cpp
+++ b/Engine/Plugins/Interchange/Runtime/Source/Nodes/Private/InterchangeMeshNode.cpp
@@ -109,14 +109,7 @@ FString UInterchangeMeshNode::GetKeyDisplayName(const UE::Interchange::FAttribut
 	}
 	else if (NodeAttributeKeyString.StartsWith(UE::Interchange::FMeshNodeStaticData::GetSkeletonDependenciesKey().ToString()))
 	{
-		KeyDisplayName = TEXT("Skeleton Dependencies Index ");
-		const FString IndexKey = UE::Interchange::TArrayAttributeHelper<FString>::IndexKey();
-		int32 IndexPosition = NodeAttributeKeyString.Find(IndexKey) + IndexKey.Len();
-		if (IndexPosition < NodeAttributeKeyString.Len())
-		{
-			KeyDisplayName += NodeAttributeKeyString.RightChop(IndexPosition);
-		}
-		return KeyDisplayName;
+		return KeyDisplayName = UE::Interchange::BuildDisplayNameForArrayKey(TEXT("Skeleton Dependencies Index "), NodeAttributeKeyString);
 	}
 	else if (NodeAttributeKey == UE::Interchange::FMeshNodeStaticData::GetMorphTargetDependenciesKey())
 	{
@@ -124,14 +117,7 @@ FString UInterchangeMeshNode::GetKeyDisplayName(const UE::Interchange::FAttribut
 	}
 	else if (NodeAttributeKeyString.StartsWith(UE::Interchange::FMeshNodeStaticData::GetMorphTargetDependenciesKey().ToString()))
 	{
-		KeyDisplayName = TEXT("Morph Target Dependencies Index ");
-		const FString IndexKey = UE::Interchange::TArrayAttributeHelper<FString>::IndexKey();
-		int32 IndexPosition = NodeAttributeKeyString.Find(IndexKey) + IndexKey.Len();
-		if (IndexPosition < NodeAttributeKeyString.Len())
-		{
-			KeyDisplayName += NodeAttributeKeyString.RightChop(IndexPosition);
-		}
-		return KeyDisplayName;
+		return KeyDisplayName = UE::Interchange::BuildDisplayNameForArrayKey(TEXT("Morph Target Dependencies Index "), NodeAttributeKeyString);
 	}
 	else if (NodeAttributeKey == UE::Interchange::FMeshNodeStaticData::GetSceneInstancesUidsKey())
 	{
@@ -139,14 +125,7 @@ FString UInterchangeMeshNode::GetKeyDisplayName(const UE::Interchange::FAttribut
 	}
 	else if (NodeAttributeKeyString.StartsWith(UE::Interchange::FMeshNodeStaticData::GetSceneInstancesUidsKey().ToString()))
 	{
-		KeyDisplayName = TEXT("Scene mesh instances Index ");
-		const FString IndexKey = UE::Interchange::TArrayAttributeHelper<FString>::IndexKey();
-		int32 IndexPosition = NodeAttributeKeyString.Find(IndexKey) + IndexKey.Len();
-		if (IndexPosition < NodeAttributeKeyString.Len())
-		{
-			KeyDisplayName += NodeAttributeKeyString.RightChop(IndexPosition);
-		}
-		return KeyDisplayName;
+		return KeyDisplayName = UE::Interchange::BuildDisplayNameForArrayKey(TEXT("Scene mesh instances Index "), NodeAttributeKeyString);
 	}
 	else if (NodeAttributeKeyString.StartsWith(UE::Interchange::FMeshNodeStaticData::GetSlotMaterialDependenciesKey().ToString()))
 	{
diff --git a/Engine/Plugins/Interchange/Runtime/Source/Nodes/Private/InterchangeSceneNode.cpp b/Engine/Plugins/Interchange/Runtime/Source/Nodes/Private/InterchangeSceneNode.cpp
--- a/Engine/Plugins/Interchange/Runtime/Source/Nodes/Private/InterchangeSceneNode.cpp
+++ b/Engine/Plugins/Interchange/Runtime/Source/Nodes/Private/InterchangeSceneNode.cpp
@@ -115,14 +115,7 @@ FString UInterchangeSceneNode::GetKeyDisplayName(const UE::Interchange::FAttribu
 	}
 	else if (NodeAttributeKeyString.StartsWith(UE::Interchange::FSceneNodeStaticData::GetNodeSpecializeTypeBaseKey().ToString()))
 	{
-		KeyDisplayName = TEXT("Specialized type index ");
-		const FString IndexKey = UE::Interchange::TArrayAttributeHelper<FString>::IndexKey();
-		int32 IndexPosition = NodeAttributeKeyString.Find(IndexKey) + IndexKey.Len();
-		if (IndexPosition < NodeAttributeKeyString.Len())
-		{
-			KeyDisplayName += NodeAttributeKeyString.RightChop(IndexPosition);
-		}
-		return KeyDisplayName;
+		return KeyDisplayName = UE::Interchange::BuildDisplayNameForArrayKey(TEXT("Specialized type index "), NodeAttributeKeyString);
 	}
 	else if (NodeAttributeKey == UE::Interchange::FSceneNodeStaticData::GetMaterialDependencyUidsBaseKey())
 	{
@@ -131,14 +124,7 @@ FString UInterchangeSceneNode::GetKeyDisplayName(const UE::Interchange::FAttribu
 	}
 	else if (NodeAttributeKeyString.StartsWith(UE::Interchange::FSceneNodeStaticData::GetMaterialDependencyUidsBaseKey().ToString()))
 	{
-		KeyDisplayName = TEXT("Material dependency index ");
-		const FString IndexKey = UE::Interchange::TArrayAttributeHelper<FString>::IndexKey();
-		int32 IndexPosition = NodeAttributeKeyString.Find(IndexKey) + IndexKey.Len();
-		if (IndexPosition < NodeAttributeKeyString.Len())
-		{
-			KeyDisplayName += NodeAttributeKeyString.RightChop(IndexPosition);
-		}
-		return KeyDisplayName;
+		return KeyDisplayName = UE::Interchange::BuildDisplayNameForArrayKey(TEXT("Material dependency index "), NodeAttributeKeyString);
 	}
 	return Super::GetKeyDisplayName(NodeAttributeKey);
 }
diff --git a/Engine/Plugins/Interchange/Runtime/Source/Nodes/Public/InterchangeSceneNode.h b/Engine/Plugins/Interchange/Runtime/Source/Nodes/Public/InterchangeSceneNode.h
--- a/Engine/Plugins/Interchange/Runtime/Source/Nodes/Public/InterchangeSceneNode.h
+++ b/Engine/Plugins/Interchange/Runtime/Source/Nodes/Public/InterchangeSceneNode.h
@@ -121,14 +121,14 @@ public:
 	UFUNCTION(BlueprintCallable, Category = "Interchange | Node | Scene")
 	UE_API bool SetCustomLocalTransform(const UInterchangeBaseNodeContainer* BaseNodeContainer, const FTransform& AttributeValue);
 
-	UE_DEPRECATED(5.7, "Use the overload without the bResetCache parameter; it is no longer needed.")
+	UE_DEPRECATED(5.8, "Use the overload without the bResetCache parameter; it is no longer needed.")
 	UE_API bool SetCustomLocalTransform(const UInterchangeBaseNodeContainer* BaseNodeContainer, const FTransform& AttributeValue, bool bResetCache);
 
 	/** Get the default scene node global transform. This value is computed from the local transforms of all parent scene nodes. */
 	UFUNCTION(BlueprintCallable, Category = "Interchange | Node | Scene")
 	UE_API bool GetCustomGlobalTransform(const UInterchangeBaseNodeContainer* BaseNodeContainer, const FTransform& GlobalOffsetTransform, FTransform& AttributeValue) const;
 
-	UE_DEPRECATED(5.7, "Use the overload without the bForceRecache parameter; it is no longer needed.")
+	UE_DEPRECATED(5.8, "Use the overload without the bForceRecache parameter; it is no longer needed.")
 	UE_API bool GetCustomGlobalTransform(const UInterchangeBaseNodeContainer* BaseNodeContainer, const FTransform& GlobalOffsetTransform, FTransform& AttributeValue, bool bForceRecache) const;
 
 
@@ -263,51 +263,51 @@ public:
 	*/
 
 	UFUNCTION(BlueprintCallable, Category = "Interchange | Node | Joint", meta = (DeprecatedFunction, DeprecationMessage = "Moved to UInterchangeJointNode."))
-	UE_DEPRECATED(5.7, "This function moved to UInterchangeJointNode::GetBindPoseLocalTransform")
+	UE_DEPRECATED(5.8, "This function moved to UInterchangeJointNode::GetBindPoseLocalTransform")
 	UE_API bool GetCustomBindPoseLocalTransform(FTransform& AttributeValue) const;
 
 	UFUNCTION(BlueprintCallable, Category = "Interchange | Node | Joint", meta = (DeprecatedFunction, DeprecationMessage = "Moved to UInterchangeJointNode."))
-	UE_DEPRECATED(5.7, "This function moved to UInterchangeJointNode::SetBindPoseLocalTransform")
+	UE_DEPRECATED(5.8, "This function moved to UInterchangeJointNode::SetBindPoseLocalTransform")
 	UE_API bool SetCustomBindPoseLocalTransform(const UInterchangeBaseNodeContainer* BaseNodeContainer, const FTransform& AttributeValue, bool bResetCache = true);
 
 	UFUNCTION(BlueprintCallable, Category = "Interchange | Node | Joint", meta = (DeprecatedFunction, DeprecationMessage = "Moved to UInterchangeJointNode."))
-	UE_DEPRECATED(5.7, "This function moved to UInterchangeJointNode::GetBindPoseGlobalTransform")
+	UE_DEPRECATED(5.8, "This function moved to UInterchangeJointNode::GetBindPoseGlobalTransform")
 	UE_API bool GetCustomBindPoseGlobalTransform(const UInterchangeBaseNodeContainer* BaseNodeContainer, const FTransform& GlobalOffsetTransform, FTransform& AttributeValue, bool bForceRecache = false) const;
 
 	UFUNCTION(BlueprintCallable, Category = "Interchange | Node | Joint", meta = (DeprecatedFunction, DeprecationMessage = "Moved to UInterchangeJointNode."))
-	UE_DEPRECATED(5.7, "This function moved to UInterchangeJointNode::GetTimeZeroLocalTransform")
+	UE_DEPRECATED(5.8, "This function moved to UInterchangeJointNode::GetTimeZeroLocalTransform")
 	UE_API bool GetCustomTimeZeroLocalTransform(FTransform& AttributeValue) const;
 
 	UFUNCTION(BlueprintCallable, Category = "Interchange | Node | Joint", meta = (DeprecatedFunction, DeprecationMessage = "Moved to UInterchangeJointNode."))
-	UE_DEPRECATED(5.7, "This function moved to UInterchangeJointNode::SetTimeZeroLocalTransform")
+	UE_DEPRECATED(5.8, "This function moved to UInterchangeJointNode::SetTimeZeroLocalTransform")
 	UE_API bool SetCustomTimeZeroLocalTransform(const UInterchangeBaseNodeContainer* BaseNodeContainer, const FTransform& AttributeValue, bool bResetCache = true);
 
 	UFUNCTION(BlueprintCallable, Category = "Interchange | Node | Joint", meta = (DeprecatedFunction, DeprecationMessage = "Moved to UInterchangeJointNode."))
-	UE_DEPRECATED(5.7, "This function moved to UInterchangeJointNode::GetTimeZeroGlobalTransform")
+	UE_DEPRECATED(5.8, "This function moved to UInterchangeJointNode::GetTimeZeroGlobalTransform")
 	UE_API bool GetCustomTimeZeroGlobalTransform(const UInterchangeBaseNodeContainer* BaseNodeContainer, const FTransform& GlobalOffsetTransform, FTransform& AttributeValue, bool bForceRecache = false) const;
 
 	UFUNCTION(BlueprintCallable, Category = "Interchange | Node | Joint", meta = (DeprecatedFunction, DeprecationMessage = "Moved to UInterchangeJointNode."))
-	UE_DEPRECATED(5.7, "This function moved to UInterchangeJointNode::SetMeshUIDToGlobalBindPoseReferenceMap")
+	UE_DEPRECATED(5.8, "This function moved to UInterchangeJointNode::SetMeshUIDToGlobalBindPoseReferenceMap")
 	UE_API void SetGlobalBindPoseReferenceForMeshUIDs(const TMap<FString, FMatrix>& GlobalBindPoseReferenceForMeshUIDs);
 
 	UFUNCTION(BlueprintCallable, Category = "Interchange | Node | Joint", meta = (DeprecatedFunction, DeprecationMessage = "Moved to UInterchangeJointNode."))
-	UE_DEPRECATED(5.7, "This function moved to UInterchangeJointNode::GetGlobalBindPoseReferenceFromMeshUID")
+	UE_DEPRECATED(5.8, "This function moved to UInterchangeJointNode::GetGlobalBindPoseReferenceFromMeshUID")
 	UE_API bool GetGlobalBindPoseReferenceForMeshUID(const FString& MeshUID, FMatrix& GlobalBindPoseReference) const;
 
 	UFUNCTION(BlueprintCallable, Category = "Interchange | Node | Joint", meta = (DeprecatedFunction, DeprecationMessage = "Moved to UInterchangeJointNode."))
-	UE_DEPRECATED(5.7, "This function moved to UInterchangeJointNode::SetHasInvalidBindPose")
+	UE_DEPRECATED(5.8, "This function moved to UInterchangeJointNode::SetHasInvalidBindPose")
 	UE_API bool SetCustomHasBindPose(const bool& bHasBindPose);
 
 	UFUNCTION(BlueprintCallable, Category = "Interchange | Node | Joint", meta = (DeprecatedFunction, DeprecationMessage = "Moved to UInterchangeJointNode."))
-	UE_DEPRECATED(5.7, "This function moved to UInterchangeJointNode::SetHasInvalidBindPose")
+	UE_DEPRECATED(5.8, "This function moved to UInterchangeJointNode::SetHasInvalidBindPose")
 	UE_API bool GetCustomHasBindPose(bool& bHasBindPose) const;
 
 	UFUNCTION(BlueprintCallable, Category = "Interchange | Node | Joint", meta = (DeprecatedFunction, DeprecationMessage = "Moved to UInterchangeJointNode."))
-	UE_DEPRECATED(5.7, "This function moved to UInterchangeJointNode::SetGlobalMatrixForT0Rebinding")
+	UE_DEPRECATED(5.8, "This function moved to UInterchangeJointNode::SetGlobalMatrixForT0Rebinding")
 	UE_API bool SetCustomGlobalMatrixForT0Rebinding(const FMatrix& AttributeValue);
 
 	UFUNCTION(BlueprintCallable, Category = "Interchange | Node | Joint", meta = (DeprecatedFunction, DeprecationMessage = "Moved to UInterchangeJointNode."))
-	UE_DEPRECATED(5.7, "This function moved to UInterchangeJointNode::GetGlobalMatrixForT0Rebinding")
+	UE_DEPRECATED(5.8, "This function moved to UInterchangeJointNode::GetGlobalMatrixForT0Rebinding")
 	UE_API bool GetCustomGlobalMatrixForT0Rebinding(FMatrix& AttributeValue) const;
 
 	/*
diff --git a/Engine/Plugins/Interchange/Runtime/Source/Pipelines/Private/InterchangeGenericAnimationPipeline.cpp b/Engine/Plugins/Interchange/Runtime/Source/Pipelines/Private/InterchangeGenericAnimationPipeline.cpp
--- a/Engine/Plugins/Interchange/Runtime/Source/Pipelines/Private/InterchangeGenericAnimationPipeline.cpp
+++ b/Engine/Plugins/Interchange/Runtime/Source/Pipelines/Private/InterchangeGenericAnimationPipeline.cpp
@@ -938,15 +938,9 @@ void UInterchangeGenericAnimationPipeline::CreateAnimSequenceFactoryNode(UInterc
 		{
 			if (const UInterchangeSkeletalMeshFactoryNode* ConstSkeletalMeshFactoryNode = Cast<const UInterchangeSkeletalMeshFactoryNode>(BaseNodeContainer->GetFactoryNode(SkelMeshFactoryNodeUid)))
 			{
-				TArray<FString> SkeletalMeshDependencies;
-				ConstSkeletalMeshFactoryNode->GetFactoryDependencies(SkeletalMeshDependencies);
-				for (const FString& SkeletalMeshDependencyUid : SkeletalMeshDependencies)
+				if (ConstSkeletalMeshFactoryNode->HasFactoryDependency(SkeletonUid))
 				{
-					if (SkeletonUid.Equals(SkeletalMeshDependencyUid))
-					{
-						AnimSequenceFactoryNode->AddFactoryDependencyUid(SkelMeshFactoryNodeUid);
-						break;
-					}
+					AnimSequenceFactoryNode->AddFactoryDependencyUid(SkelMeshFactoryNodeUid);
 				}
 			}
 		}
diff --git a/Engine/Plugins/Interchange/Runtime/Source/Pipelines/Private/InterchangeGenericMaterialPipeline.cpp b/Engine/Plugins/Interchange/Runtime/Source/Pipelines/Private/InterchangeGenericMaterialPipeline.cpp
--- a/Engine/Plugins/Interchange/Runtime/Source/Pipelines/Private/InterchangeGenericMaterialPipeline.cpp
+++ b/Engine/Plugins/Interchange/Runtime/Source/Pipelines/Private/InterchangeGenericMaterialPipeline.cpp
@@ -3386,12 +3386,7 @@ UInterchangeMaterialExpressionFactoryNode* UInterchangeGenericMaterialPipeline::
 
 				if (BaseNodeContainer->IsNodeUidValid(TextureFactoryNodeUid))
 				{
-					TArray<FString> FactoryDependencies;
-					MaterialFactoryNode->GetFactoryDependencies(FactoryDependencies);
-					if (!FactoryDependencies.Contains(TextureFactoryNodeUid))
-					{
-						MaterialFactoryNode->AddFactoryDependencyUid(TextureFactoryNodeUid);
-					}
+					MaterialFactoryNode->AddFactoryDependencyUid(TextureFactoryNodeUid);
 				}
 			}
 		}
diff --git a/Engine/Plugins/Interchange/Runtime/Source/Pipelines/Private/InterchangeGenericMeshPipeline.cpp b/Engine/Plugins/Interchange/Runtime/Source/Pipelines/Private/InterchangeGenericMeshPipeline.cpp
--- a/Engine/Plugins/Interchange/Runtime/Source/Pipelines/Private/InterchangeGenericMeshPipeline.cpp
+++ b/Engine/Plugins/Interchange/Runtime/Source/Pipelines/Private/InterchangeGenericMeshPipeline.cpp
@@ -754,9 +754,6 @@ void UInterchangeGenericMeshPipeline::CreateAssemblyPartDependencies()
 		// Assuming we can find the part, set up the required factory node dependency so that 
 		// the part mesh asset is imported and built before the base mesh asset.
 
-		TArray<FString> FactoryDependencies;
-		CurrentFactoryNode->GetFactoryDependencies(FactoryDependencies);
-
 		for (const FString& AssemblyPartUid : AssemblyPartDependencies)
 		{
 			if (AssemblyPartUid == CurrentMeshUid)
@@ -780,10 +777,8 @@ void UInterchangeGenericMeshPipeline::CreateAssemblyPartDependencies()
 				const FString AssemblyPartFactoryNodeUid = AssemblyPartFactoryNode->GetUniqueID();
 				CurrentFactoryNode->SetAssemblyPartDependencyUid(AssemblyPartUid, AssemblyPartFactoryNodeUid);
 				// Create a factory dependency so the part asset is imported before the mesh asset
-				if (!FactoryDependencies.Contains(AssemblyPartFactoryNodeUid))
-				{
-					CurrentFactoryNode->AddFactoryDependencyUid(AssemblyPartFactoryNodeUid);
-				}
+				CurrentFactoryNode->AddFactoryDependencyUid(AssemblyPartFactoryNodeUid);
+
 				GlobalPartUidToBaseUidTable.Add(AssemblyPartUid, CurrentMeshUid);
 			}
 		}
diff --git a/Engine/Plugins/Interchange/Runtime/Source/Pipelines/Public/InterchangePipelineMeshesUtilities.h b/Engine/Plugins/Interchange/Runtime/Source/Pipelines/Public/InterchangePipelineMeshesUtilities.h
--- a/Engine/Plugins/Interchange/Runtime/Source/Pipelines/Public/InterchangePipelineMeshesUtilities.h
+++ b/Engine/Plugins/Interchange/Runtime/Source/Pipelines/Public/InterchangePipelineMeshesUtilities.h
@@ -427,6 +427,8 @@ namespace UE::Interchange::MeshesUtilities
 		, const UInterchangeBaseNodeContainer& NodeContainer
 		, TMap<FString, FString> *ExistingSlotMaterialDependenciesPtr)
 	{
+		TRACE_CPUPROFILER_EVENT_SCOPE(UE::Interchange::MeshesUtilities::ApplySlotMaterialDependencies);
+
 		bool bKeepSectionsSeparate = false;
 		int32 IndexCounter = 0; //Only use when bKeepSectionsSeparate is true
 		if (ExistingSlotMaterialDependenciesPtr)
@@ -455,12 +457,7 @@ namespace UE::Interchange::MeshesUtilities
 				MaterialFactoryNode->SetEnabled(IsMaterialImportEnabled);
 
 				// Create a factory dependency so Material asset are imported before the static mesh asset
-				TArray<FString> FactoryDependencies;
-				FactoryNode.GetFactoryDependencies(FactoryDependencies);
-				if (!FactoryDependencies.Contains(MaterialFactoryNodeUid))
-				{
-					FactoryNode.AddFactoryDependencyUid(MaterialFactoryNodeUid);
-				}
+				FactoryNode.AddFactoryDependencyUid(MaterialFactoryNodeUid);
 			}
 		}
 
diff --git a/Engine/Source/Runtime/Interchange/Core/Private/Nodes/InterchangeBaseNode.cpp b/Engine/Source/Runtime/Interchange/Core/Private/Nodes/InterchangeBaseNode.cpp
--- a/Engine/Source/Runtime/Interchange/Core/Private/Nodes/InterchangeBaseNode.cpp
+++ b/Engine/Source/Runtime/Interchange/Core/Private/Nodes/InterchangeBaseNode.cpp
@@ -160,13 +160,7 @@ FString UInterchangeBaseNode::GetKeyDisplayName(const UE::Interchange::FAttribut
 	}
 	else if (NodeAttributeKeyString.StartsWith(UE::Interchange::FBaseNodeStaticData::TargetAssetIDsKey().ToString()))
 	{
-		KeyDisplayName = TEXT("Target Asset Index ");
-		const FString IndexKey = UE::Interchange::TArrayAttributeHelper<FString>::IndexKey();
-		int32 IndexPosition = NodeAttributeKeyString.Find(IndexKey) + IndexKey.Len();
-		if (IndexPosition < NodeAttributeKeyString.Len())
-		{
-			KeyDisplayName += NodeAttributeKeyString.RightChop(IndexPosition);
-		}
+		KeyDisplayName = UE::Interchange::BuildDisplayNameForArrayKey(TEXT("Target Asset Index "), NodeAttributeKeyString);
 	}
 	
 	return KeyDisplayName;
diff --git a/Engine/Source/Runtime/Interchange/Core/Private/Nodes/InterchangeFactoryBaseNode.cpp b/Engine/Source/Runtime/Interchange/Core/Private/Nodes/InterchangeFactoryBaseNode.cpp
--- a/Engine/Source/Runtime/Interchange/Core/Private/Nodes/InterchangeFactoryBaseNode.cpp
+++ b/Engine/Source/Runtime/Interchange/Core/Private/Nodes/InterchangeFactoryBaseNode.cpp
@@ -2,6 +2,7 @@
 #include "Nodes/InterchangeFactoryBaseNode.h"
 
 #include "CoreMinimal.h"
+#include "InterchangeLogPrivate.h"
 #include "Types/AttributeStorage.h"
 #include "UObject/Class.h"
 #include "UObject/Package.h"
@@ -72,13 +73,7 @@ FString UInterchangeFactoryBaseNode::GetKeyDisplayName(const UE::Interchange::FA
 	}
 	else if (OriginalKeyName.StartsWith(UE::Interchange::FFactoryBaseNodeStaticData::FactoryDependenciesBaseKey()))
 	{
-		KeyDisplayName = TEXT("Factory Dependencies Index ");
-		const FString IndexKey = UE::Interchange::TArrayAttributeHelper<FString>::IndexKey();
-		int32 IndexPosition = OriginalKeyName.Find(IndexKey) + IndexKey.Len();
-		if (IndexPosition < OriginalKeyName.Len())
-		{
-			KeyDisplayName += OriginalKeyName.RightChop(IndexPosition);
-		}
+		KeyDisplayName = UE::Interchange::BuildDisplayNameForArrayKey(TEXT("Factory Dependencies Index "), OriginalKeyName);
 	}
 	else if (NodeAttributeKey == UE::Interchange::FFactoryBaseNodeStaticData::ReimportStrategyFlagsKey())
 	{
@@ -235,14 +230,20 @@ int32 UInterchangeFactoryBaseNode::GetFactoryDependenciesCount() const
 	return FactoryDependencies.GetCount();
 }
 
+//Deprecated.
 void UInterchangeFactoryBaseNode::GetFactoryDependency(const int32 Index, FString& OutDependency) const
 {
-	FactoryDependencies.GetItem(Index, OutDependency);
+	UE_LOG(LogInterchangeCore, Error, TEXT("GetFactoryDependency is deprecated and will not return correct value. Please use GetFactoryDependencies/HasFactoryDependency instead."));
+}
+
+TSet<FString> UInterchangeFactoryBaseNode::GetFactoryDependencies() const
+{
+	return FactoryDependencies.ToSet();
 }
 
-void UInterchangeFactoryBaseNode::GetFactoryDependencies(TArray<FString>& OutDependencies) const
+bool UInterchangeFactoryBaseNode::HasFactoryDependency(const FString& DependencyUid) const
 {
-	FactoryDependencies.GetItems(OutDependencies);
+	return FactoryDependencies.HasItem(DependencyUid);
 }
 
 bool UInterchangeFactoryBaseNode::AddFactoryDependencyUid(const FString& DependencyUid)
diff --git a/Engine/Source/Runtime/Interchange/Core/Private/Tests/SetAttributeHelperTests.cpp b/Engine/Source/Runtime/Interchange/Core/Private/Tests/SetAttributeHelperTests.cpp
new file mode 100644
--- /dev/null
+++ b/Engine/Source/Runtime/Interchange/Core/Private/Tests/SetAttributeHelperTests.cpp
@@ -0,0 +1,281 @@
+// Copyright Epic Games, Inc. All Rights Reserved.
+
+#include "Async/ParallelFor.h"
+#include "CoreMinimal.h"
+#include "Misc/AutomationTest.h"
+#include "Nodes/InterchangeBaseNodeUtilities.h"
+#include "Serialization/MemoryReader.h"
+#include "Serialization/MemoryWriter.h"
+#include "Templates/Atomic.h"
+#include "Types/AttributeStorage.h"
+
+#if WITH_DEV_AUTOMATION_TESTS
+
+IMPLEMENT_SIMPLE_AUTOMATION_TEST(
+	FSetAttributeHelperTests,
+	"System.Runtime.Interchange.SetAttributeHelperTests",
+	EAutomationTestFlags::EditorContext | EAutomationTestFlags::EngineFilter
+)
+
+namespace UE::SetAttributeHelperTests::Private
+{
+	template<typename T>
+	void RunTestInternal(FSetAttributeHelperTests& Tester, const T Item1, const T Item2, const T Item3, const T Item4)
+	{
+		using namespace UE::Interchange;
+
+		const FString BaseKeyName = "TestKey";
+
+		// Initialize from empty
+		{
+			TSharedPtr<FAttributeStorage, ESPMode::ThreadSafe> TestStorage = MakeShared<FAttributeStorage>();
+
+			UE::Interchange::TSetAttributeHelper<T> Helper;
+			Helper.Initialize(TestStorage, BaseKeyName);
+
+			Tester.TestEqual(TEXT("Count after empty Initialize"), Helper.GetCount(), 0);
+			Tester.TestEqual(TEXT("Count of Set after empty Initialize"), Helper.ToSet().Num(), 0);
+		}
+
+		// Initialize from previous TSetAttributeHelper
+		{
+			TSharedPtr<FAttributeStorage, ESPMode::ThreadSafe> TestStorage = MakeShared<FAttributeStorage>();
+			{
+				UE::Interchange::TSetAttributeHelper<T> Helper;
+				Helper.Initialize(TestStorage, BaseKeyName);
+
+				Helper.AddItem(Item1);
+				Helper.AddItem(Item2);
+				Helper.AddItem(Item3);
+
+				Tester.TestEqual(TEXT("Count after pre-filling a TSetAttributeHelper"), Helper.GetCount(), 3);
+				Tester.TestTrue(TEXT(""), Helper.HasItem(Item1));
+				Tester.TestTrue(TEXT(""), Helper.HasItem(Item2));
+				Tester.TestTrue(TEXT(""), Helper.HasItem(Item3));
+			}
+
+			UE::Interchange::TSetAttributeHelper<T> AnotherHelper;
+			AnotherHelper.Initialize(TestStorage, BaseKeyName);
+
+			Tester.TestEqual(TEXT("Count after Initialize from a previous TSetAttributeHelper"), AnotherHelper.GetCount(), 3);
+			Tester.TestTrue(TEXT(""), AnotherHelper.HasItem(Item1));
+			Tester.TestTrue(TEXT(""), AnotherHelper.HasItem(Item2));
+			Tester.TestTrue(TEXT(""), AnotherHelper.HasItem(Item3));
+
+			TSet<T> Set = AnotherHelper.ToSet();
+			Tester.TestEqual(TEXT("Count of Set after Initialize from a previous TSetAttributeHelper"), Set.Num(), 3);
+			Tester.TestTrue(TEXT(""), Set.Contains(Item1));
+			Tester.TestTrue(TEXT(""), Set.Contains(Item2));
+			Tester.TestTrue(TEXT(""), Set.Contains(Item3));
+		}
+
+		// Initialize from previous TArrayAttributeHelper
+		{
+			TSharedPtr<FAttributeStorage, ESPMode::ThreadSafe> TestStorage = MakeShared<FAttributeStorage>();
+			{
+				UE::Interchange::TArrayAttributeHelper<T> ArrayHelper;
+				ArrayHelper.Initialize(TestStorage, BaseKeyName);
+
+				ArrayHelper.AddItem(Item1);
+				ArrayHelper.AddItem(Item2);
+				ArrayHelper.AddItem(Item3);
+			}
+
+			{
+				UE::Interchange::TSetAttributeHelper<T> SetHelper;
+				SetHelper.Initialize(TestStorage, BaseKeyName);
+
+				Tester.TestEqual(TEXT("Count after Initialize from a previous TArrayAttributeHelper"), SetHelper.GetCount(), 3);
+				Tester.TestTrue(TEXT(""), SetHelper.HasItem(Item1));
+				Tester.TestTrue(TEXT(""), SetHelper.HasItem(Item2));
+				Tester.TestTrue(TEXT(""), SetHelper.HasItem(Item3));
+
+				TSet<T> Set = SetHelper.ToSet();
+				Tester.TestEqual(TEXT("Count of Set after Initialize from a previous TArrayAttributeHelper"), Set.Num(), 3);
+				Tester.TestTrue(TEXT(""), Set.Contains(Item1));
+				Tester.TestTrue(TEXT(""), Set.Contains(Item2));
+				Tester.TestTrue(TEXT(""), Set.Contains(Item3));
+
+				SetHelper.AddItem(Item4);
+				SetHelper.RemoveItem(Item1);
+				SetHelper.RemoveItem(Item3);
+
+				Tester.TestFalse(TEXT(""), SetHelper.HasItem(Item1));
+				Tester.TestTrue(TEXT(""), SetHelper.HasItem(Item2));
+				Tester.TestFalse(TEXT(""), SetHelper.HasItem(Item3));
+				Tester.TestTrue(TEXT(""), SetHelper.HasItem(Item4));
+			}
+
+			{
+				UE::Interchange::TArrayAttributeHelper<T> ArrayHelper;
+				ArrayHelper.Initialize(TestStorage, BaseKeyName);
+
+				Tester.TestEqual(TEXT("Count of TArrayAttributeHelper after Initialize from a previous TSetAttributeHelper"), ArrayHelper.GetCount(), 2);
+
+				TArray<T> OutItems;
+				ArrayHelper.GetItems(OutItems);
+
+				Tester.TestEqual(TEXT("Count of TArray returned by TArrayAttributeHelper"), OutItems.Num(), 2);
+				Tester.TestTrue(TEXT("Verify Item2 is in the TArrayAttributeHelper"), OutItems.Contains(Item2));
+				Tester.TestTrue(TEXT("Verify Item4 is in the TArrayAttributeHelper"), OutItems.Contains(Item4));
+			}
+		}
+
+		// AddItem, HasItem, RemoveItem
+		{
+			TSharedPtr<FAttributeStorage, ESPMode::ThreadSafe> TestStorage = MakeShared<FAttributeStorage>();
+			UE::Interchange::TSetAttributeHelper<T> Helper;
+			Helper.Initialize(TestStorage, BaseKeyName);
+
+			bool bRet;
+
+			// Add
+			bRet = Helper.AddItem(Item1);
+
+			Tester.TestTrue(TEXT(""), bRet);
+			Tester.TestEqual(TEXT("Count after adding 1 item"), Helper.GetCount(), 1);
+			Tester.TestTrue(TEXT(""), Helper.HasItem(Item1));
+			Tester.TestFalse(TEXT(""), Helper.HasItem(Item2));
+			Tester.TestFalse(TEXT(""), Helper.HasItem(Item3));
+			Tester.TestFalse(TEXT(""), Helper.HasItem(Item4));
+
+			// Duplicate Add
+			bRet = Helper.AddItem(Item1);
+
+			Tester.TestFalse(TEXT(""), bRet);
+			Tester.TestEqual(TEXT("Count after adding an existing item"), Helper.GetCount(), 1);
+			Tester.TestTrue(TEXT(""), Helper.HasItem(Item1));
+			Tester.TestFalse(TEXT(""), Helper.HasItem(Item2));
+			Tester.TestFalse(TEXT(""), Helper.HasItem(Item3));
+			Tester.TestFalse(TEXT(""), Helper.HasItem(Item4));
+
+			// Add
+			Helper.AddItem(Item2);
+			Helper.AddItem(Item3);
+
+			Tester.TestEqual(TEXT("Count after adding 2 items"), Helper.GetCount(), 3);
+			Tester.TestTrue(TEXT(""), Helper.HasItem(Item1));
+			Tester.TestTrue(TEXT(""), Helper.HasItem(Item2));
+			Tester.TestTrue(TEXT(""), Helper.HasItem(Item3));
+			Tester.TestFalse(TEXT(""), Helper.HasItem(Item4));
+
+			// Remove non existing
+			bRet = Helper.RemoveItem(Item4);
+
+			Tester.TestFalse(TEXT(""), bRet);
+			Tester.TestEqual(TEXT("Count after removing an non-existing item"), Helper.GetCount(), 3);
+			Tester.TestTrue(TEXT(""), Helper.HasItem(Item1));
+			Tester.TestTrue(TEXT(""), Helper.HasItem(Item2));
+			Tester.TestTrue(TEXT(""), Helper.HasItem(Item3));
+			Tester.TestFalse(TEXT(""), Helper.HasItem(Item4));
+
+			// Remove at first index
+			bRet = Helper.RemoveItem(Item1);
+
+			Tester.TestTrue(TEXT(""), bRet);
+			Tester.TestEqual(TEXT("Count after removing item 1"), Helper.GetCount(), 2);
+			Tester.TestFalse(TEXT(""), Helper.HasItem(Item1));
+			Tester.TestTrue(TEXT(""), Helper.HasItem(Item2));
+			Tester.TestTrue(TEXT(""), Helper.HasItem(Item3));
+
+			// Remove at last index
+			bRet = Helper.RemoveItem(Item3);
+
+			Tester.TestTrue(TEXT(""), bRet);
+			Tester.TestEqual(TEXT("Count after removing item 3"), Helper.GetCount(), 1);
+			Tester.TestFalse(TEXT(""), Helper.HasItem(Item1));
+			Tester.TestTrue(TEXT(""), Helper.HasItem(Item2));
+			Tester.TestFalse(TEXT(""), Helper.HasItem(Item3));
+
+			// Remove last item in the Set
+			bRet = Helper.RemoveItem(Item2);
+
+			Tester.TestTrue(TEXT(""), bRet);
+			Tester.TestEqual(TEXT("Count after removing last item"), Helper.GetCount(), 0);
+			Tester.TestFalse(TEXT(""), Helper.HasItem(Item1));
+			Tester.TestFalse(TEXT(""), Helper.HasItem(Item2));
+			Tester.TestFalse(TEXT(""), Helper.HasItem(Item3));
+		}
+
+		// Empty
+		{
+			TSharedPtr<FAttributeStorage, ESPMode::ThreadSafe> TestStorage = MakeShared<FAttributeStorage>();
+			UE::Interchange::TSetAttributeHelper<T> Helper;
+			Helper.Initialize(TestStorage, BaseKeyName);
+
+			Helper.Empty();
+			Tester.TestEqual(TEXT("Count after emptying a new TSetAttributeHelper"), Helper.GetCount(), 0);
+
+			Helper.AddItem(Item1);
+			Helper.AddItem(Item2);
+			Helper.AddItem(Item3);
+			Helper.AddItem(Item4);
+
+			Tester.TestEqual(TEXT("Count after adding 4 items"), Helper.GetCount(), 4);
+
+			Helper.Empty();
+			Tester.TestEqual(TEXT("Count after emptying a full TSetAttributeHelper"), Helper.GetCount(), 0);
+		}
+
+		// Gaps and duplicate
+		{
+			TSharedPtr<FAttributeStorage, ESPMode::ThreadSafe> TestStorage = MakeShared<FAttributeStorage>();
+			{
+				UE::Interchange::TArrayAttributeHelper<T> ArrayHelper;
+				ArrayHelper.Initialize(TestStorage, BaseKeyName);
+
+				ArrayHelper.AddItem(Item1);
+				ArrayHelper.AddItem(Item1); // Duplicate
+
+				ArrayHelper.AddItem(Item2);
+
+				ArrayHelper.AddItem(Item3);
+				ArrayHelper.AddItem(Item3); // Duplicate
+
+				ArrayHelper.AddItem(Item4);
+
+				Tester.TestEqual(TEXT("Count after pre-filling the TArrayAttributeHelper"), ArrayHelper.GetCount(), 6);
+			}
+
+			{
+				UE::Interchange::TSetAttributeHelper<T> SetHelper;
+				SetHelper.Initialize(TestStorage, BaseKeyName);
+
+				Tester.TestEqual(TEXT("Count after Initializing with duplicated items"), SetHelper.GetCount(), 4);
+				Tester.TestTrue(TEXT(""), SetHelper.HasItem(Item1));
+				Tester.TestTrue(TEXT(""), SetHelper.HasItem(Item2));
+				Tester.TestTrue(TEXT(""), SetHelper.HasItem(Item3));
+				Tester.TestTrue(TEXT(""), SetHelper.HasItem(Item4));
+			}
+
+			// Manually remove an entry in the middle to "corrupt" the storage and Initialize a TSetAttributeHelper
+			{
+				FAttributeKey KeyToRemove = UE::Interchange::BuildArrayKeyFromIndex(BaseKeyName, 2);
+				Tester.TestEqual(TEXT("Verify the key is properly removed from storage"), TestStorage->UnregisterAttribute(KeyToRemove), EAttributeStorageResult::Operation_Success);
+				
+				UE::Interchange::TSetAttributeHelper<T> SetHelper;
+				SetHelper.Initialize(TestStorage, BaseKeyName);
+
+				Tester.TestEqual(TEXT("Count after Initializing with corrupted items"), SetHelper.GetCount(), 3);
+				Tester.TestTrue(TEXT(""), SetHelper.HasItem(Item1));
+				Tester.TestFalse(TEXT(""), SetHelper.HasItem(Item2));
+				Tester.TestTrue(TEXT(""), SetHelper.HasItem(Item3));
+				Tester.TestTrue(TEXT(""), SetHelper.HasItem(Item4));
+			}
+		}
+	}
+}	 // namespace UE::SetAttributeHelperTests::Private
+
+bool FSetAttributeHelperTests::RunTest(const FString& Parameters)
+{
+	using namespace UE::SetAttributeHelperTests::Private;
+
+	RunTestInternal<FString>(*this, "Foo", "Bar", "Baz", "Qux");
+	RunTestInternal<int32>(*this, 111, 2222, 33333, 444444);
+	RunTestInternal<double>(*this, 111.111, 2222.222, 33333.333, 444444.44);
+	RunTestInternal<float>(*this, 111.111f, 2222.222f, 33333.333f, 444444.44f);
+
+	return true;
+}
+
+#endif	  // WITH_DEV_AUTOMATION_TESTS
diff --git a/Engine/Source/Runtime/Interchange/Core/Public/Nodes/InterchangeBaseNodeUtilities.h b/Engine/Source/Runtime/Interchange/Core/Public/Nodes/InterchangeBaseNodeUtilities.h
--- a/Engine/Source/Runtime/Interchange/Core/Public/Nodes/InterchangeBaseNodeUtilities.h
+++ b/Engine/Source/Runtime/Interchange/Core/Public/Nodes/InterchangeBaseNodeUtilities.h
@@ -13,6 +13,34 @@ namespace UE
 {
 namespace Interchange
 {
+	static const FString& ArrayIndexKey()
+	{
+		static FString IndexKeyString = "_NameIndex_";
+		return IndexKeyString;
+	}
+
+	static FAttributeKey BuildArrayKeyFromIndex(const FString& BaseKeyName, uint32 Index)
+	{
+		FString DepIndexKeyString = BaseKeyName + ArrayIndexKey() + FString::FromInt(Index);
+		return FAttributeKey(DepIndexKeyString);
+	}
+
+	static FString BuildDisplayNameForArrayKey(const FString& BaseDisplayName, const FString& ArrayAttributeKey)
+	{
+		static FString IndexKeyString = ArrayIndexKey();
+		const int32 IndexPosition = ArrayAttributeKey.Find(IndexKeyString);
+		if (IndexPosition != INDEX_NONE)
+		{
+			const int32 ChopIndex = IndexPosition + IndexKeyString.Len();
+			if (ChopIndex < ArrayAttributeKey.Len())
+			{
+				return BaseDisplayName + ArrayAttributeKey.RightChop(ChopIndex);
+			}
+		}
+
+		return BaseDisplayName;
+	}
+
 
 /**
  * Helper class to manage an array of items inside an attribute storage
@@ -38,15 +66,9 @@ public:
 		KeyCount = BaseKeyName;
 	}
 
-	static const FString& IndexKey()
-	{
-		static FString IndexKeyString = "_NameIndex_";
-		return IndexKeyString;
-	}
-
 	int32 GetCount() const
 	{
-		//The Class must be initialise properly before we can use it
+		//The Class must be initialize properly before we can use it
 		TSharedPtr<FAttributeStorage, ESPMode::ThreadSafe> AttributePtr = Attributes.Pin();
 		if (!ensure(AttributePtr.IsValid()))
 		{
@@ -66,7 +88,7 @@ public:
 
 	void GetItem(const int32 Index, ItemType& OutItem) const
 	{
-		//The Class must be initialise properly before we can use it
+		//The Class must be initialize properly before we can use it
 		TSharedPtr<FAttributeStorage, ESPMode::ThreadSafe> AttributePtr = Attributes.Pin();
 		if (!ensure(AttributePtr.IsValid()))
 		{
@@ -99,7 +121,7 @@ public:
 
 	void GetItems(TArray<ItemType>& OutItems) const
 	{
-		//The Class must be initialise properly before we can use it
+		//The Class must be initialize properly before we can use it
 		TSharedPtr<FAttributeStorage, ESPMode::ThreadSafe> AttributePtr = Attributes.Pin();
 		if (!ensure(AttributePtr.IsValid()))
 		{
@@ -137,7 +159,7 @@ public:
 
 	bool AddItem(const ItemType& Item)
 	{
-		//The Class must be initialise properly before we can use it
+		//The Class must be initialize properly before we can use it
 		TSharedPtr<FAttributeStorage, ESPMode::ThreadSafe> AttributePtr = Attributes.Pin();
 		if (!ensure(AttributePtr.IsValid()))
 		{
@@ -177,7 +199,7 @@ public:
 
 	bool RemoveItem(const ItemType& ItemToDelete)
 	{
-		//The Class must be initialise properly before we can use it
+		//The Class must be initialize properly before we can use it
 		TSharedPtr<FAttributeStorage, ESPMode::ThreadSafe> AttributePtr = Attributes.Pin();
 		if (!ensure(AttributePtr.IsValid()))
 		{
@@ -241,7 +263,7 @@ public:
 
 	bool RemoveAllItems()
 	{
-		//The Class must be initialise properly before we can use it
+		//The Class must be initialize properly before we can use it
 		TSharedPtr<FAttributeStorage, ESPMode::ThreadSafe> AttributePtr = Attributes.Pin();
 		if (!ensure(AttributePtr.IsValid()))
 		{
@@ -281,8 +303,7 @@ private:
 
 	FAttributeKey GetIndexKey(int32 Index) const
 	{
-		FString DepIndexKeyString = GetKeyCount().ToString() + IndexKey() + FString::FromInt(Index);
-		return FAttributeKey(DepIndexKeyString);
+		return UE::Interchange::BuildArrayKeyFromIndex(GetKeyCount().ToString(), Index);
 	}
 };
 
@@ -603,5 +624,247 @@ private:
 	TWeakPtr<FAttributeStorage, ESPMode::ThreadSafe> Attributes = nullptr;
 };
 
+
+/**
+ * Helper class to manage a TSet where the value type can be any type support by the FAttributeStorage
+ * This class can also load storage made for the TArrayAttributeHelper but will ignore duplicate values
+ */
+template<typename ValueType>
+class TSetAttributeHelper
+{
+	template<typename T>
+	using TAttributeHandle = FAttributeStorage::TAttributeHandle<T>;
+
+public:
+	void Initialize(const TSharedPtr<FAttributeStorage, ESPMode::ThreadSafe>& InAttributes, const FString& BaseKeyName)
+	{
+		LLM_SCOPE_BYNAME(TEXT("Interchange"));
+
+		Attributes = InAttributes;
+		KeyCount = BaseKeyName;
+
+		if (!InAttributes->ContainAttribute(KeyCount))
+		{
+			EAttributeStorageResult Result = InAttributes->RegisterAttribute<int32>(KeyCount, 0);
+			check(Result == EAttributeStorageResult::Operation_Success);
+			KeyCountHandle = InAttributes->GetAttributeHandle<int32>(KeyCount);
+		}
+		else
+		{
+			KeyCountHandle = InAttributes->GetAttributeHandle<int32>(KeyCount);
+			RebuildCache();
+		}
+	}
+
+	TSetAttributeHelper() = default;
+	TSetAttributeHelper(const TSetAttributeHelper&) = default;
+	TSetAttributeHelper(TSetAttributeHelper&&) = default;
+	TSetAttributeHelper& operator=(const TSetAttributeHelper&) = default;
+	TSetAttributeHelper& operator=(TSetAttributeHelper&&) = default;
+
+	~TSetAttributeHelper()
+	{
+		Attributes.Reset();
+	}
+
+	/*
+	* Adds an item. Returns true if the item is added successfully.
+	* Returns false if the item is already present or if Initialize was not properly called before adding an item.
+	*/
+	bool AddItem(const ValueType& InValue)
+	{
+		if (!CachedValues.Contains(InValue))
+		{
+			TSharedPtr<FAttributeStorage, ESPMode::ThreadSafe> AttributesPtr = Attributes.Pin();
+			if (!ensure(AttributesPtr.IsValid()))
+			{
+				return false;
+			}
+
+			FAttributeKey AttributeKey = GetIndexKey(CachedValues.Num());
+			ensure(AttributesPtr->RegisterAttribute<ValueType>(AttributeKey, InValue) == EAttributeStorageResult::Operation_Success);
+			CachedValues.Add(InValue, AttributeKey);
+			KeyCountHandle.Set(CachedValues.Num());
+
+			return true;
+		}
+
+		return false;
+	}
+
+	bool HasItem(const ValueType& InValue) const
+	{
+		TSharedPtr<FAttributeStorage, ESPMode::ThreadSafe> AttributesPtr = Attributes.Pin();
+		if (!ensure(AttributesPtr.IsValid()))
+		{
+			return false;
+		}
+
+		return CachedValues.Contains(InValue);
+	}
+
+	/*
+	* Remove an item if it exists. Will realign attributes by indices in the storage if needed (0(N)).
+	*/
+	bool RemoveItem(const ValueType& InValue)
+	{
+		TSharedPtr<FAttributeStorage, ESPMode::ThreadSafe> AttributesPtr = Attributes.Pin();
+		if (!ensure(AttributesPtr.IsValid()))
+		{
+			return false;
+		}
+
+		{
+			FAttributeKey* KeyToRemove = CachedValues.Find(InValue);
+			if (!KeyToRemove)
+			{
+				return false;
+			}
+
+			AttributesPtr->UnregisterAttribute(*KeyToRemove);
+			CachedValues.Remove(InValue);
+			KeyCountHandle.Set(CachedValues.Num());
+		}
+
+		//When removing an item, we need to re-align the keys by indices without having gaps so it can be loaded properly later.
+		uint32 Index = 0;
+		for (const TPair<ValueType, FAttributeKey>& Pair : CachedValues)
+		{
+			const FAttributeKey& AttributeKey = Pair.Value;
+			const FAttributeKey NewAttributeKey = GetIndexKey(Index++);
+			if (NewAttributeKey == AttributeKey)
+			{
+				//This item is already assigned the correct index.
+				continue;
+			}
+
+			AttributesPtr->UnregisterAttribute(AttributeKey);
+			ensure(AttributesPtr->RegisterAttribute<ValueType>(NewAttributeKey, Pair.Key) == EAttributeStorageResult::Operation_Success);
+
+			CachedValues[Pair.Key] = NewAttributeKey;
+		}
+
+		return true;
+	}
+
+	void Reserve(int32 Number)
+	{
+		CachedValues.Reserve(Number);
+	}
+
+	void Empty(int32 NumOfExpectedElements = 0)
+	{
+		EmptyInternal(NumOfExpectedElements);
+	}
+
+	int32 GetCount() const
+	{
+		TSharedPtr<FAttributeStorage, ESPMode::ThreadSafe> AttributesPtr = Attributes.Pin();
+		ensure(AttributesPtr.IsValid());
+		return CachedValues.Num();
+	}
+
+	TSetAttributeHelper& operator=(const TSet<ValueType>& InSet)
+	{
+		TSharedPtr<FAttributeStorage, ESPMode::ThreadSafe> AttributesPtr = Attributes.Pin();
+		if (!ensure(AttributesPtr.IsValid()))
+		{
+			return *this;
+		}
+
+		// Empty
+		EmptyInternal(InSet.Num());
+
+		KeyCountHandle.Set(InSet.Num());
+
+		for (const ValueType& Value : InSet)
+		{
+			AddItem(Value);
+		}
+
+		return *this;
+	}
+
+	TSet<ValueType> ToSet() const
+	{
+		TSet<ValueType> Set;
+
+		TSharedPtr<FAttributeStorage, ESPMode::ThreadSafe> AttributesPtr = Attributes.Pin();
+		if (!ensure(AttributesPtr.IsValid()))
+		{
+			return Set;
+		}
+		
+		CachedValues.GetKeys(Set);
+
+		return Set;
+	}
+
+	void RebuildCache()
+	{
+		LLM_SCOPE_BYNAME(TEXT("Interchange"));
+
+		TSharedPtr<FAttributeStorage, ESPMode::ThreadSafe> AttributesPtr = Attributes.Pin();
+		if (AttributesPtr.IsValid() && KeyCountHandle.IsValid())
+		{
+			int32 ItemCount;
+			if (KeyCountHandle.Get(ItemCount) == EAttributeStorageResult::Operation_Success)
+			{
+				CachedValues.Empty(ItemCount);
+
+				for (int32 Index = 0; Index < ItemCount; Index++)
+				{
+					TAttributeHandle<ValueType> ValueAttribute = AttributesPtr->GetAttributeHandle<ValueType>(GetIndexKey(Index));
+					if (!ValueAttribute.IsValid())
+					{
+						continue;
+					}
+
+					ValueType Value;
+					if (ValueAttribute.Get(Value) == EAttributeStorageResult::Operation_Success)
+					{
+						CachedValues.Add(Value, ValueAttribute.GetKey());
+					}
+				}
+			}
+		}
+	}
+
+private:
+	void EmptyInternal(int32 NumOfExpectedElements)
+	{
+		TSharedPtr<FAttributeStorage, ESPMode::ThreadSafe> AttributesPtr = Attributes.Pin();
+		if (!ensure(AttributesPtr.IsValid()))
+		{
+			return;
+		}
+
+		for (const TPair<ValueType, FAttributeKey>& Pair : CachedValues)
+		{
+			const FAttributeKey& KeyAttribute = Pair.Value;
+			AttributesPtr->UnregisterAttribute(KeyAttribute);
+		}
+
+		CachedValues.Empty(NumOfExpectedElements);
+		KeyCountHandle.Set(0);
+	}
+
+	TMap<ValueType, FAttributeKey> CachedValues;
+	TAttributeHandle<int32> KeyCountHandle;
+	TWeakPtr<FAttributeStorage, ESPMode::ThreadSafe> Attributes = nullptr;
+
+	FAttributeKey KeyCount;
+	FAttributeKey GetKeyCount() const
+	{
+		ensure(!KeyCount.Key.IsEmpty()); return KeyCount;
+	}
+
+	FAttributeKey GetIndexKey(int32 Index) const
+	{
+		ensure(!KeyCount.Key.IsEmpty());
+		return UE::Interchange::BuildArrayKeyFromIndex(KeyCount.Key, Index); //Must be identical to the TArrayAttributeHelper
+	}
+};
+
 } //ns Interchange
 } //ns UE
diff --git a/Engine/Source/Runtime/Interchange/Core/Public/Nodes/InterchangeFactoryBaseNode.h b/Engine/Source/Runtime/Interchange/Core/Public/Nodes/InterchangeFactoryBaseNode.h
--- a/Engine/Source/Runtime/Interchange/Core/Public/Nodes/InterchangeFactoryBaseNode.h
+++ b/Engine/Source/Runtime/Interchange/Core/Public/Nodes/InterchangeFactoryBaseNode.h
@@ -191,6 +191,16 @@ class UInterchangeFactoryBaseNode : public UInterchangeBaseNode
 public:
 	INTERCHANGECORE_API UInterchangeFactoryBaseNode();
 
+	virtual void Serialize(FArchive& Ar) override
+	{
+		Super::Serialize(Ar);
+
+		if (Ar.IsLoading() && bIsInitialized)
+		{
+			FactoryDependencies.RebuildCache();
+		}
+	}
+
 #if WITH_EDITOR
 	INTERCHANGECORE_API virtual FString GetKeyDisplayName(const UE::Interchange::FAttributeKey& NodeAttributeKey) const override;
 	INTERCHANGECORE_API virtual FString GetAttributeCategory(const UE::Interchange::FAttributeKey& NodeAttributeKey) const override;
@@ -273,14 +283,21 @@ public:
 	 * Retrieve the dependencies for this object.
 	 */
 	UFUNCTION(BlueprintCallable, Category = "Interchange | Node")
-	INTERCHANGECORE_API void GetFactoryDependencies(TArray<FString>& OutDependencies ) const;
+	INTERCHANGECORE_API TSet<FString> GetFactoryDependencies() const;
 
 	/**
 	 * Retrieve one dependency for this object.
 	 */
 	UFUNCTION(BlueprintCallable, Category = "Interchange | Node")
+	UE_DEPRECATED(5.8, "Use GetFactoryDependencies or HasFactoryDependency instead.")
 	INTERCHANGECORE_API void GetFactoryDependency(const int32 Index, FString& OutDependency) const;
 
+	/**
+	 * Returns true if this factory node has a dependency on the factory which Uid is DependencyUid.
+	 */
+	UFUNCTION(BlueprintCallable, Category = "Interchange | Node")
+	INTERCHANGECORE_API bool HasFactoryDependency(const FString& DependencyUid) const;
+
 	/**
 	 * Add one dependency to this object.
 	 */
@@ -431,7 +448,7 @@ protected:
 	 * Those dependencies are used by the Interchange parsing task to make sure the assets are created in the correct order.
 	 * Examples: The Mesh factory node will have dependencies on the material factory node, and the Material factory node will have dependencies on the texture factory node.
 	 */
-	UE::Interchange::TArrayAttributeHelper<FString> FactoryDependencies;
+	UE::Interchange::TSetAttributeHelper<FString> FactoryDependencies;
 
 	/* This array holds the delegate to apply the attribute that has to be set on an UObject. */
 	TMap<UClass*, TArray<UE::Interchange::FApplyAttributeToAsset>> ApplyCustomAttributeDelegates;
diff --git a/Engine/Source/Runtime/Interchange/Engine/Private/InterchangeImportCommon.cpp b/Engine/Source/Runtime/Interchange/Engine/Private/InterchangeImportCommon.cpp
--- a/Engine/Source/Runtime/Interchange/Engine/Private/InterchangeImportCommon.cpp
+++ b/Engine/Source/Runtime/Interchange/Engine/Private/InterchangeImportCommon.cpp
@@ -47,7 +47,7 @@ namespace UE::Interchange
 			return AssetImportData;
 		}
 
-		void RecursivelyDuplicateFactoryNodeDependencies(UInterchangeBaseNodeContainer* SourceContainer, UInterchangeBaseNodeContainer* DestinationContainer, TArray<FString>& FactoryDependencies)
+		void RecursivelyDuplicateFactoryNodeDependencies(UInterchangeBaseNodeContainer* SourceContainer, UInterchangeBaseNodeContainer* DestinationContainer, TSet<FString>& FactoryDependencies)
 		{
 			if (!DestinationContainer)
 			{
@@ -61,8 +61,7 @@ namespace UE::Interchange
 					FObjectDuplicationParameters DupParam(FactoryNode, DestinationContainer);
 					DestinationContainer->AddNode(CastChecked<UInterchangeBaseNode>(StaticDuplicateObjectEx(DupParam)));
 
-					TArray<FString> ChildFactoryDependencies;
-					FactoryNode->GetFactoryDependencies(ChildFactoryDependencies);
+					TSet<FString> ChildFactoryDependencies = FactoryNode->GetFactoryDependencies();
 					RecursivelyDuplicateFactoryNodeDependencies(SourceContainer, DestinationContainer, ChildFactoryDependencies);
 				}
 			}
@@ -77,7 +76,7 @@ namespace UE::Interchange
 			{
 				//Set the interchange node graph data
 				AssetImportData->NodeUniqueID = Parameters.NodeUniqueID;
-				TArray<FString> FactoryDependencies;
+				TSet<FString> FactoryDependencies;
 				FactoryDependencies.Add(AssetImportData->NodeUniqueID);
 				RecursivelyDuplicateFactoryNodeDependencies(Parameters.NodeContainer, FactoryNodeContainer, FactoryDependencies);
 			}
diff --git a/Engine/Source/Runtime/Interchange/Engine/Private/Tasks/InterchangeTaskParsing.cpp b/Engine/Source/Runtime/Interchange/Engine/Private/Tasks/InterchangeTaskParsing.cpp
--- a/Engine/Source/Runtime/Interchange/Engine/Private/Tasks/InterchangeTaskParsing.cpp
+++ b/Engine/Source/Runtime/Interchange/Engine/Private/Tasks/InterchangeTaskParsing.cpp
@@ -41,7 +41,7 @@ namespace UE::InterchangeTaskParsing::Private
 		FString UniqueID;
 		int32 SourceIndex = INDEX_NONE;
 		bool bIsSceneNode = false;
-		TArray<FString> Dependencies;
+		TSet<FString> Dependencies;
 		uint64 GraphEventRef = INTERCHANGE_INVALID_TASK_ID;
 		TArray<uint64> Prerequisites;
 		const UClass* FactoryClass;
@@ -288,7 +288,7 @@ void UE::Interchange::FTaskParsing::Execute()
 							NodeTaskData.SourceIndex = SourceIndex;
 							NodeTaskData.bIsSceneNode = bIsSceneNode;
 							NodeTaskData.Nodes.Add(FactoryNode);
-							FactoryNode->GetFactoryDependencies(NodeTaskData.Dependencies);
+							NodeTaskData.Dependencies = FactoryNode->GetFactoryDependencies();
 							NodeTaskData.FactoryClass = RegisteredFactoryClass;
 							int32 TaskDataIndex = TaskDatas.Add(NodeTaskData);
 							TaskUniqueIdToIndex.Add(NodeTaskData.UniqueID, TaskDataIndex);
