Improve BP scripting of filepicker, enabling BP users to pass in settings to the file picker. This enables user to determine if single file or multi-file picking is allowed, as well as all factory are support (all file types) or a subset.

#rnx
#okforgithhub public

#Synced-CL 51095113

diff --git a/Engine/Plugins/Interchange/Editor/Source/Utilities/Private/InterchangeOpenFileDialog.cpp b/Engine/Plugins/Interchange/Editor/Source/Utilities/Private/InterchangeOpenFileDialog.cpp
--- a/Engine/Plugins/Interchange/Editor/Source/Utilities/Private/InterchangeOpenFileDialog.cpp
+++ b/Engine/Plugins/Interchange/Editor/Source/Utilities/Private/InterchangeOpenFileDialog.cpp
@@ -55,7 +55,7 @@ namespace UE::Interchange::Utilities::Private
 		return FormatString;
 	}
 
-	bool FilePickerDialog(const FString& Extensions, FInterchangeFilePickerParameters& Parameters, TArray<FString>& OutFilenames)
+	bool FilePickerDialog(const FString& Extensions, const FInterchangeFilePickerParameters& Parameters, TArray<FString>& OutFilenames)
 	{
 		// First, display the file open dialog for selecting the file.
 		IDesktopPlatform* DesktopPlatform = FDesktopPlatformModule::Get();
@@ -80,7 +80,7 @@ namespace UE::Interchange::Utilities::Private
 	}
 } //ns UE::Interchange::Utilities::Private
 
-bool UInterchangeFilePickerGeneric::FilePickerForTranslatorAssetType(const EInterchangeTranslatorAssetType TranslatorAssetType, FInterchangeFilePickerParameters& Parameters, TArray<FString>& OutFilenames)
+bool UInterchangeFilePickerGeneric::FilePickerForTranslatorAssetType(const EInterchangeTranslatorAssetType TranslatorAssetType, const FInterchangeFilePickerParameters& Parameters, TArray<FString>& OutFilenames)
 {
 	FString Extensions = UE::Interchange::Utilities::Private::GetOpenFileDialogExtensions
 	(
@@ -90,7 +90,7 @@ bool UInterchangeFilePickerGeneric::FilePickerForTranslatorAssetType(const EInte
 	return UE::Interchange::Utilities::Private::FilePickerDialog(Extensions, Parameters, OutFilenames);
 }
 
-bool UInterchangeFilePickerGeneric::FilePickerForTranslatorType(const EInterchangeTranslatorType TranslatorType, FInterchangeFilePickerParameters& Parameters, TArray<FString>& OutFilenames)
+bool UInterchangeFilePickerGeneric::FilePickerForTranslatorType(const EInterchangeTranslatorType TranslatorType, const FInterchangeFilePickerParameters& Parameters, TArray<FString>& OutFilenames)
 {
 	FString Extensions = UE::Interchange::Utilities::Private::GetOpenFileDialogExtensions
 	(
diff --git a/Engine/Plugins/Interchange/Editor/Source/Utilities/Public/InterchangeOpenFileDialog.h b/Engine/Plugins/Interchange/Editor/Source/Utilities/Public/InterchangeOpenFileDialog.h
--- a/Engine/Plugins/Interchange/Editor/Source/Utilities/Public/InterchangeOpenFileDialog.h
+++ b/Engine/Plugins/Interchange/Editor/Source/Utilities/Public/InterchangeOpenFileDialog.h
@@ -17,8 +17,8 @@ public:
 
 protected:
 
-	UE_API virtual bool FilePickerForTranslatorAssetType(const EInterchangeTranslatorAssetType TranslatorAssetType, FInterchangeFilePickerParameters& Parameters, TArray<FString>& OutFilenames) override;
-	UE_API virtual bool FilePickerForTranslatorType(const EInterchangeTranslatorType TranslatorType, FInterchangeFilePickerParameters& Parameters, TArray<FString>& OutFilenames) override;
+	UE_API virtual bool FilePickerForTranslatorAssetType(const EInterchangeTranslatorAssetType TranslatorAssetType, const FInterchangeFilePickerParameters& Parameters, TArray<FString>& OutFilenames) override;
+	UE_API virtual bool FilePickerForTranslatorType(const EInterchangeTranslatorType TranslatorType, const FInterchangeFilePickerParameters& Parameters, TArray<FString>& OutFilenames) override;
 };
 
 #undef UE_API
diff --git a/Engine/Source/Runtime/Interchange/Engine/Public/InterchangeFilePickerBase.h b/Engine/Source/Runtime/Interchange/Engine/Public/InterchangeFilePickerBase.h
--- a/Engine/Source/Runtime/Interchange/Engine/Public/InterchangeFilePickerBase.h
+++ b/Engine/Source/Runtime/Interchange/Engine/Public/InterchangeFilePickerBase.h
@@ -15,19 +15,19 @@ struct FInterchangeFilePickerParameters
 	GENERATED_BODY()
 
 	//If true, the user will be able to select multiple files.
-	UPROPERTY(EditAnywhere, Category = "Interchange | File Picker")
+	UPROPERTY(EditAnywhere, BlueprintReadWrite, Category = "Interchange | File Picker")
 	bool bAllowMultipleFiles = false;
 
 	//If not empty, this will override the default title.
-	UPROPERTY(EditAnywhere, Category = "Interchange | File Picker")
+	UPROPERTY(EditAnywhere, BlueprintReadWrite, Category = "Interchange | File Picker")
 	FText Title = FText();
 
 	//Set the default open path that the dialog will show to the user.
-	UPROPERTY(EditAnywhere, Category = "Interchange | File Picker")
+	UPROPERTY(EditAnywhere, BlueprintReadWrite, Category = "Interchange | File Picker")
 	FString DefaultPath;
 
 	//If true, the user will be able to select any unreal editor factory + interchange file types.
-	UPROPERTY(EditAnywhere, Category = "Interchange | File Picker")
+	UPROPERTY(EditAnywhere, BlueprintReadWrite, Category = "Interchange | File Picker")
 	bool bShowAllFactoriesExtension = false;
 
 	//Add some extension to the picker. Format text item that way TEXT("fbx;Filmbox")
@@ -46,10 +46,10 @@ public:
 	 * Non-virtual helper that allows Blueprint to implement an event-based function to implement FilePickerForTranslatorAssetType().
 	 */
 	UFUNCTION(BlueprintNativeEvent, BlueprintCallable, Category = "Interchange | Utilities")
-	INTERCHANGEENGINE_API bool ScriptedFilePickerForTranslatorAssetType(const EInterchangeTranslatorAssetType TranslatorAssetType, FInterchangeFilePickerParameters& Parameters, TArray<FString>& OutFilenames);
+	INTERCHANGEENGINE_API bool ScriptedFilePickerForTranslatorAssetType(const EInterchangeTranslatorAssetType TranslatorAssetType, const FInterchangeFilePickerParameters& Parameters, TArray<FString>& OutFilenames);
 
 	/** The default implementation, which is called if the Blueprint does not have any implementation, calls the virtual FilePickerForTranslatorAssetType(). */
-	bool ScriptedFilePickerForTranslatorAssetType_Implementation(const EInterchangeTranslatorAssetType TranslatorAssetType, FInterchangeFilePickerParameters& Parameters, TArray<FString>& OutFilenames)
+	bool ScriptedFilePickerForTranslatorAssetType_Implementation(const EInterchangeTranslatorAssetType TranslatorAssetType, const FInterchangeFilePickerParameters& Parameters, TArray<FString>& OutFilenames)
 	{
 		//By default we call the virtual function FilePickerForTranslatorAssetType
 		return FilePickerForTranslatorAssetType(TranslatorAssetType, Parameters, OutFilenames);
@@ -59,10 +59,10 @@ public:
 	 * Non-virtual helper that allows Blueprint to implement an event-based function to implement FilePickerForTranslatorType().
 	 */
 	UFUNCTION(BlueprintNativeEvent, BlueprintCallable, Category = "Interchange | Utilities")
-	INTERCHANGEENGINE_API bool ScriptedFilePickerForTranslatorType(const EInterchangeTranslatorType TranslatorType, FInterchangeFilePickerParameters& Parameters, TArray<FString>& OutFilenames);
+	INTERCHANGEENGINE_API bool ScriptedFilePickerForTranslatorType(const EInterchangeTranslatorType TranslatorType, const FInterchangeFilePickerParameters& Parameters, TArray<FString>& OutFilenames);
 
 	/** The default implementation, which is called if the Blueprint does not have any implementation, calls the virtual FilePickerForTranslatorType(). */
-	bool ScriptedFilePickerForTranslatorType_Implementation(const EInterchangeTranslatorType TranslatorAssetType, FInterchangeFilePickerParameters& Parameters, TArray<FString>& OutFilenames)
+	bool ScriptedFilePickerForTranslatorType_Implementation(const EInterchangeTranslatorType TranslatorAssetType, const FInterchangeFilePickerParameters& Parameters, TArray<FString>& OutFilenames)
 	{
 		//By default we call the virtual function FilePickerForTranslatorType
 		return FilePickerForTranslatorType(TranslatorAssetType, Parameters, OutFilenames);
@@ -77,7 +77,7 @@ protected:
 	 * 
 	 * @Return - true if OutFilename was set with a valid file path, or false otherwise.
 	 */
-	virtual bool FilePickerForTranslatorAssetType(const EInterchangeTranslatorAssetType TranslatorAssetType, FInterchangeFilePickerParameters& Parameters, TArray<FString>& OutFilenames) { return false; }
+	virtual bool FilePickerForTranslatorAssetType(const EInterchangeTranslatorAssetType TranslatorAssetType, const FInterchangeFilePickerParameters& Parameters, TArray<FString>& OutFilenames) { return false; }
 
 	/**
 	 * This function must set OutFilename with a valid file that can be import to create the specified TranslatorType.
@@ -85,6 +85,6 @@ protected:
 	 *
 	 * @Return - true if OutFilename was set with a valid file path, or false otherwise.
 	 */
-	virtual bool FilePickerForTranslatorType(const EInterchangeTranslatorType TranslatorType, FInterchangeFilePickerParameters& Parameters, TArray<FString>& OutFilenames) { return false; }
+	virtual bool FilePickerForTranslatorType(const EInterchangeTranslatorType TranslatorType, const FInterchangeFilePickerParameters& Parameters, TArray<FString>& OutFilenames) { return false; }
 
 };
