[BUGFIX] Interchange Analytics GUID collision

#rnx
#Synced-CL 51376356

diff --git a/Engine/Source/Runtime/Interchange/Engine/Private/InterchangeManager.cpp b/Engine/Source/Runtime/Interchange/Engine/Private/InterchangeManager.cpp
--- a/Engine/Source/Runtime/Interchange/Engine/Private/InterchangeManager.cpp
+++ b/Engine/Source/Runtime/Interchange/Engine/Private/InterchangeManager.cpp
@@ -1693,7 +1693,7 @@ UInterchangeManager::ImportInternal(const FString& ContentPath, const UInterchan
 	// Do not collect Analytics while automation testing.
 	if (!GIsAutomationTesting)
 	{
-		AsyncHelper->AnalyticsHandler = CreateAnalyticsHandler(ContentPath);
+		AsyncHelper->AnalyticsHandler = CreateAnalyticsHandler(SourceData ? SourceData->GetFilename() : FString());
 	}
 
 	UE::Interchange::Private::FInterchangeAnalyticsHelper AnalyticsHelper(AsyncHelper->AnalyticsHandler.Get());
@@ -2389,7 +2389,7 @@ void UInterchangeManager::SetAnalyticsHandlerClass(const TSoftClassPtr<UIntercha
 	AnalyticsHandlerClass = InAnalyticsHandlerClass;
 }
 
-TStrongObjectPtr<UInterchangeAnalyticsHandlerBase> UInterchangeManager::CreateAnalyticsHandler(const FString& ContentPath)
+TStrongObjectPtr<UInterchangeAnalyticsHandlerBase> UInterchangeManager::CreateAnalyticsHandler(const FString& ImportFilePath)
 {
 	TStrongObjectPtr<UInterchangeAnalyticsHandlerBase> AnalyticsHandler = nullptr;
 	if (!AnalyticsHandlerClass.IsNull())
@@ -2399,7 +2399,7 @@ TStrongObjectPtr<UInterchangeAnalyticsHandlerBase> UInterchangeManager::CreateAn
 			if (UInterchangeAnalyticsHandlerBase* AnalyticsHandlerBase = NewObject<UInterchangeAnalyticsHandlerBase>(GetTransientPackage(), LoadedAnalyticsHandlerClass, NAME_None, RF_NoFlags))
 			{
 				AnalyticsHandler = TStrongObjectPtr<UInterchangeAnalyticsHandlerBase>(AnalyticsHandlerBase);
-				AnalyticsHandler->SetGUID(FGuid::NewDeterministicGuid(ContentPath, static_cast<uint64>(FDateTime::UtcNow().GetTicks())));
+				AnalyticsHandler->SetGUID(FGuid::NewDeterministicGuid(ImportFilePath, static_cast<uint64>(FDateTime::UtcNow().GetTicks())));
 			}
 		}
 	}
diff --git a/Engine/Source/Runtime/Interchange/Engine/Public/InterchangeManager.h b/Engine/Source/Runtime/Interchange/Engine/Public/InterchangeManager.h
--- a/Engine/Source/Runtime/Interchange/Engine/Public/InterchangeManager.h
+++ b/Engine/Source/Runtime/Interchange/Engine/Public/InterchangeManager.h
@@ -976,7 +976,7 @@ protected:
 	void RemoveImportTask(int32 AsyncHelperUniqueId);
 
 private:
-	TStrongObjectPtr<UInterchangeAnalyticsHandlerBase> CreateAnalyticsHandler(const FString& ContentPath);
+	TStrongObjectPtr<UInterchangeAnalyticsHandlerBase> CreateAnalyticsHandler(const FString& ImportFilePath);
 
 private:
 
