Interchange Analytics - Added support to send nested structure properties via analytics

#rnx

#Synced-CL 51376356

diff --git a/Engine/Plugins/Interchange/Runtime/Source/Analytics/Private/InterchangeAnalyticsHandlerDefault.cpp b/Engine/Plugins/Interchange/Runtime/Source/Analytics/Private/InterchangeAnalyticsHandlerDefault.cpp
--- a/Engine/Plugins/Interchange/Runtime/Source/Analytics/Private/InterchangeAnalyticsHandlerDefault.cpp
+++ b/Engine/Plugins/Interchange/Runtime/Source/Analytics/Private/InterchangeAnalyticsHandlerDefault.cpp
@@ -215,15 +215,25 @@ void UInterchangeAnalyticsHandlerDefault::SendPipelineAnalytics(UInterchangePipe
 	PipelineAttribs.Add(FAnalyticsEventAttribute(TEXT("Name"), PipelineChainName));
 	PipelineAttribs.Add(FAnalyticsEventAttribute(TEXT("Class"), Class->GetName()));
 
-	for (FProperty* Property = Class->PropertyLink; Property; Property = Property->PropertyLinkNext)
+	for (TPropertyValueIterator<FProperty> It(
+			Pipeline.GetClass(),
+			&Pipeline, 
+			EPropertyValueIteratorFlags::FullRecursion, 
+			EFieldIteratorFlags::ExcludeDeprecated); 
+			It; 
+			++It)
 	{
-		if (!Pipeline.CanAddPropertyToAnalytics(Property))
+		if (!Pipeline.CanAddPropertyToAnalytics(It))
 		{
+			It.SkipRecursiveProperty();
 			continue;
 		}
 
-		FObjectProperty* SubObject = CastField<FObjectProperty>(Property);
-		UInterchangePipelineBase* SubPipeline = SubObject ? Cast<UInterchangePipelineBase>(SubObject->GetObjectPropertyValue_InContainer(&Pipeline)) : nullptr;
+		const FProperty* Property = It.Key();
+		const void* PropertyValue = It.Value();
+
+		const FObjectProperty* SubObject = CastField<FObjectProperty>(Property);
+		UInterchangePipelineBase* SubPipeline = SubObject ? Cast<UInterchangePipelineBase>(SubObject->GetObjectPropertyValue(PropertyValue)) : nullptr;
 
 		if (SubPipeline)
 		{
@@ -236,7 +246,7 @@ void UInterchangeAnalyticsHandlerDefault::SendPipelineAnalytics(UInterchangePipe
 		}
 		else
 		{
-			Pipeline.AddPropertyAnalytics(Property, PipelineAttribs);
+			Pipeline.AddPropertyAnalytics(It, PipelineAttribs);
 		}
 	}
 
diff --git a/Engine/Source/Runtime/Interchange/Core/Private/InterchangePipelineBase.cpp b/Engine/Source/Runtime/Interchange/Core/Private/InterchangePipelineBase.cpp
--- a/Engine/Source/Runtime/Interchange/Core/Private/InterchangePipelineBase.cpp
+++ b/Engine/Source/Runtime/Interchange/Core/Private/InterchangePipelineBase.cpp
@@ -181,12 +181,14 @@ FName UInterchangePipelineBase::GetResultsPropertyName()
 	return GET_MEMBER_NAME_CHECKED(UInterchangePipelineBase, Results);
 }
 
-bool UInterchangePipelineBase::CanAddPropertyToAnalytics(FProperty* Property) const
+bool UInterchangePipelineBase::CanAddPropertyToAnalytics(TPropertyValueIterator<FProperty>& PropertyValueIterator) const
 {
+	const FProperty* Property = PropertyValueIterator.Key();
 	if (!ensure(Property))
 	{
 		return false;
 	}
+	const void* PropertyValue = PropertyValueIterator.Value();
 
 	if (Property->HasAnyPropertyFlags(CPF_Transient))
 	{
@@ -204,9 +206,9 @@ bool UInterchangePipelineBase::CanAddPropertyToAnalytics(FProperty* Property) co
 	}
 
 #if WITH_EDITORONLY_DATA
-	FObjectProperty* SubObject = CastField<FObjectProperty>(Property);
+	const FObjectProperty* SubObject = CastField<FObjectProperty>(Property);
 	UInterchangePipelineBase* SubPipeline = SubObject 
-												? Cast<UInterchangePipelineBase>(SubObject->GetObjectPropertyValue_InContainer(this)) 
+												? Cast<UInterchangePipelineBase>(SubObject->GetObjectPropertyValue(PropertyValue)) 
 												: nullptr;
 	if (!SubPipeline && Property->GetMetaData("Category").IsEmpty())
 	{
@@ -218,53 +220,87 @@ bool UInterchangePipelineBase::CanAddPropertyToAnalytics(FProperty* Property) co
 	return true;
 }
 
-void UInterchangePipelineBase::AddPropertyAnalytics(FProperty* Property, TArray<FAnalyticsEventAttribute>& PipelineAttributes)
+void UInterchangePipelineBase::AddPropertyAnalytics(TPropertyValueIterator<FProperty>& PropertyValueIterator, TArray<FAnalyticsEventAttribute>& PipelineAttributes)
 {
-	if (!ensure(Property))
-	{
-		return;
-	}
+	const FProperty* Property = PropertyValueIterator.Key();
+	const void* Value = PropertyValueIterator.Value();
 
-	const FString PropertyName = Property->GetName();
-	constexpr int PortFlags = 0;
+	TArray<const FProperty*> PropertyChain;
+	PropertyValueIterator.GetPropertyChain(PropertyChain);
 
-	if (FArrayProperty* Array = CastField<FArrayProperty>(Property))
+	const FString PropertyName = [&PropertyChain]() -> FString {
+		if (PropertyChain.Num() == 1)
+		{
+			return PropertyChain[0]->GetName();
+		}
+
+		TStringBuilder<256> StringBuilder;	
+		for (int32 PropertyIndex = PropertyChain.Num() - 1; PropertyIndex >= 0; --PropertyIndex)
+		{
+			StringBuilder << PropertyChain[PropertyIndex]->GetName();
+			if (PropertyIndex > 0)
+			{
+				StringBuilder << TEXT(".");
+			}
+		}
+
+		return StringBuilder.ToString();
+	}();
+	
+	bool bUseDefaultHandler = false;
+	constexpr int32 PortFlags = 0;
+	if (const FArrayProperty* Array = CastField<FArrayProperty>(Property))
 	{
-		FScriptArrayHelper_InContainer ArrayHelper(Array, this);
+		FScriptArrayHelper ArrayHelper(Array, Value);
 		for (int32 i = 0; i < ArrayHelper.Num(); i++)
 		{
-			FString	Buffer;
-			Array->Inner->ExportTextItem_Direct(Buffer, ArrayHelper.GetRawPtr(i), ArrayHelper.GetRawPtr(i), this, PortFlags);
+			FString	ValueString;
+			Array->Inner->ExportTextItem_Direct(ValueString, ArrayHelper.GetRawPtr(i), ArrayHelper.GetRawPtr(i), this, PortFlags);
 
-			PipelineAttributes.Add(FAnalyticsEventAttribute(FString::Printf(TEXT("Property%s.%s"), *GetCategoryNameForAnalytics(*Property), *PropertyName), Buffer));
+			PipelineAttributes.Add(FAnalyticsEventAttribute(FString::Printf(TEXT("Property%s.%s"), *GetCategoryNameForAnalytics(PropertyChain), *PropertyName), ValueString));
 		}
 	}
-	else
+	else if(const FStructProperty* StructProperty = CastField<FStructProperty>(Property))
 	{
-		for (int32 Index = 0; Index < Property->ArrayDim; Index++)
+		if (StructProperty->Struct == TBaseStructure<FVector>::Get()
+			|| StructProperty->Struct == TBaseStructure<FRotator>::Get()
+			|| StructProperty->Struct == TBaseStructure<FColor>::Get()
+			|| StructProperty->Struct == TBaseStructure<FInt32Interval>::Get())
 		{
-			FString PropertyIndexName = PropertyName;
-			if (Property->ArrayDim != 1)
-			{
-				PropertyIndexName += TEXT("[") + FString::FromInt(Index) + TEXT("]");
-			}
-
-			FString	Value;
-			Property->ExportText_InContainer(Index, Value, this, this, this, PortFlags);
-			PipelineAttributes.Add(FAnalyticsEventAttribute(FString::Printf(TEXT("Property%s.%s"), *GetCategoryNameForAnalytics(*Property), *PropertyIndexName), Value));
+			bUseDefaultHandler = true;
+			PropertyValueIterator.SkipRecursiveProperty();
 		}
 	}
+	else
+	{
+		bUseDefaultHandler = true;
+	}
+
+	if (bUseDefaultHandler)
+	{
+		FString	ValueString;
+		Property->ExportText_Direct(ValueString, Value, Value, this, PortFlags);
+		PipelineAttributes.Add(FAnalyticsEventAttribute(FString::Printf(TEXT("Property%s.%s"), *GetCategoryNameForAnalytics(PropertyChain), *PropertyName), ValueString));
+	}
 }
 
-FString UInterchangePipelineBase::GetCategoryNameForAnalytics(const FProperty& Property) const
+FString UInterchangePipelineBase::GetCategoryNameForAnalytics(const TArray<const FProperty*>& PropertyChain) const
 {
 	FString CategoryName = FString();
+	if (PropertyChain.IsEmpty())
+	{
+		return CategoryName;
+	}
+
 
 #if WITH_EDITORONLY_DATA
-	CategoryName = Property.GetMetaData("Category");
-	CategoryName.ReplaceCharInline(TEXT('.'), TEXT('_'));
-	CategoryName.RemoveSpacesInline();
-	CategoryName = TEXT(".") + CategoryName;
+	if (const FProperty* Property = PropertyChain.Last())
+	{
+		CategoryName = Property->GetMetaData("Category");
+		CategoryName.ReplaceCharInline(TEXT('.'), TEXT('_'));
+		CategoryName.RemoveSpacesInline();
+		CategoryName = TEXT(".") + CategoryName;
+	}
 #endif
 
 	return CategoryName;
diff --git a/Engine/Source/Runtime/Interchange/Core/Public/InterchangePipelineBase.h b/Engine/Source/Runtime/Interchange/Core/Public/InterchangePipelineBase.h
--- a/Engine/Source/Runtime/Interchange/Core/Public/InterchangePipelineBase.h
+++ b/Engine/Source/Runtime/Interchange/Core/Public/InterchangePipelineBase.h
@@ -578,13 +578,25 @@ public:
 
 	UE_API UInterchangePipelineBase* GetMostPipelineOuter() const;
 
-	UE_API virtual bool CanAddPropertyToAnalytics(FProperty* Property) const;
+	UE_DEPRECATED(5.8, "Use the version with PropertyValueIterator instead.")
+	virtual bool CanAddPropertyToAnalytics(FProperty* Property) const
+	{
+		return false;
+	}
+
+	UE_API virtual bool CanAddPropertyToAnalytics(TPropertyValueIterator<FProperty>& PropertyValueIterator) const;
+
+	UE_DEPRECATED(5.8, "Use the version with PropertyValueIterator instead.")
+	virtual void AddPropertyAnalytics(FProperty* Property, TArray<FAnalyticsEventAttribute>& PipelineAttributes)
+	{
+	}
+
+	UE_API virtual void AddPropertyAnalytics(TPropertyValueIterator<FProperty>& PropertyValueIterator, TArray<FAnalyticsEventAttribute>& PipelineAttributes);
 
-	UE_API virtual void AddPropertyAnalytics(FProperty* Property, TArray<FAnalyticsEventAttribute>& PipelineAttributes);
 
 protected:
 
-	UE_API virtual FString GetCategoryNameForAnalytics(const FProperty& Property) const;
+	UE_API virtual FString GetCategoryNameForAnalytics(const TArray<const FProperty*>& PropertyChain) const;
 
 	UE_API virtual FString GetPipelineDisplayName() const;
 
