Interchange - Putting node Uid debugging using UObject Name behind a CVar.

Adding a CVar to display actual Uid collision.

#rnx

#Synced-CL 51380339

diff --git a/Engine/Source/Runtime/Interchange/Core/Private/Nodes/InterchangeBaseNode.cpp b/Engine/Source/Runtime/Interchange/Core/Private/Nodes/InterchangeBaseNode.cpp
--- a/Engine/Source/Runtime/Interchange/Core/Private/Nodes/InterchangeBaseNode.cpp
+++ b/Engine/Source/Runtime/Interchange/Core/Private/Nodes/InterchangeBaseNode.cpp
@@ -2,6 +2,7 @@
 #include "Nodes/InterchangeBaseNode.h"
 
 #include "CoreMinimal.h"
+#include "HAL/IConsoleManager.h"
 #include "Misc/SecureHash.h"
 #include "Types/AttributeStorage.h"
 #include "UObject/Class.h"
@@ -10,6 +11,14 @@
 #include "InterchangeHelper.h"
 #include "InterchangeLogPrivate.h"
 
+static bool GUseNodeUidAsName = false;
+static FAutoConsoleVariableRef CvarGUseInterchangeUidAsName(
+	TEXT("Interchange.FeatureFlags.UseNodeUidAsName"),
+	GUseNodeUidAsName,
+	TEXT("Whether to use the InterchangeBaseNode <DisplayName | UniqueUid> as the UObject Name. Useful for debugging."),
+	ECVF_Default);
+
+#define INTERCHANGE_USE_NODE_UID_AS_NAME (WITH_EDITOR && !UE_BUILD_SHIPPING)
 
 #include UE_INLINE_GENERATED_CPP_BY_NAME(InterchangeBaseNode)
 
@@ -89,18 +98,21 @@ UInterchangeBaseNode::UInterchangeBaseNode()
 	RegisterAttribute<uint8>(UE::Interchange::FBaseNodeStaticData::NodeContainerTypeKey(), static_cast<uint8>(EInterchangeNodeContainerType::None));
 }
 
-#if UE_BUILD_DEBUG || UE_BUILD_DEVELOPMENT
+#if INTERCHANGE_USE_NODE_UID_AS_NAME
 static void UpdateNameForDebugging(UInterchangeBaseNode* Node)
 {
-	//For debugging purpose
-	const FString Name = Node->GetDisplayLabel() + " | " + Node->GetUniqueID();
-	if (!Node->Rename(*Name, nullptr, REN_Test))
+	if (!GUseNodeUidAsName)
 	{
-		UE_LOG(LogInterchangeCore, Error, TEXT("Duplicate NodeUid Found `%s`. This might indicates a potential issue in how this node was added without checking for its existence first."), *Name);
 		return;
 	}
 
-	Node->Rename(*Name, nullptr, REN_DontCreateRedirectors | REN_NonTransactional);
+	//For debugging purpose
+	const FString Name = Node->GetDisplayLabel() + TEXT(" | ") + Node->GetUniqueID();
+
+	if (Name.Len() < NAME_SIZE && Node->Rename(*Name, nullptr, REN_Test))
+	{
+		Node->Rename(*Name, nullptr, REN_DontCreateRedirectors | REN_NonTransactional);
+	}
 }
 #endif
 
@@ -125,7 +137,7 @@ void UInterchangeBaseNode::InitializeNode(const FString& UniqueID, const FString
 
 	bIsInitialized = true;
 
-#if UE_BUILD_DEBUG || UE_BUILD_DEVELOPMENT
+#if INTERCHANGE_USE_NODE_UID_AS_NAME
 	UpdateNameForDebugging(this);
 #endif
 }
@@ -374,7 +386,7 @@ bool UInterchangeBaseNode::SetDisplayLabel(const FString& InDisplayLabel)
 		bSuccess = Handle.IsValid();
 	}
 
-#if UE_BUILD_DEBUG || UE_BUILD_DEVELOPMENT
+#if INTERCHANGE_USE_NODE_UID_AS_NAME
 	if (bIsInitialized && bSuccess)
 	{
 		UpdateNameForDebugging(this);
diff --git a/Engine/Source/Runtime/Interchange/Core/Private/Nodes/InterchangeBaseNodeContainer.cpp b/Engine/Source/Runtime/Interchange/Core/Private/Nodes/InterchangeBaseNodeContainer.cpp
--- a/Engine/Source/Runtime/Interchange/Core/Private/Nodes/InterchangeBaseNodeContainer.cpp
+++ b/Engine/Source/Runtime/Interchange/Core/Private/Nodes/InterchangeBaseNodeContainer.cpp
@@ -50,11 +50,6 @@ bool UInterchangeBaseNodeContainer::RemoveFactoryNode(const FString& NodeUniqueI
 	if (UInterchangeFactoryBaseNode* OldNode = GetFactoryNode(NodeUniqueID))
 	{
 		Nodes.Remove(NodeUniqueID);
-
-		OldNode->ClearFlags(RF_Standalone);
-		OldNode->ClearInternalFlags(EInternalObjectFlags::Async);
-		OldNode->Rename(nullptr, GetTransientPackage(), REN_DontCreateRedirectors | REN_NonTransactional);
-
 		return true;
 	}
 
