[InterchangeBaseNodeContainer]
- Optimize by getting derived classes in the optimal way
- 191ms -> 24us for enumerating derived classes

#Synced-CL 51263423

diff --git a/Engine/Source/Runtime/Interchange/Core/Private/Nodes/InterchangeBaseNodeContainer.cpp b/Engine/Source/Runtime/Interchange/Core/Private/Nodes/InterchangeBaseNodeContainer.cpp
--- a/Engine/Source/Runtime/Interchange/Core/Private/Nodes/InterchangeBaseNodeContainer.cpp
+++ b/Engine/Source/Runtime/Interchange/Core/Private/Nodes/InterchangeBaseNodeContainer.cpp
@@ -50,6 +50,11 @@ bool UInterchangeBaseNodeContainer::RemoveFactoryNode(const FString& NodeUniqueI
 	if (UInterchangeFactoryBaseNode* OldNode = GetFactoryNode(NodeUniqueID))
 	{
 		Nodes.Remove(NodeUniqueID);
+
+		OldNode->ClearFlags(RF_Standalone);
+		OldNode->ClearInternalFlags(EInternalObjectFlags::Async);
+		OldNode->Rename(nullptr, GetTransientPackage(), REN_DontCreateRedirectors | REN_NonTransactional);
+
 		return true;
 	}
 
@@ -354,15 +359,18 @@ void UInterchangeBaseNodeContainer::SerializeNodeContainerData(FArchive& Ar)
 	{
 		//Find all the potential node class
 		TMap<FString, UClass*> ClassPerName;
-		for (FThreadSafeObjectIterator It(UClass::StaticClass()); It; ++It)
+
+		TArray<UClass*> Classes;
+		GetDerivedClasses(UInterchangeBaseNode::StaticClass(), Classes);
+
+		for (UClass* Class : Classes)
 		{
-			UClass* Class = Cast<UClass>(*It);
-			if (Class->IsChildOf(UInterchangeBaseNode::StaticClass()))
-			{
-				ClassPerName.Add(Class->GetFullName(), Class);
-			}
+			ClassPerName.Add(Class->GetFullName(), Class);
 		}
 
+		// The base class is never returned by GetDerivedClasses but was present before, add it manually.
+		ClassPerName.Add(UInterchangeBaseNode::StaticClass()->GetFullName(), UInterchangeBaseNode::StaticClass());
+
 		for (int32 NodeIndex = 0; NodeIndex < NodeCount; ++NodeIndex)
 		{
 			FString ClassFullName;
