[InterchangeBaseNodeContainer]
- Fix bad merge after sync

#Synced-CL 51263423

diff --git a/Engine/Source/Runtime/Interchange/Core/Private/Nodes/InterchangeBaseNodeContainer.cpp b/Engine/Source/Runtime/Interchange/Core/Private/Nodes/InterchangeBaseNodeContainer.cpp
--- a/Engine/Source/Runtime/Interchange/Core/Private/Nodes/InterchangeBaseNodeContainer.cpp
+++ b/Engine/Source/Runtime/Interchange/Core/Private/Nodes/InterchangeBaseNodeContainer.cpp
@@ -50,11 +50,6 @@ bool UInterchangeBaseNodeContainer::RemoveFactoryNode(const FString& NodeUniqueI
 	if (UInterchangeFactoryBaseNode* OldNode = GetFactoryNode(NodeUniqueID))
 	{
 		Nodes.Remove(NodeUniqueID);
-
-		OldNode->ClearFlags(RF_Standalone);
-		OldNode->ClearInternalFlags(EInternalObjectFlags::Async);
-		OldNode->Rename(nullptr, GetTransientPackage(), REN_DontCreateRedirectors | REN_NonTransactional);
-
 		return true;
 	}
 
