Interchange: Fix payloads still being requested for disabled factory nodes

#Synced-CL 51466015

diff --git a/Engine/Source/Runtime/Interchange/Engine/Private/Tasks/InterchangeTaskParsing.cpp b/Engine/Source/Runtime/Interchange/Engine/Private/Tasks/InterchangeTaskParsing.cpp
--- a/Engine/Source/Runtime/Interchange/Engine/Private/Tasks/InterchangeTaskParsing.cpp
+++ b/Engine/Source/Runtime/Interchange/Engine/Private/Tasks/InterchangeTaskParsing.cpp
@@ -136,7 +136,7 @@ namespace UE::InterchangeTaskParsing::Private
 				}
 			}
 
-			if (!FactoryNode->IsEnabled() || ensureMsgf(!OutCreatedTasksAssetNames.Contains(AssetFullPath),
+			if (FactoryNode->IsEnabled() && ensureMsgf(!OutCreatedTasksAssetNames.Contains(AssetFullPath),
 				TEXT("Found multiple task data with the same asset name (%s). Only one will be executed."), *AssetFullPath))
 			{
 				LLM_SCOPE_BYNAME(TEXT("Interchange"));
