[Interchange Scene Graph]  Support for Prefabs

#rnx
#Synced-CL 52496820

diff --git a/Engine/Source/Runtime/Interchange/Engine/Private/InterchangeSceneImportAsset.cpp b/Engine/Source/Runtime/Interchange/Engine/Private/InterchangeSceneImportAsset.cpp
--- a/Engine/Source/Runtime/Interchange/Engine/Private/InterchangeSceneImportAsset.cpp
+++ b/Engine/Source/Runtime/Interchange/Engine/Private/InterchangeSceneImportAsset.cpp
@@ -235,24 +235,7 @@ UObject* UInterchangeSceneImportAsset::GetSceneObject(const FString& PackageName
 	{
 		if (IsValid(SceneObject))
 		{
-			// No AssetImportData on AActors nor UWorlds nor ULevels. Return found SceneObject
-			if (SceneObject->IsA<AActor>() || SceneObject->IsA<UWorld>() || SceneObject->IsA<ULevel>())
-			{
-				return SceneObject;
-			}
-
-			// Most likely an asset, check whether SceneObject has actually already been imported
-			TArray<UObject*> SubObjects;
-			GetObjectsWithOuter(SceneObject, SubObjects);
-			for (UObject* SubObject : SubObjects)
-			{
-				if (SubObject && SubObject->IsA<UInterchangeAssetImportData>())
-				{
-					return SceneObject;
-				}
-			}
-
-			return nullptr;
+			return SceneObject;
 		}
 
 		// SceneObject is still in memory but invalid. Move it to TransientPackage
